#TTP 328983 - Landscape: Move to Level tool doesn't do anything if you don't have a selection, despite rendering a component highlight

#proj UE4

#summary Fixed move to level tool not doing anything if you don't have a selection
#fix Also fixed tool clicks getting stuck

diff --git a/Engine/Source/Editor/LandscapeEditor/Private/LandscapeEdMode.cpp b/Engine/Source/Editor/LandscapeEditor/Private/LandscapeEdMode.cpp
--- a/Engine/Source/Editor/LandscapeEditor/Private/LandscapeEdMode.cpp
+++ b/Engine/Source/Editor/LandscapeEditor/Private/LandscapeEdMode.cpp
@@ -616,6 +616,15 @@ void FEdModeLandscape::Tick(FLevelEditorViewportClient* ViewportClient,float Del
 /** FEdMode: Called when the mouse is moved over the viewport */
 bool FEdModeLandscape::MouseMove( FLevelEditorViewportClient* ViewportClient, FViewport* Viewport, int32 MouseX, int32 MouseY )
 {
+	if (bToolActive && !Viewport->KeyState(EKeys::LeftMouseButton))
+	{
+		if (CurrentToolSet && CurrentToolSet->GetTool())
+		{
+			CurrentToolSet->GetTool()->EndTool(ViewportClient);
+			bToolActive = false;
+		}
+	}
+
 	if (GEditor->PlayWorld != NULL)
 	{
 		return false;
diff --git a/Engine/Source/Editor/LandscapeEditor/Private/LandscapeEdModeComponentTools.cpp b/Engine/Source/Editor/LandscapeEditor/Private/LandscapeEdModeComponentTools.cpp
--- a/Engine/Source/Editor/LandscapeEditor/Private/LandscapeEdModeComponentTools.cpp
+++ b/Engine/Source/Editor/LandscapeEditor/Private/LandscapeEdModeComponentTools.cpp
@@ -334,6 +334,7 @@ public:
 				{
 					return;
 				}
+				LandscapeInfo->GetComponentsInRegion(X1 + 1, Y1 + 1, X2 - 1, Y2 - 1, SelectedComponents);
 				bBrush = true;
 			}
 
