#summary UDN - Fix re-import of 8-bit heightmaps to an existing landscape

diff --git a/Engine/Source/Editor/LandscapeEditor/Private/LandscapeEdMode.cpp b/Engine/Source/Editor/LandscapeEditor/Private/LandscapeEdMode.cpp
--- a/Engine/Source/Editor/LandscapeEditor/Private/LandscapeEdMode.cpp
+++ b/Engine/Source/Editor/LandscapeEditor/Private/LandscapeEdMode.cpp
@@ -2351,9 +2351,19 @@ void FEdModeLandscape::ImportData(const FLandscapeTargetListInfo& TargetInfo, co
 					const TArray<uint8>* RawData = NULL;
 					if (TargetInfo.TargetType == ELandscapeToolTargetType::Heightmap)
 					{
+						if (ImageWrapper->GetFormat() != ERGBFormat::Gray)
+						{
+							FMessageDialog::Open(EAppMsgType::Ok, NSLOCTEXT("LandscapeEditor.NewLandscape", "Import_HeightmapFileColorPng", "The Heightmap file appears to be a color png, grayscale is expected. The import *can* continue, but the result may not be what you expect..."));
+						}
+						if (ImageWrapper->GetBitDepth() != 16)
+						{
+							FMessageDialog::Open(EAppMsgType::Ok, NSLOCTEXT("LandscapeEditor.NewLandscape", "Import_HeightmapFileLowBitDepth", "The Heightmap file appears to be an 8-bit png, 16-bit is preferred. The import *can* continue, but the result may be lower quality than desired."));
+						}
+
 						if (ImageWrapper->GetBitDepth() <= 8)
 						{
 							ImageWrapper->GetRaw(ERGBFormat::Gray, 8, RawData);
+							Data.Reset();
 							Data.AddUninitialized(RawData->Num() * sizeof(uint16));
 							uint16* DataPtr = (uint16*)Data.GetData();
 							for (int32 i = 0; i < RawData->Num(); i++)
@@ -2369,6 +2379,11 @@ void FEdModeLandscape::ImportData(const FLandscapeTargetListInfo& TargetInfo, co
 					}
 					else
 					{
+						if (ImageWrapper->GetFormat() != ERGBFormat::Gray)
+						{
+							FMessageDialog::Open(EAppMsgType::Ok, NSLOCTEXT("LandscapeEditor.NewLandscape", "Import_LayerColorPng", "The Layer file appears to be a color png, grayscale is expected. The import *can* continue, but the result may not be what you expect..."));
+						}
+
 						ImageWrapper->GetRaw(ERGBFormat::Gray, 8, RawData);
 						Data = *RawData; // agh I want to use MoveTemp() here
 					}
