#summary Use asset tagging to skip loading of landscape layer info objects when it isn't necessary to do so

diff --git a/Engine/Source/Editor/LandscapeEditor/Private/LandscapeEditorDetailCustomization_TargetLayers.cpp b/Engine/Source/Editor/LandscapeEditor/Private/LandscapeEditorDetailCustomization_TargetLayers.cpp
--- a/Engine/Source/Editor/LandscapeEditor/Private/LandscapeEditorDetailCustomization_TargetLayers.cpp
+++ b/Engine/Source/Editor/LandscapeEditor/Private/LandscapeEditorDetailCustomization_TargetLayers.cpp
@@ -559,13 +559,14 @@ void FLandscapeEditorCustomNodeBuilder_TargetLayers::OnReimportLayer(const TShar
 
 bool FLandscapeEditorCustomNodeBuilder_TargetLayers::ShouldFilterLayerInfo(const class FAssetData& AssetData, FName LayerName)
 {
-	ULandscapeLayerInfoObject* LayerInfo = CastChecked<ULandscapeLayerInfoObject>(AssetData.GetAsset());
-	if (LayerInfo->LayerName != LayerName)
+	const FString* const LayerNameMetaData = AssetData.TagsAndValues.Find("LayerName");
+	if (LayerNameMetaData && !LayerNameMetaData->IsEmpty())
 	{
-		return true;
+		return FName(**LayerNameMetaData) != LayerName;
 	}
 
-	return false;
+	ULandscapeLayerInfoObject* LayerInfo = CastChecked<ULandscapeLayerInfoObject>(AssetData.GetAsset());
+	return LayerInfo->LayerName != LayerName;
 }
 
 void FLandscapeEditorCustomNodeBuilder_TargetLayers::OnTargetLayerSetObject(const UObject* Object, const TSharedRef<FLandscapeTargetListInfo> Target)
diff --git a/Engine/Source/Runtime/Engine/Classes/Landscape/LandscapeLayerInfoObject.h b/Engine/Source/Runtime/Engine/Classes/Landscape/LandscapeLayerInfoObject.h
--- a/Engine/Source/Runtime/Engine/Classes/Landscape/LandscapeLayerInfoObject.h
+++ b/Engine/Source/Runtime/Engine/Classes/Landscape/LandscapeLayerInfoObject.h
@@ -9,7 +9,7 @@ class ULandscapeLayerInfoObject : public UObject
 {
 	GENERATED_UCLASS_BODY()
 
-	UPROPERTY(VisibleAnywhere, Category=LandscapeLayerInfoObject)
+	UPROPERTY(VisibleAnywhere, Category=LandscapeLayerInfoObject, AssetRegistrySearchable)
 	FName LayerName;
 
 	UPROPERTY(EditAnywhere, Category=LandscapeLayerInfoObject)
