#TTP 323397 - Landscape: Remove Paint tool options that are intended for sculpting only

#proj UE4

#summary Removed sculpting-only options from flatten tool when in paint mode

diff --git a/Engine/Source/Editor/LandscapeEditor/Classes/LandscapeEditorObject.h b/Engine/Source/Editor/LandscapeEditor/Classes/LandscapeEditorObject.h
--- a/Engine/Source/Editor/LandscapeEditor/Classes/LandscapeEditorObject.h
+++ b/Engine/Source/Editor/LandscapeEditor/Classes/LandscapeEditorObject.h
@@ -222,11 +222,11 @@ class ULandscapeEditorObject : public UObject
 	TEnumAsByte<ELandscapeToolFlattenMode::Type> FlattenMode;
 
 	// Flattens to the angle of the clicked point, instead of horizontal
-	UPROPERTY(Category="Tool Settings", EditAnywhere, meta=(ShowForTools="ToolSet_Flatten"))
+	UPROPERTY(Category="Tool Settings", EditAnywhere, meta=(ShowForTools="ToolSet_Flatten", ShowForTargetTypes="Heightmap"))
 	bool bUseSlopeFlatten;
 
 	// Constantly picks new values to flatten towards when dragging around, instead of only using the first clicked point
-	UPROPERTY(Category="Tool Settings", EditAnywhere, meta=(ShowForTools="ToolSet_Flatten"))
+	UPROPERTY(Category="Tool Settings", EditAnywhere, meta=(ShowForTools="ToolSet_Flatten", ShowForTargetTypes="Heightmap"))
 	bool bPickValuePerApply;
 
 	// Enable to flatten towards a target height
@@ -234,7 +234,7 @@ class ULandscapeEditorObject : public UObject
 	bool bUseFlattenTarget;
 
 	// Target height to flatten towards (in Unreal Units)
-	UPROPERTY(Category="Tool Settings", EditAnywhere, meta=(ShowForTools="ToolSet_Flatten", editcondition="bUseFlattenTarget", UIMin="-32768", UIMax="32768"))
+	UPROPERTY(Category="Tool Settings", EditAnywhere, meta=(ShowForTools="ToolSet_Flatten", ShowForTargetTypes="Heightmap", editcondition="bUseFlattenTarget", UIMin="-32768", UIMax="32768"))
 	float FlattenTarget;
 
 	// Ramp Tool:
diff --git a/Engine/Source/Editor/LandscapeEditor/Private/LandscapeEdModePaintTools.cpp b/Engine/Source/Editor/LandscapeEditor/Private/LandscapeEdModePaintTools.cpp
--- a/Engine/Source/Editor/LandscapeEditor/Private/LandscapeEdModePaintTools.cpp
+++ b/Engine/Source/Editor/LandscapeEditor/Private/LandscapeEdModePaintTools.cpp
@@ -416,17 +416,20 @@ public:
 template<class ToolTarget>
 class FLandscapeToolStrokeFlatten : public FLandscapeToolStrokePaintBase<ToolTarget>
 {
-	bool bInitializedFlattenHeight;
+	typename ToolTarget::CacheClass::DataType FlattenHeight;
+
 	FVector FlattenNormal;
 	float FlattenPlaneDist;
-	typename ToolTarget::CacheClass::DataType FlattenHeight;
+	bool bInitializedFlattenHeight;
+	bool bTargetIsHeightmap;
 
 public:
 	FLandscapeToolStrokeFlatten(FEdModeLandscape* InEdMode, const FLandscapeToolTarget& InTarget)
 	:	FLandscapeToolStrokePaintBase<ToolTarget>(InEdMode, InTarget)
 	,	bInitializedFlattenHeight(false)
+	,	bTargetIsHeightmap(InTarget.TargetType == ELandscapeToolTargetType::Heightmap)
 	{
-		if (InEdMode->UISettings->bUseFlattenTarget && InTarget.TargetType == ELandscapeToolTargetType::Heightmap)
+		if (InEdMode->UISettings->bUseFlattenTarget && bTargetIsHeightmap)
 		{
 			FTransform LocalToWorld = InTarget.LandscapeInfo->GetLandscapeProxy()->ActorToWorld();
 			float Height = InEdMode->UISettings->FlattenTarget / LocalToWorld.GetScale3D().Z - LocalToWorld.GetTranslation().Z;
@@ -439,7 +442,7 @@ public:
 	{
 		if (!this->LandscapeInfo) return;
 
-		if( !bInitializedFlattenHeight || UISettings->bPickValuePerApply)
+		if (!bInitializedFlattenHeight || (UISettings->bPickValuePerApply && bTargetIsHeightmap))
 		{
 			bInitializedFlattenHeight = false;
 			float FlattenX = MousePositions[0].PositionX;
@@ -451,7 +454,7 @@ public:
 			float HeightValue = this->Cache.GetValue(FlattenX, FlattenY);
 			FlattenHeight = HeightValue;
 
-			if (UISettings->bUseSlopeFlatten)
+			if (UISettings->bUseSlopeFlatten && bTargetIsHeightmap)
 			{
 				FlattenNormal = this->Cache.GetNormal(FlattenHeightX, FlattenHeightY);
 				FlattenPlaneDist = -(FlattenNormal | FVector(FlattenX, FlattenY, HeightValue) );
@@ -483,12 +486,6 @@ public:
 		TArray<typename ToolTarget::CacheClass::DataType> HeightData;
 		this->Cache.GetCachedData(X1,Y1,X2,Y2,HeightData);
 
-		// For Add or Sub Flatten Mode
-		// Apply Ratio...
-		TMap<int32, float> RatioInfo;
-		int32 MaxDelta = INT_MIN;
-		int32 MinDelta = INT_MAX;
-
 		// Apply the brush
 		for( auto It = BrushInfo.CreateIterator(); It; ++It )
 		{
@@ -499,8 +496,7 @@ public:
 			{
 				int32 HeightDataIndex = (X-X1) + (Y-Y1)*(1+X2-X1);
 
-				// Conserve stiff
-				if (!UISettings->bUseSlopeFlatten)
+				if (!(UISettings->bUseSlopeFlatten && bTargetIsHeightmap))
 				{
 					int32 Delta = HeightData[HeightDataIndex] - FlattenHeight;
 					switch(UISettings->FlattenMode)
@@ -508,15 +504,13 @@ public:
 					case ELandscapeToolNoiseMode::Add:
 						if (Delta < 0)
 						{
-							MinDelta = FMath::Min<int32>(Delta, MinDelta);
-							RatioInfo.Add(HeightDataIndex, It.Value() * UISettings->ToolStrength * Pressure * Delta);
+							HeightData[HeightDataIndex] = FMath::Lerp(HeightData[HeightDataIndex], FlattenHeight, It.Value() * UISettings->ToolStrength * Pressure);
 						}
 						break;
 					case ELandscapeToolNoiseMode::Sub:
 						if (Delta > 0)
 						{
-							MaxDelta = FMath::Max<int32>(Delta, MaxDelta);
-							RatioInfo.Add(HeightDataIndex, It.Value() * UISettings->ToolStrength * Pressure * Delta);
+							HeightData[HeightDataIndex] = FMath::Lerp(HeightData[HeightDataIndex], FlattenHeight, It.Value() * UISettings->ToolStrength * Pressure);
 						}
 						break;
 					default:
@@ -554,24 +548,6 @@ public:
 			}
 		}
 
-		if (!UISettings->bUseSlopeFlatten)
-		{
-			for( TMap<int32, float>::TIterator It(RatioInfo); It; ++It )
-			{
-				switch(UISettings->FlattenMode)
-				{
-				case ELandscapeToolNoiseMode::Add:
-					HeightData[It.Key()] = FMath::Lerp( HeightData[It.Key()], FlattenHeight, It.Value() / (float)MinDelta );
-					break;
-				case ELandscapeToolNoiseMode::Sub:
-					HeightData[It.Key()] = FMath::Lerp( HeightData[It.Key()], FlattenHeight, It.Value() / (float)MaxDelta );
-					break;
-				default:
-					break;
-				}
-			}
-		}
-
 		this->Cache.SetCachedData(X1,Y1,X2,Y2,HeightData, UISettings->PaintingRestriction);
 		this->Cache.Flush();
 	}
diff --git a/Engine/Source/Editor/LandscapeEditor/Private/SLandscapeEditor.cpp b/Engine/Source/Editor/LandscapeEditor/Private/SLandscapeEditor.cpp
--- a/Engine/Source/Editor/LandscapeEditor/Private/SLandscapeEditor.cpp
+++ b/Engine/Source/Editor/LandscapeEditor/Private/SLandscapeEditor.cpp
@@ -280,6 +280,20 @@ bool SLandscapeEditor::GetIsPropertyVisible(const UProperty* const Property) con
 				return false;
 			}
 		}
+		if (Property->HasMetaData("ShowForTargetTypes"))
+		{
+			static const TCHAR* TargetTypeNames[] = { TEXT("Heightmap"), TEXT("Weightmap"), TEXT("Visibility") };
+
+			TArray<FString> ShowForTargetTypes;
+			Property->GetMetaData("ShowForTargetTypes").ParseIntoArray(&ShowForTargetTypes, TEXT(","), true);
+
+			const ELandscapeToolTargetType::Type CurrentTargetType = LandscapeEdMode->CurrentToolTarget.TargetType;
+			if (CurrentTargetType == ELandscapeToolTargetType::Invalid ||
+				ShowForTargetTypes.FindByKey(TargetTypeNames[CurrentTargetType]) == nullptr)
+			{
+				return false;
+			}
+		}
 
 		return true;
 	}
