Landscape: - Fixed landscape layer data importing from a file

diff --git a/Engine/Source/Editor/LandscapeEditor/Private/LandscapeEdMode.cpp b/Engine/Source/Editor/LandscapeEditor/Private/LandscapeEdMode.cpp
--- a/Engine/Source/Editor/LandscapeEditor/Private/LandscapeEdMode.cpp
+++ b/Engine/Source/Editor/LandscapeEditor/Private/LandscapeEdMode.cpp
@@ -2451,8 +2451,6 @@ ALandscape* FEdModeLandscape::ChangeComponentSetting(int32 NumComponentsX, int32
 			FLandscapeEditDataInterface LandscapeEdit(LandscapeInfo);
 			TArray<uint16> HeightData;
 			TArray<FLandscapeImportLayerInfo> ImportLayerInfos;
-			TArray<TArray<uint8>, TInlineAllocator<8>> LayerDataArrays;
-			TArray<uint8*> LayerDataPtrs;
 			FVector LandscapeOffset = FVector::ZeroVector;
 			float LandscapeScaleFactor = 1.0f;
 
@@ -2471,18 +2469,15 @@ ALandscape* FEdModeLandscape::ChangeComponentSetting(int32 NumComponentsX, int32
 				{
 					if (LayerSettings.LayerInfoObj != NULL)
 					{
-						TArray<uint8>* LayerData = new(LayerDataArrays)(TArray<uint8>);
-						LayerData->AddZeroed(OldVertsX * OldVertsY * sizeof(uint8));
-
+						auto ImportLayerInfo = new(ImportLayerInfos) FLandscapeImportLayerInfo(LayerSettings);
+						ImportLayerInfo->LayerData.AddZeroed(OldVertsX * OldVertsY * sizeof(uint8));
+						
 						TMinX = OldMinX; TMinY = OldMinY; TMaxX = OldMaxX; TMaxY = OldMaxY;
-						LandscapeEdit.GetWeightData(LayerSettings.LayerInfoObj, TMinX, TMinY, TMaxX, TMaxY, LayerData->GetData(), 0);
+						LandscapeEdit.GetWeightData(LayerSettings.LayerInfoObj, TMinX, TMinY, TMaxX, TMaxY, ImportLayerInfo->LayerData.GetData(), 0);
 
-						*LayerData = LandscapeEditorUtils::ResampleData(*LayerData,
+						ImportLayerInfo->LayerData = LandscapeEditorUtils::ResampleData(ImportLayerInfo->LayerData,
 							OldVertsX, OldVertsY,
 							NewVertsX, NewVertsY);
-
-						new(ImportLayerInfos) FLandscapeImportLayerInfo(LayerSettings);
-						LayerDataPtrs.Add(LayerData->GetData());
 					}
 				}
 
@@ -2516,18 +2511,15 @@ ALandscape* FEdModeLandscape::ChangeComponentSetting(int32 NumComponentsX, int32
 				{
 					if (LayerSettings.LayerInfoObj != NULL)
 					{
-						TArray<uint8>* LayerData = new(LayerDataArrays)(TArray<uint8>);
-						LayerData->AddZeroed(NewVertsX * NewVertsY * sizeof(uint8));
+						auto ImportLayerInfo = new(ImportLayerInfos) FLandscapeImportLayerInfo(LayerSettings);
+						ImportLayerInfo->LayerData.AddZeroed(NewVertsX * NewVertsY * sizeof(uint8));
 
 						TMinX = RequestedMinX; TMinY = RequestedMinY; TMaxX = RequestedMaxX; TMaxY = RequestedMaxY;
-						LandscapeEdit.GetWeightData(LayerSettings.LayerInfoObj, TMinX, TMinY, TMaxX, TMaxY, LayerData->GetData(), 0);
+						LandscapeEdit.GetWeightData(LayerSettings.LayerInfoObj, TMinX, TMinY, TMaxX, TMaxY, ImportLayerInfo->LayerData.GetData(), 0);
 
-						*LayerData = LandscapeEditorUtils::ExpandData(*LayerData,
+						ImportLayerInfo->LayerData = LandscapeEditorUtils::ExpandData(ImportLayerInfo->LayerData,
 							RequestedMinX, RequestedMinY, RequestedMaxX, RequestedMaxY,
 							NewMinX, NewMinY, NewMaxX, NewMaxY);
-
-						new(ImportLayerInfos) FLandscapeImportLayerInfo(LayerSettings);
-						LayerDataPtrs.Add(LayerData->GetData());
 					}
 				}
 
@@ -2542,7 +2534,7 @@ ALandscape* FEdModeLandscape::ChangeComponentSetting(int32 NumComponentsX, int32
 
 			Landscape->LandscapeMaterial = OldLandscapeProxy->LandscapeMaterial;
 			Landscape->CollisionMipLevel = OldLandscapeProxy->CollisionMipLevel;
-			Landscape->Import(FGuid::NewGuid(), NewVertsX, NewVertsY, NumSubsections*SubsectionSizeQuads, NumSubsections, SubsectionSizeQuads, HeightData.GetData(), *OldLandscapeProxy->ReimportHeightmapFilePath, ImportLayerInfos, LayerDataPtrs.Num() ? LayerDataPtrs.GetData() : NULL);
+			Landscape->Import(FGuid::NewGuid(), NewVertsX, NewVertsY, NumSubsections*SubsectionSizeQuads, NumSubsections, SubsectionSizeQuads, HeightData.GetData(), *OldLandscapeProxy->ReimportHeightmapFilePath, ImportLayerInfos);
 
 			Landscape->MaxLODLevel                 = OldLandscapeProxy->MaxLODLevel;
 			Landscape->ExportLOD                   = OldLandscapeProxy->ExportLOD;
diff --git a/Engine/Source/Editor/LandscapeEditor/Private/LandscapeEditorDetailCustomization_NewLandscape.cpp b/Engine/Source/Editor/LandscapeEditor/Private/LandscapeEditorDetailCustomization_NewLandscape.cpp
--- a/Engine/Source/Editor/LandscapeEditor/Private/LandscapeEditorDetailCustomization_NewLandscape.cpp
+++ b/Engine/Source/Editor/LandscapeEditor/Private/LandscapeEditorDetailCustomization_NewLandscape.cpp
@@ -740,8 +740,6 @@ FReply FLandscapeEditorDetailCustomization_NewLandscape::OnCreateButtonClicked()
 		}
 
 		TArray<FLandscapeImportLayerInfo> LayerInfos;
-		TArray<TArray<uint8> > LayerDatas;
-		TArray<uint8*> LayerDataPointers;
 
 		if (LandscapeEdMode->NewLandscapePreviewMode == ENewLandscapePreviewMode::ImportLandscape)
 		{
@@ -750,18 +748,16 @@ FReply FLandscapeEditorDetailCustomization_NewLandscape::OnCreateButtonClicked()
 
 			const TArray<FLandscapeImportLayer>& ImportLandscapeLayersList = LandscapeEdMode->UISettings->ImportLandscape_Layers;
 			LayerInfos.Reserve(ImportLandscapeLayersList.Num());
-			LayerDatas.Reserve(ImportLandscapeLayersList.Num());
-			LayerDataPointers.Reserve(ImportLandscapeLayersList.Num());
-
+		
 			// Fill in LayerInfos array and allocate data
 			for (int32 LayerIdx = 0; LayerIdx < ImportLandscapeLayersList.Num(); LayerIdx++)
 			{
-				const FLandscapeImportLayerInfo& ImportLayer = ImportLandscapeLayersList[LayerIdx];
+				LayerInfos.Add(ImportLandscapeLayersList[LayerIdx]); //slicing is fine here
+				FLandscapeImportLayerInfo& ImportLayer = LayerInfos.Last();
 
 				if (ImportLayer.LayerInfo != NULL && ImportLayer.SourceFilePath != "")
 				{
-					TArray<uint8>& LayerData = *new(LayerDatas) TArray<uint8>();
-					FFileHelper::LoadFileToArray(LayerData, *ImportLayer.SourceFilePath, FILEREAD_Silent);
+					FFileHelper::LoadFileToArray(ImportLayer.LayerData, *ImportLayer.SourceFilePath, FILEREAD_Silent);
 
 					if (ImportLayer.SourceFilePath.EndsWith(".png"))
 					{
@@ -769,25 +765,21 @@ FReply FLandscapeEditorDetailCustomization_NewLandscape::OnCreateButtonClicked()
 						IImageWrapperPtr ImageWrapper = ImageWrapperModule.CreateImageWrapper(EImageFormat::PNG);
 
 						const TArray<uint8>* RawData = NULL;
-						if (ImageWrapper->SetCompressed(LayerData.GetData(), LayerData.Num()) &&
+						if (ImageWrapper->SetCompressed(ImportLayer.LayerData.GetData(), ImportLayer.LayerData.Num()) &&
 							ImageWrapper->GetRaw(ERGBFormat::Gray, 8, RawData))
 						{
-							LayerData = *RawData; // agh I want to use MoveTemp() here
+							ImportLayer.LayerData = *RawData; // agh I want to use MoveTemp() here
 						}
 						else
 						{
-							LayerData.Empty();
+							ImportLayer.LayerData.Empty();
 						}
 					}
 
-					if (LayerData.Num() == ImportSizeX * ImportSizeY)
-					{
-						LayerDataPointers.Add(new uint8[SizeX * SizeY]);
-						LayerInfos.Add(ImportLayer);
-					}
-					else
+					// Remove invalid raw weightmap data
+					if (ImportLayer.LayerData.Num() != (ImportSizeX * ImportSizeY))
 					{
-						LayerDatas.RemoveAt(LayerDatas.Num()-1);
+						ImportLayer.LayerData.Empty();
 					}
 				}
 			}
@@ -804,20 +796,21 @@ FReply FLandscapeEditorDetailCustomization_NewLandscape::OnCreateButtonClicked()
 					-OffsetX, -OffsetY, SizeX - OffsetX - 1, SizeY - OffsetY - 1);
 
 				// Layers
-				for (int32 LayerIdx = 0; LayerIdx < LayerDataPointers.Num(); LayerIdx++)
+				for (int32 LayerIdx = 0; LayerIdx < LayerInfos.Num(); LayerIdx++)
 				{
-					TArray<uint8>& ImportLayerData = LayerDatas[LayerIdx];
-
-					ImportLayerData = LandscapeEditorUtils::ExpandData(ImportLayerData,
-						0, 0, ImportSizeX - 1, ImportSizeY - 1,
-						-OffsetX, -OffsetY, SizeX - OffsetX - 1, SizeY - OffsetY - 1);
+					TArray<uint8>& ImportLayerData = LayerInfos[LayerIdx].LayerData;
+					if (ImportLayerData.Num())
+					{
+						ImportLayerData = LandscapeEditorUtils::ExpandData(ImportLayerData,
+							0, 0, ImportSizeX - 1, ImportSizeY - 1,
+							-OffsetX, -OffsetY, SizeX - OffsetX - 1, SizeY - OffsetY - 1);
+					}
 				}
 			}
 
 			LandscapeEdMode->UISettings->ClearImportLandscapeData();
 		}
 
-
 		const float ComponentSize = QuadsPerComponent;
 		const FVector Offset = FTransform(LandscapeEdMode->UISettings->NewLandscape_Rotation, FVector::ZeroVector, LandscapeEdMode->UISettings->NewLandscape_Scale).TransformVector(FVector(-ComponentCountX * ComponentSize / 2, -ComponentCountY * ComponentSize / 2, 0));
 		ALandscape* Landscape = GWorld->SpawnActor<ALandscape>(LandscapeEdMode->UISettings->NewLandscape_Location + Offset, LandscapeEdMode->UISettings->NewLandscape_Rotation);
@@ -828,7 +821,7 @@ FReply FLandscapeEditorDetailCustomization_NewLandscape::OnCreateButtonClicked()
 		}
 
 		Landscape->SetActorRelativeScale3D(LandscapeEdMode->UISettings->NewLandscape_Scale);
-		Landscape->Import(FGuid::NewGuid(), SizeX, SizeY, QuadsPerComponent, LandscapeEdMode->UISettings->NewLandscape_SectionsPerComponent, LandscapeEdMode->UISettings->NewLandscape_QuadsPerSection, Data.GetData(), NULL, LayerInfos, LayerInfos.Num() > 0 ? LayerDataPointers.GetTypedData() : NULL);
+		Landscape->Import(FGuid::NewGuid(), SizeX, SizeY, QuadsPerComponent, LandscapeEdMode->UISettings->NewLandscape_SectionsPerComponent, LandscapeEdMode->UISettings->NewLandscape_QuadsPerSection, Data.GetData(), NULL, LayerInfos);
 
 		// automatically calculate a lighting LOD that won't crash lightmass (hopefully)
 		// < 2048x2048 -> LOD0
@@ -877,12 +870,6 @@ FReply FLandscapeEditorDetailCustomization_NewLandscape::OnCreateButtonClicked()
 
 		LandscapeEdMode->SetCurrentTool("ToolSet_Select"); // change tool so switching back to the manage mode doesn't give "New Landscape" again
 		LandscapeEdMode->SetCurrentTool("ToolSet_Sculpt"); // change to sculpting mode and tool
-
-		// clean up data
-		for (int32 LayerIdx=0;LayerIdx<LayerDataPointers.Num();LayerIdx++)
-		{
-			delete[] LayerDataPointers[LayerIdx];
-		}
 	}
 
 	return FReply::Handled();
diff --git a/Engine/Source/Editor/WorldBrowser/Private/Tiles/WorldTileCollectionModel.cpp b/Engine/Source/Editor/WorldBrowser/Private/Tiles/WorldTileCollectionModel.cpp
--- a/Engine/Source/Editor/WorldBrowser/Private/Tiles/WorldTileCollectionModel.cpp
+++ b/Engine/Source/Editor/WorldBrowser/Private/Tiles/WorldTileCollectionModel.cpp
@@ -1304,6 +1304,7 @@ void FWorldTileCollectionModel::ImportTiledLandscape_Executed()
 			TileImportSettings.SectionsPerComponent = ImportSettings.SectionsPerComponent;
 			TileImportSettings.SizeX				= ImportSettings.TileResolution;
 			TileImportSettings.SizeY				= ImportSettings.TileResolution;
+			TileImportSettings.HeightmapFilename	= Filename;
 			TileImportSettings.LandscapeTransform.SetScale3D(TileScale);
 
 			if (ReadHeighmapFile(TileImportSettings.HeightData, *Filename, FILEREAD_Silent))
diff --git a/Engine/Source/Editor/WorldBrowser/Private/Tiles/WorldTileModel.cpp b/Engine/Source/Editor/WorldBrowser/Private/Tiles/WorldTileModel.cpp
--- a/Engine/Source/Editor/WorldBrowser/Private/Tiles/WorldTileModel.cpp
+++ b/Engine/Source/Editor/WorldBrowser/Private/Tiles/WorldTileModel.cpp
@@ -929,9 +929,8 @@ ALandscapeProxy* FWorldTileModel::ImportLandscape(const FLandscapeImportSettings
 		Settings.SectionsPerComponent, 
 		Settings.QuadsPerSection, 
 		Settings.HeightData.GetData(), 
-		NULL, 
-		Settings.ImportLayers, 
-		NULL);
+		*Settings.HeightmapFilename, 
+		Settings.ImportLayers);
 
 	return Landscape;
 }
diff --git a/Engine/Source/Editor/WorldBrowser/Private/Tiles/WorldTileModel.h b/Engine/Source/Editor/WorldBrowser/Private/Tiles/WorldTileModel.h
--- a/Engine/Source/Editor/WorldBrowser/Private/Tiles/WorldTileModel.h
+++ b/Engine/Source/Editor/WorldBrowser/Private/Tiles/WorldTileModel.h
@@ -46,6 +46,7 @@ public:
 		int32								SizeY;
 		TArray<uint16>						HeightData;
 		TArray<FLandscapeImportLayerInfo>	ImportLayers;
+		FString								HeightmapFilename;
 	};
 	
 	/**
diff --git a/Engine/Source/Runtime/Engine/Classes/Landscape/LandscapeProxy.h b/Engine/Source/Runtime/Engine/Classes/Landscape/LandscapeProxy.h
--- a/Engine/Source/Runtime/Engine/Classes/Landscape/LandscapeProxy.h
+++ b/Engine/Source/Runtime/Engine/Classes/Landscape/LandscapeProxy.h
@@ -124,6 +124,9 @@ struct FLandscapeImportLayerInfo
 
 	UPROPERTY(Category="Import", EditAnywhere)
 	FString SourceFilePath; // Optional
+	
+	// Raw weightmap data
+	TArray<uint8> LayerData;		
 #endif
 
 #if WITH_EDITOR
@@ -387,7 +390,7 @@ public:
 	ENGINE_API void Import(FGuid Guid, int32 VertsX, int32 VertsY, 
 							int32 ComponentSizeQuads, int32 NumSubsections, int32 SubsectionSizeQuads, 
 							const uint16* HeightData, const TCHAR* HeightmapFileName, 
-							TArray<FLandscapeImportLayerInfo> ImportLayerInfos, uint8* AlphaDataPointers[] );
+							const TArray<FLandscapeImportLayerInfo>& ImportLayerInfos);
 
 	/**
 	 * Exports landscape into raw mesh
diff --git a/Engine/Source/Runtime/Engine/Private/Landscape/LandscapeEdit.cpp b/Engine/Source/Runtime/Engine/Private/Landscape/LandscapeEdit.cpp
--- a/Engine/Source/Runtime/Engine/Private/Landscape/LandscapeEdit.cpp
+++ b/Engine/Source/Runtime/Engine/Private/Landscape/LandscapeEdit.cpp
@@ -1679,7 +1679,7 @@ ULandscapeLayerInfoObject* ALandscapeProxy::CreateLayerInfo(const TCHAR* LayerNa
 void ALandscapeProxy::Import(FGuid Guid, int32 VertsX, int32 VertsY, 
 								int32 InComponentSizeQuads, int32 InNumSubsections, int32 InSubsectionSizeQuads, 
 								const uint16* HeightData, const TCHAR* HeightmapFileName, 
-								TArray<FLandscapeImportLayerInfo> ImportLayerInfos, uint8* AlphaDataPointers[] )
+								const TArray<FLandscapeImportLayerInfo>& ImportLayerInfos )
 {
 	GWarn->BeginSlowTask( LOCTEXT("BeingImportingLandscapeTask", "Importing Landscape"), true);
 
@@ -1839,11 +1839,14 @@ void ALandscapeProxy::Import(FGuid Guid, int32 VertsX, int32 VertsY,
 			{
 				FLandscapeComponentAlphaInfo* NewAlphaInfo = new(EditingAlphaLayerData) FLandscapeComponentAlphaInfo(LandscapeComponent, LayerIndex);
 
-				for( int32 AlphaY=0;AlphaY<=LandscapeComponent->ComponentSizeQuads;AlphaY++ )
+				if (ImportLayerInfos[LayerIndex].LayerData.Num())
 				{
-					uint8* OldAlphaRowStart = &AlphaDataPointers[LayerIndex][ (AlphaY+LandscapeComponent->GetSectionBase().Y) * VertsX + (LandscapeComponent->GetSectionBase().X) ];
-					uint8* NewAlphaRowStart = &NewAlphaInfo->AlphaValues[AlphaY * (LandscapeComponent->ComponentSizeQuads+1)];
-					FMemory::Memcpy(NewAlphaRowStart, OldAlphaRowStart, LandscapeComponent->ComponentSizeQuads+1);
+					for( int32 AlphaY=0;AlphaY<=LandscapeComponent->ComponentSizeQuads;AlphaY++ )
+					{
+						const uint8* OldAlphaRowStart = &ImportLayerInfos[LayerIndex].LayerData[ (AlphaY+LandscapeComponent->GetSectionBase().Y) * VertsX + (LandscapeComponent->GetSectionBase().X) ];
+						uint8* NewAlphaRowStart = &NewAlphaInfo->AlphaValues[AlphaY * (LandscapeComponent->ComponentSizeQuads+1)];
+						FMemory::Memcpy(NewAlphaRowStart, OldAlphaRowStart, LandscapeComponent->ComponentSizeQuads+1);
+					}
 				}
 			}
 
