#TTP 332017 - EDITOR: LANDSCAPE: When importing a colored png file for the heightmap, a message appears warning the results may not be desirable but you are forced to continue the import anyways

#proj UE4

#summary Allow user to cancel import of colour or low bit-depth landscape heightmap/weightmap

diff --git a/Engine/Source/Editor/LandscapeEditor/Private/LandscapeEdMode.cpp b/Engine/Source/Editor/LandscapeEditor/Private/LandscapeEdMode.cpp
--- a/Engine/Source/Editor/LandscapeEditor/Private/LandscapeEdMode.cpp
+++ b/Engine/Source/Editor/LandscapeEditor/Private/LandscapeEdMode.cpp
@@ -2353,11 +2353,19 @@ void FEdModeLandscape::ImportData(const FLandscapeTargetListInfo& TargetInfo, co
 					{
 						if (ImageWrapper->GetFormat() != ERGBFormat::Gray)
 						{
-							FMessageDialog::Open(EAppMsgType::Ok, NSLOCTEXT("LandscapeEditor.NewLandscape", "Import_HeightmapFileColorPng", "The Heightmap file appears to be a color png, grayscale is expected. The import *can* continue, but the result may not be what you expect..."));
+							EAppReturnType::Type Result = FMessageDialog::Open(EAppMsgType::OkCancel, NSLOCTEXT("LandscapeEditor.NewLandscape", "Import_HeightmapFileColorPng", "The Heightmap file appears to be a color png, grayscale is expected. The import *can* continue, but the result may not be what you expect..."));
+							if (Result != EAppReturnType::Ok)
+							{
+								return;
+							}
 						}
 						if (ImageWrapper->GetBitDepth() != 16)
 						{
-							FMessageDialog::Open(EAppMsgType::Ok, NSLOCTEXT("LandscapeEditor.NewLandscape", "Import_HeightmapFileLowBitDepth", "The Heightmap file appears to be an 8-bit png, 16-bit is preferred. The import *can* continue, but the result may be lower quality than desired."));
+							EAppReturnType::Type Result = FMessageDialog::Open(EAppMsgType::OkCancel, NSLOCTEXT("LandscapeEditor.NewLandscape", "Import_HeightmapFileLowBitDepth", "The Heightmap file appears to be an 8-bit png, 16-bit is preferred. The import *can* continue, but the result may be lower quality than desired."));
+							if (Result != EAppReturnType::Ok)
+							{
+								return;
+							}
 						}
 
 						if (ImageWrapper->GetBitDepth() <= 8)
@@ -2381,7 +2389,11 @@ void FEdModeLandscape::ImportData(const FLandscapeTargetListInfo& TargetInfo, co
 					{
 						if (ImageWrapper->GetFormat() != ERGBFormat::Gray)
 						{
-							FMessageDialog::Open(EAppMsgType::Ok, NSLOCTEXT("LandscapeEditor.NewLandscape", "Import_LayerColorPng", "The Layer file appears to be a color png, grayscale is expected. The import *can* continue, but the result may not be what you expect..."));
+							EAppReturnType::Type Result = FMessageDialog::Open(EAppMsgType::OkCancel, NSLOCTEXT("LandscapeEditor.NewLandscape", "Import_LayerColorPng", "The Layer file appears to be a color png, grayscale is expected. The import *can* continue, but the result may not be what you expect..."));
+							if (Result != EAppReturnType::Ok)
+							{
+								return;
+							}
 						}
 
 						ImageWrapper->GetRaw(ERGBFormat::Gray, 8, RawData);
