#TTP 332271 - LANDSCAPE: FOLIAGE: Foliage only snaps to current Landscape proxies

#proj UE4

#summary Made foliage snap to *all* modified landscape components when sculpting, not just those in the current level

diff --git a/Engine/Source/Editor/LandscapeEditor/Private/LandscapeEdModeTools.h b/Engine/Source/Editor/LandscapeEditor/Private/LandscapeEdModeTools.h
--- a/Engine/Source/Editor/LandscapeEditor/Private/LandscapeEdModeTools.h
+++ b/Engine/Source/Editor/LandscapeEditor/Private/LandscapeEdModeTools.h
@@ -546,14 +546,24 @@ struct FHeightmapAccessor
 			// Update data
 			ChangedComponents.Append(Components);
 
-			for (TSet<ULandscapeComponent*>::TConstIterator It(Components); It; ++It)
+			for (ULandscapeComponent* Component : Components)
 			{
-				(*It)->InvalidateLightingCache();
+				Component->InvalidateLightingCache();
 			}
 
 			// Notify foliage to move any attached instances
-			AInstancedFoliageActor* IFA = AInstancedFoliageActor::GetInstancedFoliageActor(LandscapeInfo->GetLandscapeProxy()->GetWorld(), false);
-			if( IFA )
+			bool bUpdateFoliage = false;
+			for (ULandscapeComponent* Component : Components)
+			{
+				AInstancedFoliageActor* IFA = AInstancedFoliageActor::GetInstancedFoliageActorForLevel(Component->GetComponentLevel());
+				if (IFA)
+				{
+					bUpdateFoliage = true;
+					break;
+				}
+			}
+
+			if (bUpdateFoliage)
 			{
 				// Calculate landscape local-space bounding box of old data, to look for foliage instances.
 				TArray<ULandscapeHeightfieldCollisionComponent*> CollisionComponents;
@@ -561,21 +571,21 @@ struct FHeightmapAccessor
 				TArray<FBox> PreUpdateLocalBoxes;
 				PreUpdateLocalBoxes.Empty(Components.Num());
 
-				int32 Index=0;
-				for(TSet<ULandscapeComponent*>::TConstIterator It(Components);It;++It)
+				for (ULandscapeComponent* Component : Components)
 				{
-					CollisionComponents.Add( (*It)->CollisionComponent.Get() );
-					PreUpdateLocalBoxes.Add( FBox( FVector((float)X1, (float)Y1, (*It)->CachedLocalBox.Min.Z), FVector((float)X2, (float)Y2, (*It)->CachedLocalBox.Max.Z) ) );
-					Index++;
+					CollisionComponents.Add(Component->CollisionComponent.Get());
+					PreUpdateLocalBoxes.Add(FBox(FVector((float)X1, (float)Y1, Component->CachedLocalBox.Min.Z), FVector((float)X2, (float)Y2, Component->CachedLocalBox.Max.Z)));
 				}
 
 				// Update landscape.
-				LandscapeEdit->SetHeightData( X1, Y1, X2, Y2, Data, 0, true);
+				LandscapeEdit->SetHeightData(X1, Y1, X2, Y2, Data, 0, true);
 
 				// Snap foliage for each component.
-				for( Index=0; Index<CollisionComponents.Num();Index++ )
+				for (int32 Index = 0; Index < CollisionComponents.Num(); ++Index)
 				{
-					IFA->SnapInstancesForLandscape( CollisionComponents[Index], PreUpdateLocalBoxes[Index].TransformBy(LandscapeInfo->GetLandscapeProxy()->LandscapeActorToWorld().ToMatrixWithScale()).ExpandBy(1.f) );
+					ULandscapeHeightfieldCollisionComponent* CollisionComponent = CollisionComponents[Index];
+					AInstancedFoliageActor* IFA = AInstancedFoliageActor::GetInstancedFoliageActorForLevel(CollisionComponent->GetComponentLevel());
+					IFA->SnapInstancesForLandscape(CollisionComponent, PreUpdateLocalBoxes[Index].TransformBy(LandscapeInfo->GetLandscapeProxy()->LandscapeActorToWorld().ToMatrixWithScale()).ExpandBy(1.f));
 				}
 			}
 			else
@@ -712,8 +722,18 @@ struct FXYOffsetmapAccessor
 			}
 
 			// Notify foliage to move any attached instances
-			AInstancedFoliageActor* IFA = AInstancedFoliageActor::GetInstancedFoliageActor(LandscapeInfo->GetLandscapeProxy()->GetWorld(), false);
-			if( IFA )
+			bool bUpdateFoliage = false;
+			for (ULandscapeComponent* Component : Components)
+			{
+				AInstancedFoliageActor* IFA = AInstancedFoliageActor::GetInstancedFoliageActorForLevel(Component->GetComponentLevel());
+				if (IFA)
+				{
+					bUpdateFoliage = true;
+					break;
+				}
+			}
+
+			if (bUpdateFoliage)
 			{
 				// Calculate landscape local-space bounding box of old data, to look for foliage instances.
 				TArray<ULandscapeHeightfieldCollisionComponent*> CollisionComponents;
@@ -721,12 +741,10 @@ struct FXYOffsetmapAccessor
 				TArray<FBox> PreUpdateLocalBoxes;
 				PreUpdateLocalBoxes.Empty(Components.Num());
 
-				int32 Index=0;
-				for(TSet<ULandscapeComponent*>::TConstIterator It(Components);It;++It)
+				for (ULandscapeComponent* Component : Components)
 				{
-					CollisionComponents.Add( (*It)->CollisionComponent.Get() );
-					PreUpdateLocalBoxes.Add( FBox( FVector((float)X1, (float)Y1, (*It)->CachedLocalBox.Min.Z), FVector((float)X2, (float)Y2, (*It)->CachedLocalBox.Max.Z) ) );
-					Index++;
+					CollisionComponents.Add(Component->CollisionComponent.Get());
+					PreUpdateLocalBoxes.Add(FBox(FVector((float)X1, (float)Y1, Component->CachedLocalBox.Min.Z), FVector((float)X2, (float)Y2, Component->CachedLocalBox.Max.Z)));
 				}
 
 				// Update landscape.
@@ -734,9 +752,11 @@ struct FXYOffsetmapAccessor
 				LandscapeEdit->SetHeightData( X1, Y1, X2, Y2, NewHeights.GetTypedData(), 0, true);
 
 				// Snap foliage for each component.
-				for( Index=0; Index<CollisionComponents.Num();Index++ )
+				for (int32 Index = 0; Index < CollisionComponents.Num(); ++Index)
 				{
-					IFA->SnapInstancesForLandscape( CollisionComponents[Index], PreUpdateLocalBoxes[Index].TransformBy(LandscapeInfo->GetLandscapeProxy()->LandscapeActorToWorld().ToMatrixWithScale()).ExpandBy(1.f) );
+					ULandscapeHeightfieldCollisionComponent* CollisionComponent = CollisionComponents[Index];
+					AInstancedFoliageActor* IFA = AInstancedFoliageActor::GetInstancedFoliageActorForLevel(CollisionComponent->GetComponentLevel());
+					IFA->SnapInstancesForLandscape(CollisionComponent, PreUpdateLocalBoxes[Index].TransformBy(LandscapeInfo->GetLandscapeProxy()->LandscapeActorToWorld().ToMatrixWithScale()).ExpandBy(1.f));
 				}
 			}
 			else
