Landscape Spline : Added expose bAutoRotateOnJoin for Spline tool UI

diff --git a/Engine/Source/Editor/LandscapeEditor/Private/LandscapeEdModeSplineTools.cpp b/Engine/Source/Editor/LandscapeEditor/Private/LandscapeEdModeSplineTools.cpp
--- a/Engine/Source/Editor/LandscapeEditor/Private/LandscapeEdModeSplineTools.cpp
+++ b/Engine/Source/Editor/LandscapeEditor/Private/LandscapeEdModeSplineTools.cpp
@@ -1609,6 +1609,8 @@ protected:
 	uint32 bAutoChangeConnectionsOnMove:1;
 	uint32 bDeleteLooseEnds:1;
 	uint32 bCopyMeshToNewControlPoint:1;
+
+	friend class FEdModeLandscape;
 };
 
 
@@ -1657,6 +1659,31 @@ void FEdModeLandscape::SelectAllConnectedSplineSegments()
 	}
 }
 
+void FEdModeLandscape::SetbUseAutoRotateOnJoin(bool InbAutoRotateOnJoin)
+{
+	if (SplinesToolSet /*&& SplinesToolSet == CurrentToolSet*/)
+	{
+		if (SplinesToolSet->SetToolForTarget(CurrentToolTarget) && SplinesToolSet->GetTool())
+		{
+			FLandscapeToolSplines* SplineTool = ((FLandscapeToolSplines*)SplinesToolSet->GetTool());
+			SplineTool->bAutoRotateOnJoin = InbAutoRotateOnJoin;
+		}
+	}
+}
+
+bool FEdModeLandscape::GetbUseAutoRotateOnJoin()
+{
+	if (SplinesToolSet /*&& SplinesToolSet == CurrentToolSet*/)
+	{
+		if (SplinesToolSet->SetToolForTarget(CurrentToolTarget) && SplinesToolSet->GetTool())
+		{
+			FLandscapeToolSplines* SplineTool = ((FLandscapeToolSplines*)SplinesToolSet->GetTool());
+			return SplineTool->bAutoRotateOnJoin;
+		}
+	}
+	return true; // default value
+}
+
 void FEdModeLandscape::IntializeToolSet_Splines()
 {
 	FLandscapeToolSet* ToolSet_Splines = new(LandscapeToolSets) FLandscapeToolSet(TEXT("ToolSet_Splines"));
diff --git a/Engine/Source/Editor/LandscapeEditor/Private/LandscapeEditorDetailCustomization_MiscTools.cpp b/Engine/Source/Editor/LandscapeEditor/Private/LandscapeEditorDetailCustomization_MiscTools.cpp
--- a/Engine/Source/Editor/LandscapeEditor/Private/LandscapeEditorDetailCustomization_MiscTools.cpp
+++ b/Engine/Source/Editor/LandscapeEditor/Private/LandscapeEditorDetailCustomization_MiscTools.cpp
@@ -85,6 +85,21 @@ void FLandscapeEditorDetailCustomization_MiscTools::CustomizeDetails(IDetailLayo
 				.OnClicked_Static(&FLandscapeEditorDetailCustomization_MiscTools::OnApplySelectedSplinesButtonClicked)
 			]
 		];
+		ToolsCategory.AddCustomRow("Use Auto Rotate Control Point")
+		[
+			SNew(SHorizontalBox)
+			+ SHorizontalBox::Slot()
+			.Padding(0, 6, 0, 0)
+			[
+				SNew(SCheckBox)
+				.OnCheckStateChanged(this, &FLandscapeEditorDetailCustomization_MiscTools::OnbUseAutoRotateControlPointChanged)
+				.IsChecked(this, &FLandscapeEditorDetailCustomization_MiscTools::GetbUseAutoRotateControlPoint)
+				.Content()
+				[
+					SNew(STextBlock).Text(LOCTEXT("Spline.bUseAutoRotateControlPoint.Selected", "Use Auto Rotate Control Point"))
+				]
+			]
+		];
 	}
 
 
@@ -224,6 +239,25 @@ FReply FLandscapeEditorDetailCustomization_MiscTools::OnApplySelectedSplinesButt
 	return FReply::Handled();
 }
 
+void FLandscapeEditorDetailCustomization_MiscTools::OnbUseAutoRotateControlPointChanged(ESlateCheckBoxState::Type NewState)
+{
+	FEdModeLandscape* LandscapeEdMode = GetEditorMode();
+	if (LandscapeEdMode)
+	{
+		LandscapeEdMode->SetbUseAutoRotateOnJoin(NewState == ESlateCheckBoxState::Checked);
+	}
+}
+
+ESlateCheckBoxState::Type FLandscapeEditorDetailCustomization_MiscTools::GetbUseAutoRotateControlPoint() const
+{
+	FEdModeLandscape* LandscapeEdMode = GetEditorMode();
+	if (LandscapeEdMode)
+	{
+		return LandscapeEdMode->GetbUseAutoRotateOnJoin() ? ESlateCheckBoxState::Checked : ESlateCheckBoxState::Unchecked;
+	}
+	return ESlateCheckBoxState::Unchecked;
+}
+
 FReply FLandscapeEditorDetailCustomization_MiscTools::OnApplyRampButtonClicked()
 {
 	FEdModeLandscape* LandscapeEdMode = GetEditorMode();
diff --git a/Engine/Source/Editor/LandscapeEditor/Private/LandscapeEditorDetailCustomization_MiscTools.h b/Engine/Source/Editor/LandscapeEditor/Private/LandscapeEditorDetailCustomization_MiscTools.h
--- a/Engine/Source/Editor/LandscapeEditor/Private/LandscapeEditorDetailCustomization_MiscTools.h
+++ b/Engine/Source/Editor/LandscapeEditor/Private/LandscapeEditorDetailCustomization_MiscTools.h
@@ -30,6 +30,8 @@ protected:
 	// Splines Tool
 	static FReply OnApplyAllSplinesButtonClicked();
 	static FReply OnApplySelectedSplinesButtonClicked();
+	void OnbUseAutoRotateControlPointChanged(ESlateCheckBoxState::Type NewState);
+	ESlateCheckBoxState::Type GetbUseAutoRotateControlPoint() const;
 
 	// Ramp Tool
 	static FReply OnApplyRampButtonClicked();
diff --git a/Engine/Source/Editor/LandscapeEditor/Public/LandscapeEdMode.h b/Engine/Source/Editor/LandscapeEditor/Public/LandscapeEdMode.h
--- a/Engine/Source/Editor/LandscapeEditor/Public/LandscapeEdMode.h
+++ b/Engine/Source/Editor/LandscapeEditor/Public/LandscapeEdMode.h
@@ -479,6 +479,8 @@ public:
 	void ShowSplineProperties();
 	virtual void SelectAllConnectedSplineControlPoints();
 	virtual void SelectAllConnectedSplineSegments();
+	void SetbUseAutoRotateOnJoin(bool InbAutoRotateOnJoin);
+	bool GetbUseAutoRotateOnJoin();
 
 	void ApplyRampTool();
 	bool CanApplyRampTool();
