Landscape : Add support for moving Landscape Splines to streaming levels

diff --git a/Engine/Source/Editor/DetailCustomizations/Private/LandscapeSplineDetails.cpp b/Engine/Source/Editor/DetailCustomizations/Private/LandscapeSplineDetails.cpp
--- a/Engine/Source/Editor/DetailCustomizations/Private/LandscapeSplineDetails.cpp
+++ b/Engine/Source/Editor/DetailCustomizations/Private/LandscapeSplineDetails.cpp
@@ -44,6 +44,21 @@ void FLandscapeSplineDetails::CustomizeDetails( IDetailLayoutBuilder& DetailBuil
 			.OnClicked(this, &FLandscapeSplineDetails::OnSelectConnectedSegmentsButtonClicked)
 		]
 	];
+
+	LandscapeSplineCategory.AddCustomRow( TEXT("") )
+	[
+		SNew(SHorizontalBox)
+		+ SHorizontalBox::Slot()
+		.Padding(0, 0, 2, 0)
+		.VAlign(VAlign_Center)
+		.FillWidth(1)
+		[
+			SNew(SButton)
+			.Text(LOCTEXT("Move Selected ControlPnts+Segs to Current level", "Move to current level"))
+			.HAlign(HAlign_Center)
+			.OnClicked(this, &FLandscapeSplineDetails::OnMoveToCurrentLevelButtonClicked)
+		]
+	];
 }
 
 class FEdModeLandscape* FLandscapeSplineDetails::GetEditorMode() const
@@ -73,4 +88,15 @@ FReply FLandscapeSplineDetails::OnSelectConnectedSegmentsButtonClicked()
 	return FReply::Handled();
 }
 
+FReply FLandscapeSplineDetails::OnMoveToCurrentLevelButtonClicked()
+{
+	FEdModeLandscape* LandscapeEdMode = GetEditorMode();
+	if (LandscapeEdMode && LandscapeEdMode->CurrentToolTarget.LandscapeInfo.IsValid())
+	{
+		LandscapeEdMode->SplineMoveToCurrentLevel();
+	}
+
+	return FReply::Handled();
+}
+
 #undef LOCTEXT_NAMESPACE
diff --git a/Engine/Source/Editor/DetailCustomizations/Private/LandscapeSplineDetails.h b/Engine/Source/Editor/DetailCustomizations/Private/LandscapeSplineDetails.h
--- a/Engine/Source/Editor/DetailCustomizations/Private/LandscapeSplineDetails.h
+++ b/Engine/Source/Editor/DetailCustomizations/Private/LandscapeSplineDetails.h
@@ -15,4 +15,5 @@ private:
 	class FEdModeLandscape* GetEditorMode() const;
 	FReply OnSelectConnectedControlPointsButtonClicked();
 	FReply OnSelectConnectedSegmentsButtonClicked();
+	FReply OnMoveToCurrentLevelButtonClicked();
 };
diff --git a/Engine/Source/Editor/LandscapeEditor/Private/LandscapeEdModeSplineTools.cpp b/Engine/Source/Editor/LandscapeEditor/Private/LandscapeEdModeSplineTools.cpp
--- a/Engine/Source/Editor/LandscapeEditor/Private/LandscapeEdModeSplineTools.cpp
+++ b/Engine/Source/Editor/LandscapeEditor/Private/LandscapeEdModeSplineTools.cpp
@@ -63,7 +63,7 @@ public:
 		return true; // applied to all...
 	}
 
-	void CreateSplineComponent( ALandscape* Landscape ) 
+	void CreateSplineComponent( ALandscapeProxy* Landscape ) 
 	{
 		Landscape->SplineComponent = ConstructObject<ULandscapeSplinesComponent>(ULandscapeSplinesComponent::StaticClass(), Landscape, NAME_None, RF_Transactional);
 		Landscape->SplineComponent->RelativeScale3D = FVector(1.0f) / Landscape->GetRootComponent()->RelativeScale3D;
@@ -338,7 +338,7 @@ public:
 		NewSegment->UpdateSplinePoints();
 	}
 
-	void AddControlPoint(ALandscape* Landscape, const FVector& LocalLocation)
+	void AddControlPoint(ALandscapeProxy* Landscape, const FVector& LocalLocation)
 	{
 		FScopedTransaction Transaction( LOCTEXT("LandscapeSpline_AddControlPoint", "Add Landscape Spline Control Point") );
 
@@ -636,7 +636,7 @@ public:
 	{
 		LandscapeInfo = InTarget.LandscapeInfo.Get();
 
-		ALandscape* Landscape = LandscapeInfo->LandscapeActor.Get();
+		ALandscapeProxy* Landscape = LandscapeInfo->GetCurrentLevelLandscapeProxy(true);
 		if (!Landscape)
 		{
 			return true;
@@ -647,7 +647,7 @@ public:
 			CreateSplineComponent(Landscape);
 		}
 
-		const FTransform LandscapeToSpline = Landscape->ActorToWorld().GetRelativeTransform(Landscape->SplineComponent->ComponentToWorld);
+		const FTransform LandscapeToSpline = Landscape->LandscapeActorToWorld().GetRelativeTransform(Landscape->SplineComponent->ComponentToWorld);
 
 		AddControlPoint(Landscape, LandscapeToSpline.TransformPosition(InHitLocation));
 
@@ -1158,7 +1158,7 @@ public:
 		{
 			for (auto It = EdMode->GetLandscapeList().CreateConstIterator(); It; ++It)
 			{
-				ALandscape* Landscape = (*It).Info->LandscapeActor.Get();
+				ALandscapeProxy* Landscape = (*It).Info->GetCurrentLevelLandscapeProxy(true);
 				if (Landscape != NULL && Landscape->SplineComponent != NULL)
 				{
 					for (auto ControlPointIt = Landscape->SplineComponent->ControlPoints.CreateConstIterator(); ControlPointIt; ++ControlPointIt)
@@ -1183,7 +1183,7 @@ public:
 		{
 			for (auto It = EdMode->GetLandscapeList().CreateConstIterator(); It; ++It)
 			{
-				ALandscape* Landscape = (*It).Info->LandscapeActor.Get();
+				ALandscapeProxy* Landscape = (*It).Info->GetCurrentLevelLandscapeProxy(true);
 				if (Landscape != NULL && Landscape->SplineComponent != NULL)
 				{
 					for (auto ControlPointIt = Landscape->SplineComponent->ControlPoints.CreateConstIterator(); ControlPointIt; ++ControlPointIt)
@@ -1212,7 +1212,7 @@ public:
 
 		for (auto It = EdMode->GetLandscapeList().CreateConstIterator(); It; ++It)
 		{
-			ALandscape* Landscape = (*It).Info->LandscapeActor.Get();
+			ALandscapeProxy* Landscape = (*It).Info->GetCurrentLevelLandscapeProxy(true);
 			if (Landscape != NULL && Landscape->SplineComponent != NULL)
 			{
 				Landscape->SplineComponent->ShowSplineEditorMesh(true);
@@ -1227,7 +1227,7 @@ public:
 
 		for (auto It = EdMode->GetLandscapeList().CreateConstIterator(); It; ++It)
 		{
-			ALandscape* Landscape = (*It).Info->LandscapeActor.Get();
+			ALandscapeProxy* Landscape = (*It).Info->GetCurrentLevelLandscapeProxy(true);
 			if (Landscape != NULL && Landscape->SplineComponent != NULL)
 			{
 				Landscape->SplineComponent->ShowSplineEditorMesh(false);
@@ -1546,7 +1546,7 @@ public:
 	{
 		FScopedTransaction Transaction( LOCTEXT("LandscapeSpline_Paste", "Paste Landscape Splines") );
 
-		ALandscape* Landscape = EdMode->CurrentToolTarget.LandscapeInfo->LandscapeActor.Get();
+		ALandscapeProxy* Landscape = EdMode->CurrentToolTarget.LandscapeInfo->GetCurrentLevelLandscapeProxy(true);
 		if (!Landscape)
 		{
 			return;
@@ -1659,6 +1659,150 @@ void FEdModeLandscape::SelectAllConnectedSplineSegments()
 	}
 }
 
+void FEdModeLandscape::SplineMoveToCurrentLevel()
+{
+	FScopedTransaction Transaction(LOCTEXT("LandscapeSpline_AddControlPoint", "Add Landscape Spline Control Point"));
+
+	if (SplinesToolSet /*&& SplinesToolSet == CurrentToolSet*/)
+	{
+		if (SplinesToolSet->SetToolForTarget(CurrentToolTarget) && SplinesToolSet->GetTool())
+		{
+			// Select all connected control points
+			FLandscapeToolSplines* SplineTool = ((FLandscapeToolSplines*)SplinesToolSet->GetTool());
+			SplineTool->SelectAdjacentSegments();
+			SplineTool->SelectAdjacentControlPoints();
+			SplineTool->SelectConnected();
+
+			TSet<ALandscapeProxy*> FromProxies;
+			ALandscapeProxy* Landscape = NULL;
+
+			for (auto It = SplineTool->SelectedSplineControlPoints.CreateIterator(); It; ++It)
+			{
+				ULandscapeSplineControlPoint* ControlPoint = *It;
+				ULandscapeSplinesComponent* Comp = ControlPoint->GetOuterULandscapeSplinesComponent();
+				ALandscapeProxy* FromProxy = Comp ? Cast<ALandscapeProxy>(Comp->GetOuter()) : NULL;
+				if (FromProxy)
+				{
+					if (!Landscape)
+					{
+						ULandscapeInfo* LandscapeInfo = FromProxy->GetLandscapeInfo(false);
+						if (LandscapeInfo)
+						{
+							Landscape = LandscapeInfo->GetCurrentLevelLandscapeProxy(true);
+						}
+					}
+
+					if (Landscape && Landscape != FromProxy)
+					{
+						if (Landscape->SplineComponent == NULL)
+						{
+							SplineTool->CreateSplineComponent(Landscape);
+						}
+
+						const FTransform SplineToLanscape =
+							FromProxy->LandscapeActorToWorld().GetRelativeTransform(FromProxy->SplineComponent->ComponentToWorld).InverseSafe() *
+							Landscape->LandscapeActorToWorld().GetRelativeTransform(Landscape->SplineComponent->ComponentToWorld);
+
+						//SplineTool->AddControlPoint(Landscape, SplineToLanscape.TransformPosition(ControlPoint->Location));
+						Landscape->SplineComponent->Modify();
+						if (FromProxies.Find(FromProxy) == NULL)
+						{
+							FromProxies.Add(FromProxy);
+							FromProxy->SplineComponent->Modify();
+							FromProxy->SplineComponent->MarkRenderStateDirty();
+						}
+
+						// Add/Remove Control Point/Segment/MeshComponents
+						FromProxy->SplineComponent->ControlPoints.Remove(ControlPoint);
+						ControlPoint->Rename(NULL, Landscape->SplineComponent);
+						Landscape->SplineComponent->ControlPoints.Add(ControlPoint);
+
+						ControlPoint->Location = SplineToLanscape.TransformPosition(ControlPoint->Location);
+
+						if (ControlPoint->MeshComponent)
+						{
+							UControlPointMeshComponent* Comp = ControlPoint->MeshComponent;
+							Comp->Modify();
+							Comp->UnregisterComponent();
+							Comp->DetachFromParent(true);
+							Comp->InvalidateLightingCache();
+							Comp->Rename(NULL, Landscape);
+							Comp->AttachTo(Landscape->SplineComponent);
+						}
+						
+						ControlPoint->UpdateSplinePoints();
+					}
+				}
+			}
+
+			for (auto It = SplineTool->SelectedSplineSegments.CreateIterator(); It; ++It)
+			{
+				ULandscapeSplineSegment* Segment = *It;
+				ULandscapeSplinesComponent* Comp = Segment->GetOuterULandscapeSplinesComponent();
+				ALandscapeProxy* FromProxy = Comp ? Cast<ALandscapeProxy>(Comp->GetOuter()) : NULL;
+				if (FromProxy)
+				{
+					if (!Landscape)
+					{
+						ULandscapeInfo* LandscapeInfo = FromProxy->GetLandscapeInfo(false);
+						if (LandscapeInfo)
+						{
+							Landscape = LandscapeInfo->GetCurrentLevelLandscapeProxy(true);
+						}
+					}
+
+					if (Landscape && Landscape != FromProxy)
+					{
+						if (Landscape->SplineComponent == NULL)
+						{
+							SplineTool->CreateSplineComponent(Landscape);
+						}
+
+						Landscape->SplineComponent->Modify();
+						if (FromProxies.Find(FromProxy) == NULL)
+						{
+							FromProxies.Add(FromProxy);
+							FromProxy->SplineComponent->Modify();
+							FromProxy->SplineComponent->MarkRenderStateDirty();
+						}
+
+						// Add/Remove Control Point/Segment/MeshComponents
+						FromProxy->SplineComponent->Segments.Remove(Segment);
+						Segment->Rename(NULL, Landscape->SplineComponent);
+						Landscape->SplineComponent->Segments.Add(Segment);
+
+						for (auto It = Segment->MeshComponents.CreateIterator(); It; ++It)
+						{
+							USplineMeshComponent* Comp = *It;
+							Comp->Modify();
+							Comp->UnregisterComponent();
+							Comp->DetachFromParent(true);
+							Comp->InvalidateLightingCache();
+							Comp->Rename(NULL, Landscape);
+							Comp->AttachTo(Landscape->SplineComponent);
+						}
+
+						Segment->UpdateSplinePoints();
+					}
+				}
+			}
+
+			if (!Landscape->SplineComponent->IsRegistered())
+			{
+				Landscape->SplineComponent->RegisterComponent();
+			}
+			else
+			{
+				Landscape->SplineComponent->MarkRenderStateDirty();
+			}
+
+			SplineTool->ClearSelection();
+			SplineTool->UpdatePropertiesWindows();
+			GUnrealEd->RedrawLevelEditingViewports();
+		}
+	}
+}
+
 void FEdModeLandscape::SetbUseAutoRotateOnJoin(bool InbAutoRotateOnJoin)
 {
 	if (SplinesToolSet /*&& SplinesToolSet == CurrentToolSet*/)
diff --git a/Engine/Source/Editor/LandscapeEditor/Public/LandscapeEdMode.h b/Engine/Source/Editor/LandscapeEditor/Public/LandscapeEdMode.h
--- a/Engine/Source/Editor/LandscapeEditor/Public/LandscapeEdMode.h
+++ b/Engine/Source/Editor/LandscapeEditor/Public/LandscapeEdMode.h
@@ -479,6 +479,7 @@ public:
 	void ShowSplineProperties();
 	virtual void SelectAllConnectedSplineControlPoints();
 	virtual void SelectAllConnectedSplineSegments();
+	virtual void SplineMoveToCurrentLevel();
 	void SetbUseAutoRotateOnJoin(bool InbAutoRotateOnJoin);
 	bool GetbUseAutoRotateOnJoin();
 
diff --git a/Engine/Source/Runtime/Engine/Classes/Landscape/Landscape.h b/Engine/Source/Runtime/Engine/Classes/Landscape/Landscape.h
--- a/Engine/Source/Runtime/Engine/Classes/Landscape/Landscape.h
+++ b/Engine/Source/Runtime/Engine/Classes/Landscape/Landscape.h
@@ -22,9 +22,6 @@ class ALandscape : public ALandscapeProxy
 {
 	GENERATED_UCLASS_BODY()
 
-	UPROPERTY()
-	class ULandscapeSplinesComponent* SplineComponent;
-
 	// Make a key for XYtoComponentMap
 	static FIntPoint MakeKey( int32 X, int32 Y ) { return FIntPoint(X, Y); }
 	static void UnpackKey( FIntPoint Key, int32& OutX, int32& OutY ) { OutX = Key.X; OutY = Key.Y; }
diff --git a/Engine/Source/Runtime/Engine/Classes/Landscape/LandscapeInfo.h b/Engine/Source/Runtime/Engine/Classes/Landscape/LandscapeInfo.h
--- a/Engine/Source/Runtime/Engine/Classes/Landscape/LandscapeInfo.h
+++ b/Engine/Source/Runtime/Engine/Classes/Landscape/LandscapeInfo.h
@@ -166,6 +166,7 @@ public:
 	ENGINE_API void ExportHeightmap(const FString& Filename);
 	ENGINE_API void ExportLayer(ULandscapeLayerInfoObject* LayerInfo, const FString& Filename);
 	ENGINE_API bool ApplySplines(bool bOnlySelected);
+	bool ApplySplinesInternal(bool bOnlySelected, ALandscapeProxy* Landscape);
 
 	ENGINE_API bool GetSelectedExtent(int32& MinX, int32& MinY, int32& MaxX, int32& MaxY);
 	FVector GetLandscapeCenterPos(float& LengthZ, int32 MinX = MAX_int32, int32 MinY = MAX_int32, int32 MaxX = MIN_int32, int32 MaxY = MIN_int32);
diff --git a/Engine/Source/Runtime/Engine/Classes/Landscape/LandscapeProxy.h b/Engine/Source/Runtime/Engine/Classes/Landscape/LandscapeProxy.h
--- a/Engine/Source/Runtime/Engine/Classes/Landscape/LandscapeProxy.h
+++ b/Engine/Source/Runtime/Engine/Classes/Landscape/LandscapeProxy.h
@@ -169,6 +169,9 @@ class ALandscapeProxy : public AActor, public INavRelevantActorInterface
 {
 	GENERATED_UCLASS_BODY()
 
+	UPROPERTY()
+	class ULandscapeSplinesComponent* SplineComponent;
+
 protected:
 	/** Guid for LandscapeEditorInfo **/
 	UPROPERTY()
@@ -404,8 +407,6 @@ public:
 	/** @return Current size of bounding rectangle in quads space */
 	ENGINE_API FIntRect GetBoundingRect() const;
 #endif
-
-
 };
 
 
diff --git a/Engine/Source/Runtime/Engine/Classes/Landscape/LandscapeSplineControlPoint.h b/Engine/Source/Runtime/Engine/Classes/Landscape/LandscapeSplineControlPoint.h
--- a/Engine/Source/Runtime/Engine/Classes/Landscape/LandscapeSplineControlPoint.h
+++ b/Engine/Source/Runtime/Engine/Classes/Landscape/LandscapeSplineControlPoint.h
@@ -174,4 +174,6 @@ public:
 	void UnregisterComponents();
 
 	virtual bool OwnsComponent(const class UControlPointMeshComponent* ControlPointMeshComponent) const;
+
+	friend class FEdModeLandscape;
 };
diff --git a/Engine/Source/Runtime/Engine/Classes/Landscape/LandscapeSplineSegment.h b/Engine/Source/Runtime/Engine/Classes/Landscape/LandscapeSplineSegment.h
--- a/Engine/Source/Runtime/Engine/Classes/Landscape/LandscapeSplineSegment.h
+++ b/Engine/Source/Runtime/Engine/Classes/Landscape/LandscapeSplineSegment.h
@@ -251,4 +251,6 @@ protected:
 #if WITH_EDITOR
 	bool FixSelfIntersection(FVector FLandscapeSplineInterpPoint::* Side);
 #endif
+
+	friend class FEdModeLandscape;
 };
diff --git a/Engine/Source/Runtime/Engine/Private/Landscape/LandscapeEdit.cpp b/Engine/Source/Runtime/Engine/Private/Landscape/LandscapeEdit.cpp
--- a/Engine/Source/Runtime/Engine/Private/Landscape/LandscapeEdit.cpp
+++ b/Engine/Source/Runtime/Engine/Private/Landscape/LandscapeEdit.cpp
@@ -2795,10 +2795,25 @@ void ShrinkData(TArray<T>& Data, int32 OldMinX, int32 OldMinY, int32 OldMaxX, in
 	}
 }
 
-
 bool ULandscapeInfo::ApplySplines(bool bOnlySelected)
 {
+	bool bResult = false;
+
 	ALandscape* Landscape = LandscapeActor.Get();
+
+	bResult |= ApplySplinesInternal(bOnlySelected, Landscape);
+
+	for (auto It = Proxies.CreateIterator(); It; ++It)
+	{
+		ALandscapeProxy* LandscapeProxy = (*It);
+		bResult |= ApplySplinesInternal(bOnlySelected, LandscapeProxy);
+	}
+
+	return bResult;
+}
+
+bool ULandscapeInfo::ApplySplinesInternal(bool bOnlySelected, ALandscapeProxy* Landscape)
+{
 	if (!Landscape || !Landscape->SplineComponent || Landscape->SplineComponent->ControlPoints.Num() == 0 || Landscape->SplineComponent->Segments.Num() == 0)
 	{
 		return false;
