Landscape Spline : Fixed MoveToCurrentLevel (Relative3D scale)
- Disable button when there is no landscape proxy for current level

diff --git a/Engine/Source/Editor/DetailCustomizations/Private/LandscapeSplineDetails.cpp b/Engine/Source/Editor/DetailCustomizations/Private/LandscapeSplineDetails.cpp
--- a/Engine/Source/Editor/DetailCustomizations/Private/LandscapeSplineDetails.cpp
+++ b/Engine/Source/Editor/DetailCustomizations/Private/LandscapeSplineDetails.cpp
@@ -57,6 +57,7 @@ void FLandscapeSplineDetails::CustomizeDetails( IDetailLayoutBuilder& DetailBuil
 			.Text(LOCTEXT("Move Selected ControlPnts+Segs to Current level", "Move to current level"))
 			.HAlign(HAlign_Center)
 			.OnClicked(this, &FLandscapeSplineDetails::OnMoveToCurrentLevelButtonClicked)
+			.IsEnabled(this, &FLandscapeSplineDetails::IsMoveToCurrentLevelButtonEnabled)
 		]
 	];
 }
@@ -91,7 +92,7 @@ FReply FLandscapeSplineDetails::OnSelectConnectedSegmentsButtonClicked()
 FReply FLandscapeSplineDetails::OnMoveToCurrentLevelButtonClicked()
 {
 	FEdModeLandscape* LandscapeEdMode = GetEditorMode();
-	if (LandscapeEdMode && LandscapeEdMode->CurrentToolTarget.LandscapeInfo.IsValid())
+	if (LandscapeEdMode && LandscapeEdMode->CurrentToolTarget.LandscapeInfo.IsValid() && LandscapeEdMode->CurrentToolTarget.LandscapeInfo->GetCurrentLevelLandscapeProxy(true))
 	{
 		LandscapeEdMode->SplineMoveToCurrentLevel();
 	}
@@ -99,4 +100,10 @@ FReply FLandscapeSplineDetails::OnMoveToCurrentLevelButtonClicked()
 	return FReply::Handled();
 }
 
+bool FLandscapeSplineDetails::IsMoveToCurrentLevelButtonEnabled() const
+{
+	FEdModeLandscape* LandscapeEdMode = GetEditorMode();
+	return (LandscapeEdMode && LandscapeEdMode->CurrentToolTarget.LandscapeInfo.IsValid() && LandscapeEdMode->CurrentToolTarget.LandscapeInfo->GetCurrentLevelLandscapeProxy(true));
+}
+
 #undef LOCTEXT_NAMESPACE
diff --git a/Engine/Source/Editor/DetailCustomizations/Private/LandscapeSplineDetails.h b/Engine/Source/Editor/DetailCustomizations/Private/LandscapeSplineDetails.h
--- a/Engine/Source/Editor/DetailCustomizations/Private/LandscapeSplineDetails.h
+++ b/Engine/Source/Editor/DetailCustomizations/Private/LandscapeSplineDetails.h
@@ -16,4 +16,5 @@ private:
 	FReply OnSelectConnectedControlPointsButtonClicked();
 	FReply OnSelectConnectedSegmentsButtonClicked();
 	FReply OnMoveToCurrentLevelButtonClicked();
+	bool IsMoveToCurrentLevelButtonEnabled() const;
 };
diff --git a/Engine/Source/Editor/LandscapeEditor/Private/LandscapeEdModeSplineTools.cpp b/Engine/Source/Editor/LandscapeEditor/Private/LandscapeEdModeSplineTools.cpp
--- a/Engine/Source/Editor/LandscapeEditor/Private/LandscapeEdModeSplineTools.cpp
+++ b/Engine/Source/Editor/LandscapeEditor/Private/LandscapeEdModeSplineTools.cpp
@@ -61,10 +61,10 @@ public:
 		return true; // applied to all...
 	}
 
-	void CreateSplineComponent( ALandscapeProxy* Landscape ) 
+	void CreateSplineComponent( ALandscapeProxy* Landscape, FVector Scale3D ) 
 	{
 		Landscape->SplineComponent = ConstructObject<ULandscapeSplinesComponent>(ULandscapeSplinesComponent::StaticClass(), Landscape, NAME_None, RF_Transactional);
-		Landscape->SplineComponent->RelativeScale3D = FVector(1.0f) / Landscape->GetRootComponent()->RelativeScale3D;
+		Landscape->SplineComponent->RelativeScale3D = Scale3D;
 		Landscape->SplineComponent->AttachTo(Landscape->GetRootComponent());
 		Landscape->SplineComponent->ShowSplineEditorMesh(true);
 	}
@@ -642,7 +642,7 @@ public:
 
 		if (!Landscape->SplineComponent)
 		{
-			CreateSplineComponent(Landscape);
+			CreateSplineComponent(Landscape, FVector(1.f) / Landscape->GetRootComponent()->RelativeScale3D);
 		}
 
 		const FTransform LandscapeToSpline = Landscape->LandscapeActorToWorld().GetRelativeTransform(Landscape->SplineComponent->ComponentToWorld);
@@ -1551,7 +1551,7 @@ public:
 		}
 		if (!Landscape->SplineComponent)
 		{
-			CreateSplineComponent(Landscape);
+			CreateSplineComponent(Landscape, FVector(1.f) / Landscape->GetRootComponent()->RelativeScale3D);
 		}
 
 		const TCHAR* Data = NULL;
@@ -1687,6 +1687,11 @@ void FEdModeLandscape::SplineMoveToCurrentLevel()
 						if (LandscapeInfo)
 						{
 							Landscape = LandscapeInfo->GetCurrentLevelLandscapeProxy(true);
+							if (!Landscape)
+							{
+								// No Landscape Proxy, don't support for creating only for Spline now
+								return;
+							}
 						}
 					}
 
@@ -1694,7 +1699,7 @@ void FEdModeLandscape::SplineMoveToCurrentLevel()
 					{
 						if (Landscape->SplineComponent == NULL)
 						{
-							SplineTool->CreateSplineComponent(Landscape);
+							SplineTool->CreateSplineComponent(Landscape, FromProxy->SplineComponent->RelativeScale3D);
 						}
 
 						const FTransform SplineToLanscape =
@@ -1725,7 +1730,7 @@ void FEdModeLandscape::SplineMoveToCurrentLevel()
 							Comp->DetachFromParent(true);
 							Comp->InvalidateLightingCache();
 							Comp->Rename(NULL, Landscape);
-							Comp->AttachTo(Landscape->SplineComponent);
+							Comp->AttachTo(Landscape->SplineComponent, NAME_None, EAttachLocation::KeepWorldPosition);
 						}
 						
 						ControlPoint->UpdateSplinePoints();
@@ -1753,7 +1758,7 @@ void FEdModeLandscape::SplineMoveToCurrentLevel()
 					{
 						if (Landscape->SplineComponent == NULL)
 						{
-							SplineTool->CreateSplineComponent(Landscape);
+							SplineTool->CreateSplineComponent(Landscape, FromProxy->SplineComponent->RelativeScale3D);
 						}
 
 						Landscape->SplineComponent->Modify();
@@ -1777,7 +1782,7 @@ void FEdModeLandscape::SplineMoveToCurrentLevel()
 							Comp->DetachFromParent(true);
 							Comp->InvalidateLightingCache();
 							Comp->Rename(NULL, Landscape);
-							Comp->AttachTo(Landscape->SplineComponent);
+							Comp->AttachTo(Landscape->SplineComponent, NAME_None, EAttachLocation::KeepWorldPosition);
 						}
 
 						Segment->UpdateSplinePoints();
@@ -1785,13 +1790,16 @@ void FEdModeLandscape::SplineMoveToCurrentLevel()
 				}
 			}
 
-			if (!Landscape->SplineComponent->IsRegistered())
-			{
-				Landscape->SplineComponent->RegisterComponent();
-			}
-			else
+			if (Landscape && Landscape->SplineComponent)
 			{
-				Landscape->SplineComponent->MarkRenderStateDirty();
+				if (!Landscape->SplineComponent->IsRegistered())
+				{
+					Landscape->SplineComponent->RegisterComponent();
+				}
+				else
+				{
+					Landscape->SplineComponent->MarkRenderStateDirty();
+				}
 			}
 
 			SplineTool->ClearSelection();
