Fixed TTP 334878 - Regression: EDITOR: LANDSCAPE: Landscape brush outline is still visible in orthographic viewports when switching from landscape tool

diff --git a/Engine/Source/Editor/LandscapeEditor/Private/LandscapeEdMode.cpp b/Engine/Source/Editor/LandscapeEditor/Private/LandscapeEdMode.cpp
--- a/Engine/Source/Editor/LandscapeEditor/Private/LandscapeEdMode.cpp
+++ b/Engine/Source/Editor/LandscapeEditor/Private/LandscapeEdMode.cpp
@@ -552,14 +552,14 @@ void FEdModeLandscape::Exit()
 		{
 			ALandscapeGizmoActor* Gizmo = Cast<ALandscapeGizmoActor>(*It);
 			if (Gizmo)
-			{	
+			{
 				GetWorld()->DestroyActor(Gizmo, false, false);
 			}
 		}
 	}
 
-	// Recreate Hole Collision
-	GEngine->DeferredCommands.AddUnique(TEXT("UpdateLandscapeHoleCollision"));
+	// Redraw one last time to remove any landscape editor stuff from view
+	GEditor->RedrawLevelEditingViewports();
 
 	// Call parent implementation
 	FEdMode::Exit();
@@ -2281,16 +2281,15 @@ void FEdModeLandscape::ActorMoveNotify()
 }
 
 /** Forces all level editor viewports to realtime mode */
-void FEdModeLandscape::ForceRealTimeViewports( const bool bEnable, const bool bStoreCurrentState )
+void FEdModeLandscape::ForceRealTimeViewports(const bool bEnable, const bool bStoreCurrentState)
 {
-	FLevelEditorModule& LevelEditorModule = FModuleManager::GetModuleChecked<FLevelEditorModule>( "LevelEditor");
+	FLevelEditorModule& LevelEditorModule = FModuleManager::GetModuleChecked<FLevelEditorModule>("LevelEditor");
 	TSharedPtr<ILevelEditor> LevelEditor = LevelEditorModule.GetFirstLevelEditor();
 	if (LevelEditor.IsValid())
 	{
 		TArray<TSharedPtr<ILevelViewport>> Viewports = LevelEditor->GetViewports();
-		for (auto It = Viewports.CreateConstIterator(); It; It++)
+		for (const TSharedPtr<ILevelViewport>& ViewportWindow :  Viewports)
 		{
-			const TSharedPtr<ILevelViewport>& ViewportWindow = *It;
 			if (ViewportWindow.IsValid())
 			{
 				FLevelEditorViewportClient& Viewport = ViewportWindow->GetLevelViewportClient();
