Fix various undo-related crashes in landscape spline tool

diff --git a/Engine/Source/Editor/LandscapeEditor/Private/LandscapeEdModeSplineTools.cpp b/Engine/Source/Editor/LandscapeEditor/Private/LandscapeEdModeSplineTools.cpp
--- a/Engine/Source/Editor/LandscapeEditor/Private/LandscapeEdModeSplineTools.cpp
+++ b/Engine/Source/Editor/LandscapeEditor/Private/LandscapeEdModeSplineTools.cpp
@@ -80,13 +80,13 @@ public:
 			TArray<UObject*> Objects;
 			Objects.Reset(SelectedSplineControlPoints.Num() + SelectedSplineSegments.Num());
 
-			for (auto It = SelectedSplineControlPoints.CreateConstIterator(); It; ++It)
+			for (ULandscapeSplineControlPoint* ControlPoint : SelectedSplineControlPoints)
 			{
-				Objects.Add(*It);
+				Objects.Add(ControlPoint);
 			}
-			for (auto It = SelectedSplineSegments.CreateConstIterator(); It; ++It)
+			for (ULandscapeSplineSegment* Segment : SelectedSplineSegments)
 			{
-				Objects.Add(*It);
+				Objects.Add(Segment);
 			}
 
 
@@ -97,22 +97,22 @@ public:
 
 	void ClearSelectedControlPoints()
 	{
-		for (auto It = SelectedSplineControlPoints.CreateConstIterator(); It; ++It)
+		for (ULandscapeSplineControlPoint* ControlPoint : SelectedSplineControlPoints)
 		{
-			checkSlow((*It)->IsSplineSelected());
-			(*It)->Modify();
-			(*It)->SetSplineSelected(false);
+			checkSlow(ControlPoint->IsSplineSelected());
+			ControlPoint->Modify();
+			ControlPoint->SetSplineSelected(false);
 		}
 		SelectedSplineControlPoints.Empty();
 	}
 
 	void ClearSelectedSegments()
 	{
-		for (auto It = SelectedSplineSegments.CreateConstIterator(); It; ++It)
+		for (ULandscapeSplineSegment* Segment : SelectedSplineSegments)
 		{
-			checkSlow((*It)->IsSplineSelected());
-			(*It)->Modify();
-			(*It)->SetSplineSelected(false);
+			checkSlow(Segment->IsSplineSelected());
+			Segment->Modify();
+			Segment->SetSplineSelected(false);
 		}
 		SelectedSplineSegments.Empty();
 	}
@@ -165,9 +165,8 @@ public:
 		{
 			const ULandscapeSplineControlPoint* ControlPoint = ControlPointsToProcess.Pop();
 
-			for (int32 i = 0; i < ControlPoint->ConnectedSegments.Num(); i++)
+			for (const FLandscapeSplineConnection& Connection : ControlPoint->ConnectedSegments)
 			{
-				const FLandscapeSplineConnection& Connection = ControlPoint->ConnectedSegments[i];
 				ULandscapeSplineControlPoint* OtherEnd = Connection.GetFarConnection().ControlPoint;
 
 				if (!OtherEnd->IsSplineSelected())
@@ -184,14 +183,10 @@ public:
 		{
 			const ULandscapeSplineSegment* Segment = SegmentsToProcess.Pop();
 
-			for (int32 End = 0; End <= 1; End++)
+			for (const FLandscapeSplineSegmentConnection& SegmentConnection : Segment->Connections)
 			{
-				const ULandscapeSplineControlPoint* ControlPoint = Segment->Connections[End].ControlPoint;
-
-				for (int32 i = 0; i < ControlPoint->ConnectedSegments.Num(); i++)
+				for (const FLandscapeSplineConnection& Connection : SegmentConnection.ControlPoint->ConnectedSegments)
 				{
-					const FLandscapeSplineConnection& Connection = ControlPoint->ConnectedSegments[i];
-
 					if (Connection.Segment != Segment && !Connection.Segment->IsSplineSelected())
 					{
 						SelectSegment(Connection.Segment);
@@ -204,9 +199,8 @@ public:
 
 	void SelectAdjacentControlPoints()
 	{
-		for (auto It = SelectedSplineSegments.CreateConstIterator(); It; ++It)
+		for (ULandscapeSplineSegment* Segment : SelectedSplineSegments)
 		{
-			const ULandscapeSplineSegment* Segment = *It;
 			if (!Segment->Connections[0].ControlPoint->IsSplineSelected())
 			{
 				SelectControlPoint(Segment->Connections[0].ControlPoint);
@@ -220,13 +214,10 @@ public:
 
 	void SelectAdjacentSegments()
 	{
-		for (auto It = SelectedSplineControlPoints.CreateConstIterator(); It; ++It)
+		for (ULandscapeSplineControlPoint* ControlPoint : SelectedSplineControlPoints)
 		{
-			const ULandscapeSplineControlPoint* ControlPoint = *It;
-			for (int32 i = 0; i < ControlPoint->ConnectedSegments.Num(); i++)
+			for (const FLandscapeSplineConnection& Connection : ControlPoint->ConnectedSegments)
 			{
-				const FLandscapeSplineConnection& Connection = ControlPoint->ConnectedSegments[i];
-
 				if (!Connection.Segment->IsSplineSelected())
 				{
 					SelectSegment(Connection.Segment);
@@ -251,10 +242,8 @@ public:
 			return;
 		}
 
-		for (int32 i = 0; i < Start->ConnectedSegments.Num(); i++)
+		for (const FLandscapeSplineConnection& Connection : Start->ConnectedSegments)
 		{
-			const FLandscapeSplineConnection& Connection = Start->ConnectedSegments[i];
-
 			// if the *other* end on the connected segment connects to the "end" control point...
 			if (Connection.GetFarConnection().ControlPoint == End)
 			{
@@ -367,9 +356,9 @@ public:
 				NewControlPoint->bCastShadow = FirstPoint->bCastShadow;
 			}
 
-			for (auto It = SelectedSplineControlPoints.CreateConstIterator(); It; ++It)
+			for (ULandscapeSplineControlPoint* ControlPoint : SelectedSplineControlPoints)
 			{
-				AddSegment(*It, NewControlPoint, bAutoRotateOnJoin, true);
+				AddSegment(ControlPoint, NewControlPoint, bAutoRotateOnJoin, true);
 			}
 		}
 
@@ -484,9 +473,8 @@ public:
 			}
 		}
 
-		for (int32 i = 0; i < ToDelete->ConnectedSegments.Num(); i++)
+		for (FLandscapeSplineConnection& Connection : ToDelete->ConnectedSegments)
 		{
-			FLandscapeSplineConnection& Connection = ToDelete->ConnectedSegments[i];
 			Connection.Segment->Modify();
 			Connection.Segment->DeleteSplinePoints();
 
@@ -614,13 +602,13 @@ public:
 		TArray<UObject*> Objects;
 		Objects.Reset(SelectedSplineControlPoints.Num() + SelectedSplineSegments.Num());
 
-		for (auto It = SelectedSplineControlPoints.CreateConstIterator(); It; ++It)
+		for (ULandscapeSplineControlPoint* ControlPoint : SelectedSplineControlPoints)
 		{
-			Objects.Add(*It);
+			Objects.Add(ControlPoint);
 		}
-		for (auto It = SelectedSplineSegments.CreateConstIterator(); It; ++It)
+		for (ULandscapeSplineSegment* Segment : SelectedSplineSegments)
 		{
-			Objects.Add(*It);
+			Objects.Add(Segment);
 		}
 
 		FPropertyEditorModule& PropertyModule = FModuleManager::Get().LoadModuleChecked<FPropertyEditorModule>( TEXT("PropertyEditor") );
@@ -641,7 +629,7 @@ public:
 		ALandscapeProxy* Landscape = LandscapeInfo->GetCurrentLevelLandscapeProxy(true);
 		if (!Landscape)
 		{
-			return true;
+			return false;
 		}
 
 		if (!Landscape->SplineComponent)
@@ -666,7 +654,7 @@ public:
 	virtual bool MouseMove( FLevelEditorViewportClient* ViewportClient, FViewport* Viewport, int32 x, int32 y ) OVERRIDE
 	{
 		FVector HitLocation;
-		if( EdMode->LandscapeMouseTrace(ViewportClient, x, y, HitLocation)  )
+		if (EdMode->LandscapeMouseTrace(ViewportClient, x, y, HitLocation))
 		{
 			//if( bToolActive )
 			//{
@@ -676,7 +664,7 @@ public:
 		}
 
 		return true;
-	}	
+	}
 
 	virtual void ApplyTool( FLevelEditorViewportClient* ViewportClient )
 	{
@@ -710,16 +698,14 @@ public:
 			else if (HitProxy->IsA(HActor::StaticGetType()) )
 			{
 				HActor* ActorProxy = (HActor*)HitProxy;
-				ULandscapeSplinesComponent* SplineComponent = ActorProxy->Actor->FindComponentByClass<ULandscapeSplinesComponent>();
+				AActor* Actor = ActorProxy->Actor;
+				ULandscapeSplinesComponent* SplineComponent = Actor->FindComponentByClass<ULandscapeSplinesComponent>();
 				if (SplineComponent != NULL)
 				{
-					const UControlPointMeshComponent* ControlPointMeshComponent = Cast<const UControlPointMeshComponent>(ActorProxy->PrimComponent);
-					const USplineMeshComponent* SplineMeshComponent = Cast<const USplineMeshComponent>(ActorProxy->PrimComponent);
-					if (ControlPointMeshComponent != NULL)
+					if (const UControlPointMeshComponent* ControlPointMeshComponent = Cast<const UControlPointMeshComponent>(ActorProxy->PrimComponent))
 					{
-						for (auto It = SplineComponent->ControlPoints.CreateConstIterator(); It; ++It)
+						for (ULandscapeSplineControlPoint* ControlPoint : SplineComponent->ControlPoints)
 						{
-							ULandscapeSplineControlPoint* ControlPoint = *It;
 							if (ControlPoint->OwnsComponent(ControlPointMeshComponent))
 							{
 								ClickedControlPoint = ControlPoint;
@@ -727,11 +713,10 @@ public:
 							}
 						}
 					}
-					else if (SplineMeshComponent != NULL)
+					else if (const USplineMeshComponent* SplineMeshComponent = Cast<const USplineMeshComponent>(ActorProxy->PrimComponent))
 					{
-						for (auto It = SplineComponent->Segments.CreateConstIterator(); It; ++It)
+						for (ULandscapeSplineSegment* SplineSegment : SplineComponent->Segments)
 						{
-							ULandscapeSplineSegment* SplineSegment = *It;
 							if (SplineSegment->OwnsComponent(SplineMeshComponent))
 							{
 								ClickedSplineSegment = SplineSegment;
@@ -797,18 +782,18 @@ public:
 			{
 				FScopedTransaction Transaction( LOCTEXT("LandscapeSpline_AutoRotate", "Auto-rotate Landscape Spline Control Points") );
 
-				for (auto It = SelectedSplineControlPoints.CreateConstIterator(); It; ++It)
+				for (ULandscapeSplineControlPoint* ControlPoint : SelectedSplineControlPoints)
 				{
-					(*It)->AutoCalcRotation();
-					(*It)->UpdateSplinePoints();
+					ControlPoint->AutoCalcRotation();
+					ControlPoint->UpdateSplinePoints();
 				}
 
-				for (auto It = SelectedSplineSegments.CreateConstIterator(); It; ++It)
+				for (ULandscapeSplineSegment* Segment : SelectedSplineSegments)
 				{
-					(*It)->Connections[0].ControlPoint->AutoCalcRotation();
-					(*It)->Connections[0].ControlPoint->UpdateSplinePoints();
-					(*It)->Connections[1].ControlPoint->AutoCalcRotation();
-					(*It)->Connections[1].ControlPoint->UpdateSplinePoints();
+					Segment->Connections[0].ControlPoint->AutoCalcRotation();
+					Segment->Connections[0].ControlPoint->UpdateSplinePoints();
+					Segment->Connections[1].ControlPoint->AutoCalcRotation();
+					Segment->Connections[1].ControlPoint->UpdateSplinePoints();
 				}
 
 				return true;
@@ -821,9 +806,9 @@ public:
 			{
 				FScopedTransaction Transaction( LOCTEXT("LandscapeSpline_FlipSegments", "Flip Landscape Spline Segments") );
 
-				for (auto It = SelectedSplineSegments.CreateConstIterator(); It; ++It)
+				for (ULandscapeSplineSegment* Segment : SelectedSplineSegments)
 				{
-					FlipSegment(*It);
+					FlipSegment(Segment);
 				}
 
 				return true;
@@ -836,18 +821,18 @@ public:
 			{
 				FScopedTransaction Transaction( LOCTEXT("LandscapeSpline_AutoFlipTangents", "Auto-flip Landscape Spline Tangents") );
 
-				for (auto It = SelectedSplineControlPoints.CreateConstIterator(); It; ++It)
+				for (ULandscapeSplineControlPoint* ControlPoint : SelectedSplineControlPoints)
 				{
-					(*It)->AutoFlipTangents();
-					(*It)->UpdateSplinePoints();
+					ControlPoint->AutoFlipTangents();
+					ControlPoint->UpdateSplinePoints();
 				}
 
-				for (auto It = SelectedSplineSegments.CreateConstIterator(); It; ++It)
+				for (ULandscapeSplineSegment* Segment : SelectedSplineSegments)
 				{
-					(*It)->Connections[0].ControlPoint->AutoFlipTangents();
-					(*It)->Connections[0].ControlPoint->UpdateSplinePoints();
-					(*It)->Connections[1].ControlPoint->AutoFlipTangents();
-					(*It)->Connections[1].ControlPoint->UpdateSplinePoints();
+					Segment->Connections[0].ControlPoint->AutoFlipTangents();
+					Segment->Connections[0].ControlPoint->UpdateSplinePoints();
+					Segment->Connections[1].ControlPoint->AutoFlipTangents();
+					Segment->Connections[1].ControlPoint->UpdateSplinePoints();
 				}
 
 				return true;
@@ -860,14 +845,14 @@ public:
 			{
 				FScopedTransaction Transaction( LOCTEXT("LandscapeSpline_SnapToGround", "Snap Landscape Spline to Ground") );
 
-				for (auto It = SelectedSplineControlPoints.CreateConstIterator(); It; ++It)
+				for (ULandscapeSplineControlPoint* ControlPoint : SelectedSplineControlPoints)
 				{
-					SnapControlPointToGround((*It));
+					SnapControlPointToGround(ControlPoint);
 				}
-				for (auto It = SelectedSplineSegments.CreateConstIterator(); It; ++It)
+				for (ULandscapeSplineSegment* Segment : SelectedSplineSegments)
 				{
-					SnapControlPointToGround((*It)->Connections[0].ControlPoint);
-					SnapControlPointToGround((*It)->Connections[1].ControlPoint);
+					SnapControlPointToGround(Segment->Connections[0].ControlPoint);
+					SnapControlPointToGround(Segment->Connections[1].ControlPoint);
 				}
 				UpdatePropertiesWindows();
 
@@ -897,7 +882,7 @@ public:
 			{
 				int32 HitX = InViewport->GetMouseX();
 				int32 HitY = InViewport->GetMouseY();
-				HHitProxy*	HitProxy = InViewport->GetHitProxy(HitX, HitY);
+				HHitProxy* HitProxy = InViewport->GetHitProxy(HitX, HitY);
 				if (HitProxy != NULL)
 				{
 					ULandscapeSplineControlPoint* ClickedControlPoint = NULL;
@@ -916,9 +901,8 @@ public:
 							const UControlPointMeshComponent* ControlPointMeshComponent = Cast<const UControlPointMeshComponent>(ActorProxy->PrimComponent);
 							if (ControlPointMeshComponent != NULL)
 							{
-								for (auto It = SplineComponent->ControlPoints.CreateConstIterator(); It; ++It)
+								for (ULandscapeSplineControlPoint* ControlPoint : SplineComponent->ControlPoints)
 								{
-									ULandscapeSplineControlPoint* ControlPoint = *It;
 									if (ControlPoint->OwnsComponent(ControlPointMeshComponent))
 									{
 										ClickedControlPoint = ControlPoint;
@@ -933,9 +917,9 @@ public:
 					{
 						FScopedTransaction Transaction( LOCTEXT("LandscapeSpline_AddSegment", "Add Landscape Spline Segment") );
 
-						for (auto It = SelectedSplineControlPoints.CreateConstIterator(); It; ++It)
+						for (ULandscapeSplineControlPoint* ControlPoint : SelectedSplineControlPoints)
 						{
-							AddSegment((*It), ClickedControlPoint, bAutoRotateOnJoin, bAutoRotateOnJoin);
+							AddSegment(ControlPoint, ClickedControlPoint, bAutoRotateOnJoin, bAutoRotateOnJoin);
 						}
 
 						GUnrealEd->RedrawLevelEditingViewports();
@@ -975,9 +959,8 @@ public:
 							ULandscapeSplinesComponent* SplineComponent = ActorProxy->Actor->FindComponentByClass<ULandscapeSplinesComponent>();
 							if (SplineComponent != NULL)
 							{
-								for (auto It = SplineComponent->Segments.CreateConstIterator(); It; ++It)
+								for (ULandscapeSplineSegment* SplineSegment : SplineComponent->Segments)
 								{
-									ULandscapeSplineSegment* SplineSegment = *It;
 									if (SplineSegment->OwnsComponent(SplineMeshComponent))
 									{
 										ClickedSplineSegment = SplineSegment;
@@ -992,9 +975,9 @@ public:
 					if (ClickedSplineSegment != NULL)
 					{
 						FVector HitLocation;
-						if( EdMode->LandscapeMouseTrace(InViewportClient, HitLocation) )
+						if (EdMode->LandscapeMouseTrace(InViewportClient, HitLocation))
 						{
-							FScopedTransaction Transaction( LOCTEXT("LandscapeSpline_SplitSegment", "Split Landscape Spline Segment") );
+							FScopedTransaction Transaction(LOCTEXT("LandscapeSpline_SplitSegment", "Split Landscape Spline Segment"));
 
 							SplitSegment(ClickedSplineSegment, LandscapeToSpline.TransformPosition(HitLocation));
 
@@ -1024,10 +1007,10 @@ public:
 						bMovingControlPoint = true;
 
 						GEditor->BeginTransaction( LOCTEXT("LandscapeSpline_ModifyControlPoint", "Modify Landscape Spline Control Point") );
-						for (auto It = SelectedSplineControlPoints.CreateConstIterator(); It; ++It)
+						for (ULandscapeSplineControlPoint* ControlPoint : SelectedSplineControlPoints)
 						{
-							(*It)->Modify();
-							(*It)->GetOuterULandscapeSplinesComponent()->Modify();
+							ControlPoint->Modify();
+							ControlPoint->GetOuterULandscapeSplinesComponent()->Modify();
 						}
 
 						return false; // We're not actually handling this case ourselves, just wrapping it in a transaction
@@ -1053,9 +1036,9 @@ public:
 				{
 					bMovingControlPoint = false;
 
-					for (auto It = SelectedSplineControlPoints.CreateConstIterator(); It; ++It)
+					for (ULandscapeSplineControlPoint* ControlPoint : SelectedSplineControlPoints)
 					{
-						(*It)->UpdateSplinePoints(true);
+						ControlPoint->UpdateSplinePoints(true);
 					}
 
 					GEditor->EndTransaction();
@@ -1119,9 +1102,8 @@ public:
 
 		if (SelectedSplineControlPoints.Num() > 0 && InViewportClient->GetCurrentWidgetAxis() != EAxisList::None)
 		{
-			for (auto It = SelectedSplineControlPoints.CreateConstIterator(); It; ++It)
+			for (ULandscapeSplineControlPoint* ControlPoint : SelectedSplineControlPoints)
 			{
-				ULandscapeSplineControlPoint* ControlPoint = (*It);
 				const ULandscapeSplinesComponent* SplinesComponent = ControlPoint->GetOuterULandscapeSplinesComponent();
 
 				ControlPoint->Location += SplinesComponent->ComponentToWorld.InverseTransformVector(Drag);
@@ -1158,24 +1140,46 @@ public:
 
 		if (EdMode->CurrentToolSet != NULL && EdMode->CurrentToolSet->GetTool() == this)
 		{
-			for (auto It = EdMode->GetLandscapeList().CreateConstIterator(); It; ++It)
+			for (const FLandscapeListInfo& Info : EdMode->GetLandscapeList())
 			{
-				ALandscapeProxy* Landscape = (*It).Info->GetCurrentLevelLandscapeProxy(true);
+				ALandscape* Landscape = Info.Info->LandscapeActor.Get();
 				if (Landscape != NULL && Landscape->SplineComponent != NULL)
 				{
-					for (auto ControlPointIt = Landscape->SplineComponent->ControlPoints.CreateConstIterator(); ControlPointIt; ++ControlPointIt)
+					for (ULandscapeSplineControlPoint* ControlPoint : Landscape->SplineComponent->ControlPoints)
 					{
-						if ((*ControlPointIt)->IsSplineSelected())
+						if (ControlPoint->IsSplineSelected())
 						{
-							SelectedSplineControlPoints.Add(*ControlPointIt);
+							SelectedSplineControlPoints.Add(ControlPoint);
 						}
 					}
 
-					for (auto SegmentIt = Landscape->SplineComponent->Segments.CreateConstIterator(); SegmentIt; ++SegmentIt)
+					for (ULandscapeSplineSegment* Segment : Landscape->SplineComponent->Segments)
 					{
-						if ((*SegmentIt)->IsSplineSelected())
+						if (Segment->IsSplineSelected())
 						{
-							SelectedSplineSegments.Add(*SegmentIt);
+							SelectedSplineSegments.Add(Segment);
+						}
+					}
+				}
+
+				for (ALandscapeProxy* LandscapeProxy : Info.Info->Proxies)
+				{
+					if (LandscapeProxy != NULL && LandscapeProxy->SplineComponent != NULL)
+					{
+						for (ULandscapeSplineControlPoint* ControlPoint : LandscapeProxy->SplineComponent->ControlPoints)
+						{
+							if (ControlPoint->IsSplineSelected())
+							{
+								SelectedSplineControlPoints.Add(ControlPoint);
+							}
+						}
+
+						for (ULandscapeSplineSegment* Segment : LandscapeProxy->SplineComponent->Segments)
+						{
+							if (Segment->IsSplineSelected())
+							{
+								SelectedSplineSegments.Add(Segment);
+							}
 						}
 					}
 				}
@@ -1183,19 +1187,35 @@ public:
 		}
 		else
 		{
-			for (auto It = EdMode->GetLandscapeList().CreateConstIterator(); It; ++It)
+			for (const FLandscapeListInfo& Info : EdMode->GetLandscapeList())
 			{
-				ALandscapeProxy* Landscape = (*It).Info->GetCurrentLevelLandscapeProxy(true);
+				ALandscape* Landscape = Info.Info->LandscapeActor.Get();
 				if (Landscape != NULL && Landscape->SplineComponent != NULL)
 				{
-					for (auto ControlPointIt = Landscape->SplineComponent->ControlPoints.CreateConstIterator(); ControlPointIt; ++ControlPointIt)
+					for (ULandscapeSplineControlPoint* ControlPoint : Landscape->SplineComponent->ControlPoints)
 					{
-						(*ControlPointIt)->SetSplineSelected(false);
+						ControlPoint->SetSplineSelected(false);
 					}
 
-					for (auto SegmentIt = Landscape->SplineComponent->Segments.CreateConstIterator(); SegmentIt; ++SegmentIt)
+					for (ULandscapeSplineSegment* Segment : Landscape->SplineComponent->Segments)
 					{
-						(*SegmentIt)->SetSplineSelected(false);
+						Segment->SetSplineSelected(false);
+					}
+				}
+
+				for (ALandscapeProxy* LandscapeProxy : Info.Info->Proxies)
+				{
+					if (LandscapeProxy != NULL && LandscapeProxy->SplineComponent != NULL)
+					{
+						for (ULandscapeSplineControlPoint* ControlPoint : LandscapeProxy->SplineComponent->ControlPoints)
+						{
+							ControlPoint->SetSplineSelected(false);
+						}
+
+						for (ULandscapeSplineSegment* Segment : LandscapeProxy->SplineComponent->Segments)
+						{
+							Segment->SetSplineSelected(false);
+						}
 					}
 				}
 			}
@@ -1212,9 +1232,9 @@ public:
 	{
 		GEditor->SelectNone(true, true, false);
 
-		for (auto It = EdMode->GetLandscapeList().CreateConstIterator(); It; ++It)
+		for (const FLandscapeListInfo& Info : EdMode->GetLandscapeList())
 		{
-			ALandscapeProxy* Landscape = (*It).Info->GetCurrentLevelLandscapeProxy(true);
+			ALandscapeProxy* Landscape = Info.Info->GetCurrentLevelLandscapeProxy(true);
 			if (Landscape != NULL && Landscape->SplineComponent != NULL)
 			{
 				Landscape->SplineComponent->ShowSplineEditorMesh(true);
@@ -1227,9 +1247,9 @@ public:
 		ClearSelection();
 		UpdatePropertiesWindows();
 
-		for (auto It = EdMode->GetLandscapeList().CreateConstIterator(); It; ++It)
+		for (const FLandscapeListInfo& Info : EdMode->GetLandscapeList())
 		{
-			ALandscapeProxy* Landscape = (*It).Info->GetCurrentLevelLandscapeProxy(true);
+			ALandscapeProxy* Landscape = Info.Info->GetCurrentLevelLandscapeProxy(true);
 			if (Landscape != NULL && Landscape->SplineComponent != NULL)
 			{
 				Landscape->SplineComponent->ShowSplineEditorMesh(false);
@@ -1239,61 +1259,51 @@ public:
 
 	virtual void Render(const FSceneView* View, FViewport* Viewport, FPrimitiveDrawInterface* PDI) OVERRIDE
 	{
-		if (SelectedSplineControlPoints.Num() > 0)
+		for (ULandscapeSplineControlPoint* ControlPoint : SelectedSplineControlPoints)
 		{
-			for (auto It = SelectedSplineControlPoints.CreateConstIterator(); It; ++It)
-			{
-				const ULandscapeSplineControlPoint* ControlPoint = (*It);
-				const ULandscapeSplinesComponent* SplinesComponent = ControlPoint->GetOuterULandscapeSplinesComponent();
+			const ULandscapeSplinesComponent* SplinesComponent = ControlPoint->GetOuterULandscapeSplinesComponent();
 
-				FVector HandlePos0 = SplinesComponent->ComponentToWorld.TransformPosition(ControlPoint->Location + ControlPoint->Rotation.Vector() * -20);
-				FVector HandlePos1 = SplinesComponent->ComponentToWorld.TransformPosition(ControlPoint->Location + ControlPoint->Rotation.Vector() * 20);
-				DrawDashedLine(PDI, HandlePos0, HandlePos1, FColor::White, 20, SDPG_Foreground);
+			FVector HandlePos0 = SplinesComponent->ComponentToWorld.TransformPosition(ControlPoint->Location + ControlPoint->Rotation.Vector() * -20);
+			FVector HandlePos1 = SplinesComponent->ComponentToWorld.TransformPosition(ControlPoint->Location + ControlPoint->Rotation.Vector() * 20);
+			DrawDashedLine(PDI, HandlePos0, HandlePos1, FColor::White, 20, SDPG_Foreground);
 
-				if (GEditorModeTools().GetWidgetMode() == FWidget::WM_Scale)
+			if (GEditorModeTools().GetWidgetMode() == FWidget::WM_Scale)
+			{
+				for (const FLandscapeSplineConnection& Connection : ControlPoint->ConnectedSegments)
 				{
-					for (int32 i = 0; i < ControlPoint->ConnectedSegments.Num(); i++)
-					{
-						const FLandscapeSplineConnection& Connection = ControlPoint->ConnectedSegments[i];
-
-						FVector StartLocation; FRotator StartRotation;
-						Connection.GetNearConnection().ControlPoint->GetConnectionLocationAndRotation(Connection.GetNearConnection().SocketName, StartLocation, StartRotation);
+					FVector StartLocation; FRotator StartRotation;
+					Connection.GetNearConnection().ControlPoint->GetConnectionLocationAndRotation(Connection.GetNearConnection().SocketName, StartLocation, StartRotation);
 
-						FVector StartPos = SplinesComponent->ComponentToWorld.TransformPosition(StartLocation);
-						FVector HandlePos = SplinesComponent->ComponentToWorld.TransformPosition(StartLocation + StartRotation.Vector() * Connection.GetNearConnection().TangentLen / 2);
-						PDI->DrawLine(StartPos, HandlePos, FColor::White, SDPG_Foreground);
+					FVector StartPos = SplinesComponent->ComponentToWorld.TransformPosition(StartLocation);
+					FVector HandlePos = SplinesComponent->ComponentToWorld.TransformPosition(StartLocation + StartRotation.Vector() * Connection.GetNearConnection().TangentLen / 2);
+					PDI->DrawLine(StartPos, HandlePos, FColor::White, SDPG_Foreground);
 
-						if (PDI->IsHitTesting()) PDI->SetHitProxy( new HLandscapeSplineProxy_Tangent(Connection.Segment, Connection.End) );
-						PDI->DrawPoint(HandlePos, FColor(255,255,255), 10.f, SDPG_Foreground);
-						if (PDI->IsHitTesting()) PDI->SetHitProxy( NULL );
-					}
+					if (PDI->IsHitTesting()) PDI->SetHitProxy( new HLandscapeSplineProxy_Tangent(Connection.Segment, Connection.End) );
+					PDI->DrawPoint(HandlePos, FColor(255,255,255), 10.f, SDPG_Foreground);
+					if (PDI->IsHitTesting()) PDI->SetHitProxy( NULL );
 				}
 			}
 		}
 
-		if (SelectedSplineSegments.Num() > 0)
+		if (GEditorModeTools().GetWidgetMode() == FWidget::WM_Scale)
 		{
-			if (GEditorModeTools().GetWidgetMode() == FWidget::WM_Scale)
+			for (ULandscapeSplineSegment* Segment : SelectedSplineSegments)
 			{
-				for (auto It = SelectedSplineSegments.CreateConstIterator(); It; ++It)
+				const ULandscapeSplinesComponent* SplinesComponent = Segment->GetOuterULandscapeSplinesComponent();
+				for (int32 End = 0; End <= 1; End++)
 				{
-					const ULandscapeSplineSegment* Segment = (*It);
-					const ULandscapeSplinesComponent* SplinesComponent = Segment->GetOuterULandscapeSplinesComponent();
-					for (int32 End = 0; End <= 1; End++)
-					{
-						const FLandscapeSplineSegmentConnection& Connection = Segment->Connections[End];
+					const FLandscapeSplineSegmentConnection& Connection = Segment->Connections[End];
 
-						FVector StartLocation; FRotator StartRotation;
-						Connection.ControlPoint->GetConnectionLocationAndRotation(Connection.SocketName, StartLocation, StartRotation);
+					FVector StartLocation; FRotator StartRotation;
+					Connection.ControlPoint->GetConnectionLocationAndRotation(Connection.SocketName, StartLocation, StartRotation);
 
-						FVector EndPos = SplinesComponent->ComponentToWorld.TransformPosition(StartLocation);
-						FVector EndHandlePos = SplinesComponent->ComponentToWorld.TransformPosition(StartLocation + StartRotation.Vector() * Connection.TangentLen / 2);
+					FVector EndPos = SplinesComponent->ComponentToWorld.TransformPosition(StartLocation);
+					FVector EndHandlePos = SplinesComponent->ComponentToWorld.TransformPosition(StartLocation + StartRotation.Vector() * Connection.TangentLen / 2);
 
-						PDI->DrawLine(EndPos, EndHandlePos, FColor::White, SDPG_Foreground);
-						if (PDI->IsHitTesting()) PDI->SetHitProxy( new HLandscapeSplineProxy_Tangent((*It), !!End) );
-						PDI->DrawPoint(EndHandlePos, FColor(255,255,255), 10.f, SDPG_Foreground);
-						if (PDI->IsHitTesting()) PDI->SetHitProxy( NULL );
-					}
+					PDI->DrawLine(EndPos, EndHandlePos, FColor::White, SDPG_Foreground);
+					if (PDI->IsHitTesting()) PDI->SetHitProxy( new HLandscapeSplineProxy_Tangent(Segment, !!End) );
+					PDI->DrawPoint(EndHandlePos, FColor(255,255,255), 10.f, SDPG_Foreground);
+					if (PDI->IsHitTesting()) PDI->SetHitProxy( NULL );
 				}
 			}
 		}
@@ -1472,13 +1482,13 @@ public:
 		{
 			FScopedTransaction Transaction( LOCTEXT("LandscapeSpline_Delete", "Delete Landscape Splines") );
 
-			for (auto It = SelectedSplineControlPoints.CreateConstIterator(); It; ++It)
+			for (ULandscapeSplineControlPoint* ControlPoint : SelectedSplineControlPoints)
 			{
-				DeleteControlPoint((*It), bDeleteLooseEnds);
+				DeleteControlPoint(ControlPoint, bDeleteLooseEnds);
 			}
-			for (auto It = SelectedSplineSegments.CreateConstIterator(); It; ++It)
+			for (ULandscapeSplineSegment* Segment : SelectedSplineSegments)
 			{
-				DeleteSegment((*It), bDeleteLooseEnds);
+				DeleteSegment(Segment, bDeleteLooseEnds);
 			}
 			ClearSelection();
 			UpdatePropertiesWindows();
@@ -1506,29 +1516,25 @@ public:
 			Objects.Reserve(SelectedSplineControlPoints.Num() + SelectedSplineSegments.Num() * 3); // worst case
 
 			// Control Points then segments
-			for (auto It = SelectedSplineControlPoints.CreateConstIterator(); It; ++It)
+			for (ULandscapeSplineControlPoint* ControlPoint : SelectedSplineControlPoints)
 			{
-				ULandscapeSplineControlPoint* ControlPoint = *It;
 				Objects.Add(ControlPoint);
 			}
-			for (auto It = SelectedSplineSegments.CreateConstIterator(); It; ++It)
+			for (ULandscapeSplineSegment* Segment : SelectedSplineSegments)
 			{
-				ULandscapeSplineSegment* Segment = *It;
 				Objects.AddUnique(Segment->Connections[0].ControlPoint);
 				Objects.AddUnique(Segment->Connections[1].ControlPoint);
 			}
-			for (auto It = SelectedSplineSegments.CreateConstIterator(); It; ++It)
+			for (ULandscapeSplineSegment* Segment : SelectedSplineSegments)
 			{
-				ULandscapeSplineSegment* Segment = *It;
 				Objects.Add(Segment);
 			}
 
 			// Perform export to text format
 			FStringOutputDevice Ar;
 			Ar.Logf(TEXT("Begin Splines\r\n"));
-			for (auto It = Objects.CreateConstIterator(); It; ++It)
+			for (UObject* Object : Objects)
 			{
-				UObject* Object = *It;
 				UExporter::ExportToOutputDevice(NULL, Object, NULL, Ar, TEXT("copy"), 3, PPF_None, false);
 			}
 			Ar.Logf(TEXT("End Splines\r\n"));
@@ -1575,9 +1581,9 @@ public:
 
 		if (bOffset)
 		{
-			for (auto It = OutObjects.CreateIterator(); It; ++It)
+			for (UObject* Object : OutObjects)
 			{
-				ULandscapeSplineControlPoint* ControlPoint = Cast<ULandscapeSplineControlPoint>(*It);
+				ULandscapeSplineControlPoint* ControlPoint = Cast<ULandscapeSplineControlPoint>(Object);
 				if (ControlPoint != NULL)
 				{
 					Landscape->SplineComponent->ControlPoints.Add(ControlPoint);
@@ -1678,9 +1684,8 @@ void FEdModeLandscape::SplineMoveToCurrentLevel()
 			TSet<ALandscapeProxy*> FromProxies;
 			ALandscapeProxy* Landscape = NULL;
 
-			for (auto It = SplineTool->SelectedSplineControlPoints.CreateIterator(); It; ++It)
+			for (ULandscapeSplineControlPoint* ControlPoint : SplineTool->SelectedSplineControlPoints)
 			{
-				ULandscapeSplineControlPoint* ControlPoint = *It;
 				ULandscapeSplinesComponent* Comp = ControlPoint->GetOuterULandscapeSplinesComponent();
 				ALandscapeProxy* FromProxy = Comp ? Cast<ALandscapeProxy>(Comp->GetOuter()) : NULL;
 				if (FromProxy)
@@ -1688,33 +1693,31 @@ void FEdModeLandscape::SplineMoveToCurrentLevel()
 					if (!Landscape)
 					{
 						ULandscapeInfo* LandscapeInfo = FromProxy->GetLandscapeInfo(false);
-						if (LandscapeInfo)
+						check(LandscapeInfo);
+						Landscape = LandscapeInfo->GetCurrentLevelLandscapeProxy(true);
+						if (!Landscape)
 						{
-							Landscape = LandscapeInfo->GetCurrentLevelLandscapeProxy(true);
-							if (!Landscape)
-							{
-								// No Landscape Proxy, don't support for creating only for Spline now
-								return;
-							}
+							// No Landscape Proxy, don't support for creating only for Spline now
+							return;
 						}
 					}
 
-					if (Landscape && Landscape != FromProxy)
+					if (Landscape != FromProxy)
 					{
+						Landscape->Modify();
 						if (Landscape->SplineComponent == NULL)
 						{
 							SplineTool->CreateSplineComponent(Landscape, FromProxy->SplineComponent->RelativeScale3D);
 						}
 
-						const FTransform SplineToLanscape =
-							FromProxy->LandscapeActorToWorld().GetRelativeTransform(FromProxy->SplineComponent->ComponentToWorld).InverseSafe() *
-							Landscape->LandscapeActorToWorld().GetRelativeTransform(Landscape->SplineComponent->ComponentToWorld);
+						const FTransform OldToNewTransform =
+							Landscape->SplineComponent->ComponentToWorld.GetRelativeTransform(FromProxy->SplineComponent->ComponentToWorld);
 
-						//SplineTool->AddControlPoint(Landscape, SplineToLanscape.TransformPosition(ControlPoint->Location));
 						Landscape->SplineComponent->Modify();
 						if (FromProxies.Find(FromProxy) == NULL)
 						{
 							FromProxies.Add(FromProxy);
+							FromProxy->Modify();
 							FromProxy->SplineComponent->Modify();
 							FromProxy->SplineComponent->MarkRenderStateDirty();
 						}
@@ -1724,7 +1727,7 @@ void FEdModeLandscape::SplineMoveToCurrentLevel()
 						ControlPoint->Rename(NULL, Landscape->SplineComponent);
 						Landscape->SplineComponent->ControlPoints.Add(ControlPoint);
 
-						ControlPoint->Location = SplineToLanscape.TransformPosition(ControlPoint->Location);
+						ControlPoint->Location = OldToNewTransform.TransformPosition(ControlPoint->Location);
 
 						if (ControlPoint->MeshComponent)
 						{
@@ -1736,15 +1739,14 @@ void FEdModeLandscape::SplineMoveToCurrentLevel()
 							Comp->Rename(NULL, Landscape);
 							Comp->AttachTo(Landscape->SplineComponent, NAME_None, EAttachLocation::KeepWorldPosition);
 						}
-						
-						ControlPoint->UpdateSplinePoints();
+
+						ControlPoint->UpdateSplinePoints(true, false);
 					}
 				}
 			}
 
-			for (auto It = SplineTool->SelectedSplineSegments.CreateIterator(); It; ++It)
+			for (ULandscapeSplineSegment* Segment : SplineTool->SelectedSplineSegments)
 			{
-				ULandscapeSplineSegment* Segment = *It;
 				ULandscapeSplinesComponent* Comp = Segment->GetOuterULandscapeSplinesComponent();
 				ALandscapeProxy* FromProxy = Comp ? Cast<ALandscapeProxy>(Comp->GetOuter()) : NULL;
 				if (FromProxy)
@@ -1752,14 +1754,18 @@ void FEdModeLandscape::SplineMoveToCurrentLevel()
 					if (!Landscape)
 					{
 						ULandscapeInfo* LandscapeInfo = FromProxy->GetLandscapeInfo(false);
-						if (LandscapeInfo)
+						check(LandscapeInfo);
+						Landscape = LandscapeInfo->GetCurrentLevelLandscapeProxy(true);
+						if (!Landscape)
 						{
-							Landscape = LandscapeInfo->GetCurrentLevelLandscapeProxy(true);
+							// No Landscape Proxy, don't support for creating only for Spline now
+							return;
 						}
 					}
 
-					if (Landscape && Landscape != FromProxy)
+					if (Landscape != FromProxy)
 					{
+						Landscape->Modify();
 						if (Landscape->SplineComponent == NULL)
 						{
 							SplineTool->CreateSplineComponent(Landscape, FromProxy->SplineComponent->RelativeScale3D);
@@ -1769,6 +1775,7 @@ void FEdModeLandscape::SplineMoveToCurrentLevel()
 						if (FromProxies.Find(FromProxy) == NULL)
 						{
 							FromProxies.Add(FromProxy);
+							FromProxy->Modify();
 							FromProxy->SplineComponent->Modify();
 							FromProxy->SplineComponent->MarkRenderStateDirty();
 						}
@@ -1778,9 +1785,8 @@ void FEdModeLandscape::SplineMoveToCurrentLevel()
 						Segment->Rename(NULL, Landscape->SplineComponent);
 						Landscape->SplineComponent->Segments.Add(Segment);
 
-						for (auto It = Segment->MeshComponents.CreateIterator(); It; ++It)
+						for (USplineMeshComponent* Comp : Segment->MeshComponents)
 						{
-							USplineMeshComponent* Comp = *It;
 							Comp->Modify();
 							Comp->UnregisterComponent();
 							Comp->DetachFromParent(true);
diff --git a/Engine/Source/Runtime/Engine/Classes/Landscape/LandscapeSplineControlPoint.h b/Engine/Source/Runtime/Engine/Classes/Landscape/LandscapeSplineControlPoint.h
--- a/Engine/Source/Runtime/Engine/Classes/Landscape/LandscapeSplineControlPoint.h
+++ b/Engine/Source/Runtime/Engine/Classes/Landscape/LandscapeSplineControlPoint.h
@@ -154,7 +154,7 @@ public:
 	virtual void AutoSetConnections(bool bIncludingValid);
 
 	/** Update spline points */
-	virtual void UpdateSplinePoints(bool bUpdateCollision = true);
+	virtual void UpdateSplinePoints(bool bUpdateCollision = true, bool bUpdateAttachedSegments = true);
 
 	/** Delete spline points */
 	virtual void DeleteSplinePoints();
diff --git a/Engine/Source/Runtime/Engine/Private/Landscape/LandscapeEditInterface.cpp b/Engine/Source/Runtime/Engine/Private/Landscape/LandscapeEditInterface.cpp
--- a/Engine/Source/Runtime/Engine/Private/Landscape/LandscapeEditInterface.cpp
+++ b/Engine/Source/Runtime/Engine/Private/Landscape/LandscapeEditInterface.cpp
@@ -53,10 +53,10 @@ void FLandscapeEditDataInterface::Flush()
 	// delete all the FLandscapeTextureDataInfo allocations
 	for( TMap<UTexture2D*, FLandscapeTextureDataInfo*>::TIterator It(TextureDataMap); It;  ++It )
 	{
-		delete It.Value();
+		delete It.Value();	// FLandscapeTextureDataInfo destructors will unlock any texture data
 	}
 
-	TextureDataMap.Empty();	// FLandscapeTextureDataInfo destructors will unlock any texture data
+	TextureDataMap.Empty();
 }
 
 #include "LevelUtils.h"
@@ -181,10 +181,10 @@ bool FLandscapeEditDataInterface::GetComponentsInRegion(int32 X1, int32 Y1, int3
 	for( int32 ComponentIndexY=ComponentIndexY1;ComponentIndexY<=ComponentIndexY2;ComponentIndexY++ )
 	{
 		for( int32 ComponentIndexX=ComponentIndexX1;ComponentIndexX<=ComponentIndexX2;ComponentIndexX++ )
-		{		
+		{
 			ULandscapeComponent* Component = LandscapeInfo->XYtoComponentMap.FindRef(FIntPoint(ComponentIndexX,ComponentIndexY));
 			if( Component )
-			{				
+			{
 				bNotLocked = bNotLocked && ( !FLevelUtils::IsLevelLocked(Component->GetLandscapeProxy()->GetLevel()) ) && FLevelUtils::IsLevelVisible(Component->GetLandscapeProxy()->GetLevel());
 				if (OutComponents)
 				{
@@ -248,7 +248,7 @@ void FLandscapeEditDataInterface::SetHeightData(int32 X1, int32 Y1, int32 X2, in
 	for( int32 ComponentIndexY=ComponentIndexY1;ComponentIndexY<=ComponentIndexY2;ComponentIndexY++ )
 	{
 		for( int32 ComponentIndexX=ComponentIndexX1;ComponentIndexX<=ComponentIndexX2;ComponentIndexX++ )
-		{	
+		{
 			FIntPoint ComponentKey(ComponentIndexX,ComponentIndexY);
 			ULandscapeComponent* Component = LandscapeInfo->XYtoComponentMap.FindRef(ComponentKey);
 
diff --git a/Engine/Source/Runtime/Engine/Private/Landscape/LandscapeSplines.cpp b/Engine/Source/Runtime/Engine/Private/Landscape/LandscapeSplines.cpp
--- a/Engine/Source/Runtime/Engine/Private/Landscape/LandscapeSplines.cpp
+++ b/Engine/Source/Runtime/Engine/Private/Landscape/LandscapeSplines.cpp
@@ -712,7 +712,7 @@ bool ULandscapeSplineControlPoint::OwnsComponent(const class UControlPointMeshCo
 }
 
 #if WITH_EDITOR
-void ULandscapeSplineControlPoint::UpdateSplinePoints(bool bUpdateCollision)
+void ULandscapeSplineControlPoint::UpdateSplinePoints(bool bUpdateCollision, bool bUpdateAttachedSegments)
 {
 	ULandscapeSplinesComponent* OuterSplines = GetOuterULandscapeSplinesComponent();
 
@@ -738,7 +738,8 @@ void ULandscapeSplineControlPoint::UpdateSplinePoints(bool bUpdateCollision)
 			if (MeshComponent == NULL)
 			{
 				Modify();
-				MeshComponent = ConstructObject<UControlPointMeshComponent>(UControlPointMeshComponent::StaticClass(), OuterSplines->GetOuter(), NAME_None, RF_Transactional|RF_TextExportTransient);
+				OuterSplines->GetOuter()->Modify();
+				MeshComponent = ConstructObject<UControlPointMeshComponent>(UControlPointMeshComponent::StaticClass(), OuterSplines->GetOuter(), NAME_None, RF_Transactional | RF_TextExportTransient);
 				MeshComponent->bSelected = bSelected;
 				MeshComponent->AttachTo(OuterSplines);
 			}
@@ -765,6 +766,7 @@ void ULandscapeSplineControlPoint::UpdateSplinePoints(bool bUpdateCollision)
 			}
 
 			Modify();
+			MeshComponent->GetOuter()->Modify();
 			MeshComponent->DestroyComponent();
 			MeshComponent = NULL;
 
@@ -865,9 +867,12 @@ void ULandscapeSplineControlPoint::UpdateSplinePoints(bool bUpdateCollision)
 		bNavDirty = false;
 	}
 
-	for (const FLandscapeSplineConnection& Connection : ConnectedSegments)
+	if (bUpdateAttachedSegments)
 	{
-		Connection.Segment->UpdateSplinePoints(bUpdateCollision);
+		for (const FLandscapeSplineConnection& Connection : ConnectedSegments)
+		{
+			Connection.Segment->UpdateSplinePoints(bUpdateCollision);
+		}
 	}
 }
 
@@ -890,8 +895,7 @@ void ULandscapeSplineControlPoint::DeleteSplinePoints()
 
 	if (MeshComponent != NULL)
 	{
-		MeshComponent->Modify();
-		MeshComponent->UnregisterComponent();
+		MeshComponent->GetOuter()->Modify();
 		MeshComponent->DestroyComponent();
 		MeshComponent = NULL;
 	}
@@ -1383,6 +1387,7 @@ void ULandscapeSplineSegment::UpdateSplinePoints(bool bUpdateCollision)
 			}
 			else
 			{
+				OuterSplines->GetOuter()->Modify();
 				MeshComponent = ConstructObject<USplineMeshComponent>(USplineMeshComponent::StaticClass(), OuterSplines->GetOuter(), NAME_None, RF_Transactional|RF_TextExportTransient);
 				MeshComponent->bSelected = bSelected;
 				MeshComponent->AttachTo(OuterSplines);
@@ -1427,7 +1432,7 @@ void ULandscapeSplineSegment::UpdateSplinePoints(bool bUpdateCollision)
 
 			for (int32 i = iMesh; i < MeshComponents.Num(); i++)
 			{
-				//MeshComponents[i]->Modify();
+				MeshComponents[i]->GetOuter()->Modify();
 				MeshComponents[i]->DestroyComponent();
 			}
 
@@ -1618,6 +1623,7 @@ void ULandscapeSplineSegment::UpdateSplinePoints(bool bUpdateCollision)
 
 			for (auto* MeshComponent : MeshComponents)
 			{
+				MeshComponent->GetOuter()->Modify();
 				MeshComponent->DestroyComponent();
 			}
 			MeshComponents.Empty();
@@ -1665,8 +1671,7 @@ void ULandscapeSplineSegment::DeleteSplinePoints()
 
 	for (auto* MeshComponent : MeshComponents)
 	{
-		MeshComponent->Modify();
-		MeshComponent->UnregisterComponent();
+		MeshComponent->GetOuter()->Modify();
 		MeshComponent->DestroyComponent();
 	}
 	MeshComponents.Empty();
