Fixed "Flatten to Target" for landscapes with a non-zero Z coordinate.
Fixed the flatten preview grid getting stuck on. Added a checkbox to toggle the display of the grid on/off.

diff --git a/Engine/Source/Editor/LandscapeEditor/Classes/LandscapeEditorObject.h b/Engine/Source/Editor/LandscapeEditor/Classes/LandscapeEditorObject.h
--- a/Engine/Source/Editor/LandscapeEditor/Classes/LandscapeEditorObject.h
+++ b/Engine/Source/Editor/LandscapeEditor/Classes/LandscapeEditorObject.h
@@ -2,7 +2,7 @@
 
 #pragma once
 
-#include "Landscape/LandscapeProxy.h"
+#include "Landscape/LandscapeProxy.h" // for ELandscapeLayerPaintingRestriction
 #include "LandscapeEditorObject.generated.h"
 
 UENUM()
@@ -238,6 +238,10 @@ class ULandscapeEditorObject : public UObject
 	UPROPERTY(Category="Tool Settings", EditAnywhere, meta=(ShowForTools="ToolSet_Flatten", ShowForTargetTypes="Heightmap", editcondition="bUseFlattenTarget", UIMin="-32768", UIMax="32768"))
 	float FlattenTarget;
 
+	// Whether to show the preview grid for the flatten target height
+	UPROPERTY(Category = "Tool Settings", EditAnywhere, AdvancedDisplay, meta = (DisplayName = "Show Preview Grid", ShowForTools = "ToolSet_Flatten", ShowForTargetTypes = "Heightmap", editcondition = "bUseFlattenTarget", HideEditConditionToggle, UIMin = "-32768", UIMax = "32768"))
+	bool bShowFlattenTargetPreview;
+
 	// Ramp Tool:
 
 	// Width of ramp
diff --git a/Engine/Source/Editor/LandscapeEditor/Private/LandscapeEdModePaintTools.cpp b/Engine/Source/Editor/LandscapeEditor/Private/LandscapeEdModePaintTools.cpp
--- a/Engine/Source/Editor/LandscapeEditor/Private/LandscapeEdModePaintTools.cpp
+++ b/Engine/Source/Editor/LandscapeEditor/Private/LandscapeEdModePaintTools.cpp
@@ -430,7 +430,7 @@ public:
 		if (InEdMode->UISettings->bUseFlattenTarget && bTargetIsHeightmap)
 		{
 			FTransform LocalToWorld = InTarget.LandscapeInfo->GetLandscapeProxy()->ActorToWorld();
-			float Height = InEdMode->UISettings->FlattenTarget / LocalToWorld.GetScale3D().Z - LocalToWorld.GetTranslation().Z;
+			float Height = (InEdMode->UISettings->FlattenTarget - LocalToWorld.GetTranslation().Z) / LocalToWorld.GetScale3D().Z;
 			FlattenHeight = LandscapeDataAccess::GetTexHeight(Height);
 			bInitializedFlattenHeight = true;
 		}
@@ -572,12 +572,17 @@ public:
 	virtual const TCHAR* GetToolName() OVERRIDE { return TEXT("Flatten"); }
 	virtual FText GetDisplayName() OVERRIDE { return NSLOCTEXT("UnrealEd", "LandscapeMode_Flatten", "Flatten"); };
 
+	virtual void Tick(FLevelEditorViewportClient* ViewportClient,float DeltaTime)
+	{
+		bool bShowGrid = EdMode->UISettings->bUseFlattenTarget && EdMode->CurrentToolTarget.TargetType == ELandscapeToolTargetType::Heightmap && EdMode->UISettings->bShowFlattenTargetPreview;
+		MeshComponent->SetVisibility(bShowGrid);
+	}
+
 	virtual bool MouseMove(FLevelEditorViewportClient* ViewportClient, FViewport* Viewport, int32 x, int32 y) OVERRIDE
 	{
 		bool bResult = FLandscapeToolPaintBase<ToolTarget, FLandscapeToolStrokeFlatten<ToolTarget>>::MouseMove(ViewportClient, Viewport, x, y);
 
-		if (ViewportClient->IsLevelEditorClient() && MeshComponent != NULL &&
-			this->EdMode->UISettings->bUseFlattenTarget && this->EdMode->CurrentToolTarget.TargetType == ELandscapeToolTargetType::Heightmap)
+		if (ViewportClient->IsLevelEditorClient() && MeshComponent != NULL)
 		{
 			this->EdMode->LandscapeMouseTrace((FLevelEditorViewportClient*)ViewportClient, x, y, LastMousePosition);
 
@@ -585,7 +590,7 @@ public:
 			FVector Origin;
 			Origin.X = FMath::RoundToFloat(LastMousePosition.X);
 			Origin.Y = FMath::RoundToFloat(LastMousePosition.Y);
-			Origin.Z = (FMath::RoundToFloat((this->EdMode->UISettings->FlattenTarget / LocalToWorld.GetScale3D().Z - LocalToWorld.GetTranslation().Z) * LANDSCAPE_INV_ZSCALE) - 0.1f) * LANDSCAPE_ZSCALE;
+			Origin.Z = (FMath::RoundToFloat((EdMode->UISettings->FlattenTarget - LocalToWorld.GetTranslation().Z) / LocalToWorld.GetScale3D().Z * LANDSCAPE_INV_ZSCALE) - 0.1f) * LANDSCAPE_ZSCALE;
 			MeshComponent->SetRelativeLocation(Origin, false);
 		}
 
@@ -601,6 +606,15 @@ public:
 		MeshComponent->StaticMesh = PlaneMesh;
 		MeshComponent->AttachTo(LandscapeProxy->GetRootComponent());
 		MeshComponent->RegisterComponent();
+
+		bool bShowGrid = EdMode->UISettings->bUseFlattenTarget && EdMode->CurrentToolTarget.TargetType == ELandscapeToolTargetType::Heightmap && EdMode->UISettings->bShowFlattenTargetPreview;
+		MeshComponent->SetVisibility(bShowGrid);
+
+		// Try to set a sane initial location for the preview grid
+		const FTransform LocalToWorld = this->EdMode->CurrentToolTarget.LandscapeInfo->GetLandscapeProxy()->GetRootComponent()->GetComponentToWorld();
+		FVector Origin = FVector::ZeroVector;
+		Origin.Z = (FMath::RoundToFloat((EdMode->UISettings->FlattenTarget - LocalToWorld.GetTranslation().Z) / LocalToWorld.GetScale3D().Z * LANDSCAPE_INV_ZSCALE) - 0.1f) * LANDSCAPE_ZSCALE;
+		MeshComponent->SetRelativeLocation(Origin, false);
 	}
 
 	virtual void ExitTool()
diff --git a/Engine/Source/Editor/LandscapeEditor/Private/LandscapeEditorObject.cpp b/Engine/Source/Editor/LandscapeEditor/Private/LandscapeEditorObject.cpp
--- a/Engine/Source/Editor/LandscapeEditor/Private/LandscapeEditorObject.cpp
+++ b/Engine/Source/Editor/LandscapeEditor/Private/LandscapeEditorObject.cpp
@@ -21,6 +21,7 @@ ULandscapeEditorObject::ULandscapeEditorObject(const class FPostConstructInitial
 	, bPickValuePerApply(false)
 	, bUseFlattenTarget(false)
 	, FlattenTarget(0)
+	, bShowFlattenTargetPreview(true)
 
 	, RampWidth(2000)
 	, RampSideFalloff(0.4f)
diff --git a/Engine/Source/Editor/PropertyEditor/Private/Presentation/PropertyEditor/PropertyEditor.cpp b/Engine/Source/Editor/PropertyEditor/Private/Presentation/PropertyEditor/PropertyEditor.cpp
--- a/Engine/Source/Editor/PropertyEditor/Private/Presentation/PropertyEditor/PropertyEditor.cpp
+++ b/Engine/Source/Editor/PropertyEditor/Private/Presentation/PropertyEditor/PropertyEditor.cpp
@@ -539,22 +539,26 @@ bool FPropertyEditor::GetEditConditionPropertyAddress( UBoolProperty*& Condition
 
 bool FPropertyEditor::SupportsEditConditionToggle( UProperty* InProperty ) 
 {
-	bool bIsConditionalPropertyVisible = false;
+	bool bShowEditConditionToggle = false;
+
+	if (!InProperty->HasMetaData(TEXT("HideEditConditionToggle")))
 	{
 		bool bNegateValue = false;
-		UBoolProperty* ConditionalProperty = GetEditConditionProperty( InProperty, bNegateValue );
-		if( ConditionalProperty != NULL )
+		UBoolProperty* ConditionalProperty = GetEditConditionProperty(InProperty, bNegateValue);
+		if (ConditionalProperty != NULL)
 		{
-			if( ConditionalProperty->HasAllPropertyFlags( CPF_Edit ) )
+			bShowEditConditionToggle = true;
+
+			if (ConditionalProperty->HasAllPropertyFlags(CPF_Edit))
 			{
 				// Conditionally-dependent property is already exposed for editing, so no need to draw another
 				// check box next to this property's label
-				bIsConditionalPropertyVisible = true;
+				bShowEditConditionToggle = false;
 			}
 		}
 	}
 
-	return !bIsConditionalPropertyVisible;
+	return bShowEditConditionToggle;
 }
 
 UBoolProperty* FPropertyEditor::GetEditConditionProperty( const UProperty* InProperty, bool& bNegate ) 
