fix mac build

diff --git a/Engine/Source/Editor/LandscapeEditor/Private/LandscapeEdModePaintTools.cpp b/Engine/Source/Editor/LandscapeEditor/Private/LandscapeEdModePaintTools.cpp
--- a/Engine/Source/Editor/LandscapeEditor/Private/LandscapeEdModePaintTools.cpp
+++ b/Engine/Source/Editor/LandscapeEditor/Private/LandscapeEdModePaintTools.cpp
@@ -574,7 +574,7 @@ public:
 
 	virtual void Tick(FLevelEditorViewportClient* ViewportClient,float DeltaTime)
 	{
-		bool bShowGrid = EdMode->UISettings->bUseFlattenTarget && EdMode->CurrentToolTarget.TargetType == ELandscapeToolTargetType::Heightmap && EdMode->UISettings->bShowFlattenTargetPreview;
+		bool bShowGrid = this->EdMode->UISettings->bUseFlattenTarget && this->EdMode->CurrentToolTarget.TargetType == ELandscapeToolTargetType::Heightmap && this->EdMode->UISettings->bShowFlattenTargetPreview;
 		MeshComponent->SetVisibility(bShowGrid);
 	}
 
@@ -590,7 +590,7 @@ public:
 			FVector Origin;
 			Origin.X = FMath::RoundToFloat(LastMousePosition.X);
 			Origin.Y = FMath::RoundToFloat(LastMousePosition.Y);
-			Origin.Z = (FMath::RoundToFloat((EdMode->UISettings->FlattenTarget - LocalToWorld.GetTranslation().Z) / LocalToWorld.GetScale3D().Z * LANDSCAPE_INV_ZSCALE) - 0.1f) * LANDSCAPE_ZSCALE;
+			Origin.Z = (FMath::RoundToFloat((this->EdMode->UISettings->FlattenTarget - LocalToWorld.GetTranslation().Z) / LocalToWorld.GetScale3D().Z * LANDSCAPE_INV_ZSCALE) - 0.1f) * LANDSCAPE_ZSCALE;
 			MeshComponent->SetRelativeLocation(Origin, false);
 		}
 
@@ -607,13 +607,13 @@ public:
 		MeshComponent->AttachTo(LandscapeProxy->GetRootComponent());
 		MeshComponent->RegisterComponent();
 
-		bool bShowGrid = EdMode->UISettings->bUseFlattenTarget && EdMode->CurrentToolTarget.TargetType == ELandscapeToolTargetType::Heightmap && EdMode->UISettings->bShowFlattenTargetPreview;
+		bool bShowGrid = this->EdMode->UISettings->bUseFlattenTarget && this->EdMode->CurrentToolTarget.TargetType == ELandscapeToolTargetType::Heightmap && this->EdMode->UISettings->bShowFlattenTargetPreview;
 		MeshComponent->SetVisibility(bShowGrid);
 
 		// Try to set a sane initial location for the preview grid
 		const FTransform LocalToWorld = this->EdMode->CurrentToolTarget.LandscapeInfo->GetLandscapeProxy()->GetRootComponent()->GetComponentToWorld();
 		FVector Origin = FVector::ZeroVector;
-		Origin.Z = (FMath::RoundToFloat((EdMode->UISettings->FlattenTarget - LocalToWorld.GetTranslation().Z) / LocalToWorld.GetScale3D().Z * LANDSCAPE_INV_ZSCALE) - 0.1f) * LANDSCAPE_ZSCALE;
+		Origin.Z = (FMath::RoundToFloat((this->EdMode->UISettings->FlattenTarget - LocalToWorld.GetTranslation().Z) / LocalToWorld.GetScale3D().Z * LANDSCAPE_INV_ZSCALE) - 0.1f) * LANDSCAPE_ZSCALE;
 		MeshComponent->SetRelativeLocation(Origin, false);
 	}
 
