Dynamically scale the sprites in the Landscape ramp tool

diff --git a/Engine/Source/Editor/LandscapeEditor/Private/LandscapeEdModeRampTool.cpp b/Engine/Source/Editor/LandscapeEditor/Private/LandscapeEdModeRampTool.cpp
--- a/Engine/Source/Editor/LandscapeEditor/Private/LandscapeEdModeRampTool.cpp
+++ b/Engine/Source/Editor/LandscapeEditor/Private/LandscapeEdModeRampTool.cpp
@@ -295,13 +295,23 @@ public:
 			for (int32 i = 0; i < NumPoints; i++)
 			{
 				WorldPoints[i] = LandscapeToWorld.TransformPosition(Points[i]);
+			}
+
+			float SpriteScale = EdMode->UISettings->RampWidth / 4;
+			if (NumPoints > 1)
+			{
+				SpriteScale = FMath::Min(SpriteScale, (WorldPoints[1] - WorldPoints[0]).Size() / 2);
+			}
+			SpriteScale = FMath::Clamp<float>(SpriteScale, 10, 500);
 
+			for (int32 i = 0; i < NumPoints; i++)
+			{
 				const FLinearColor SpriteColor = (i == SelectedPoint) ? SelectedSpriteColor : FLinearColor::White;
 
 				PDI->SetHitProxy(new HLandscapeRampToolPointHitProxy(i));
 				PDI->DrawSprite(WorldPoints[i],
-					SpriteTexture->Resource->GetSizeX() * 2,
-					SpriteTexture->Resource->GetSizeY() * 2,
+					SpriteScale,
+					SpriteScale,
 					SpriteTexture->Resource,
 					SpriteColor,
 					SDPG_Foreground,
