Fix Landscape flatten tool

diff --git a/Engine/Source/Editor/LandscapeEditor/Private/LandscapeEdModePaintTools.cpp b/Engine/Source/Editor/LandscapeEditor/Private/LandscapeEdModePaintTools.cpp
--- a/Engine/Source/Editor/LandscapeEditor/Private/LandscapeEdModePaintTools.cpp
+++ b/Engine/Source/Editor/LandscapeEditor/Private/LandscapeEdModePaintTools.cpp
@@ -574,6 +574,8 @@ public:
 
 	virtual void Tick(FLevelEditorViewportClient* ViewportClient,float DeltaTime)
 	{
+		FLandscapeToolPaintBase<ToolTarget, FLandscapeToolStrokeFlatten<ToolTarget>>::Tick(ViewportClient, DeltaTime);
+
 		bool bShowGrid = this->EdMode->UISettings->bUseFlattenTarget && this->EdMode->CurrentToolTarget.TargetType == ELandscapeToolTargetType::Heightmap && this->EdMode->UISettings->bShowFlattenTargetPreview;
 		MeshComponent->SetVisibility(bShowGrid);
 	}
@@ -619,6 +621,8 @@ public:
 
 	virtual void ExitTool()
 	{
+		FLandscapeToolPaintBase<ToolTarget, FLandscapeToolStrokeFlatten<ToolTarget>>::ExitTool();
+
 		MeshComponent->DetachFromParent();
 		MeshComponent->DestroyComponent();
 	}
