UE4 Refactoring. Changed OVERRIDE and FINAL macros to keywords override and final.

diff --git a/Engine/Source/Editor/Kismet/Private/FormatTextDetails.h b/Engine/Source/Editor/Kismet/Private/FormatTextDetails.h
--- a/Engine/Source/Editor/Kismet/Private/FormatTextDetails.h
+++ b/Engine/Source/Editor/Kismet/Private/FormatTextDetails.h
@@ -24,7 +24,7 @@ public:
 	~FFormatTextDetails();
 
 	/** IDetailCustomization interface */
-	virtual void CustomizeDetails( IDetailLayoutBuilder& DetailLayout ) OVERRIDE;
+	virtual void CustomizeDetails( IDetailLayoutBuilder& DetailLayout ) override;
 
 	/** Forces a refresh on the details customization */
 	void OnForceRefresh();
@@ -60,13 +60,13 @@ public:
 
 private:
 	/** IDetailCustomNodeBuilder Interface*/
-	virtual void SetOnRebuildChildren( FSimpleDelegate InOnRegenerateChildren ) OVERRIDE { OnRebuildChildren = InOnRegenerateChildren; }
-	virtual void GenerateHeaderRowContent( FDetailWidgetRow& NodeRow ) OVERRIDE {}
-	virtual void GenerateChildContent( IDetailChildrenBuilder& ChildrenBuilder ) OVERRIDE;
-	virtual void Tick( float DeltaTime ) OVERRIDE {}
-	virtual bool RequiresTick() const OVERRIDE { return false; }
-	virtual FName GetName() const OVERRIDE { return NAME_None; }
-	virtual bool InitiallyCollapsed() const OVERRIDE { return false; }
+	virtual void SetOnRebuildChildren( FSimpleDelegate InOnRegenerateChildren ) override { OnRebuildChildren = InOnRegenerateChildren; }
+	virtual void GenerateHeaderRowContent( FDetailWidgetRow& NodeRow ) override {}
+	virtual void GenerateChildContent( IDetailChildrenBuilder& ChildrenBuilder ) override;
+	virtual void Tick( float DeltaTime ) override {}
+	virtual bool RequiresTick() const override { return false; }
+	virtual FName GetName() const override { return NAME_None; }
+	virtual bool InitiallyCollapsed() const override { return false; }
 
 private:
 	FSimpleDelegate OnRebuildChildren;
@@ -91,13 +91,13 @@ public:
 
 private:
 	/** IDetailCustomNodeBuilder Interface*/
-	virtual void SetOnRebuildChildren( FSimpleDelegate InOnRegenerateChildren ) OVERRIDE {};
-	virtual void GenerateHeaderRowContent( FDetailWidgetRow& NodeRow ) OVERRIDE;
-	virtual void GenerateChildContent( IDetailChildrenBuilder& ChildrenBuilder ) OVERRIDE {};
-	virtual void Tick( float DeltaTime ) OVERRIDE {}
-	virtual bool RequiresTick() const OVERRIDE { return false; }
-	virtual FName GetName() const OVERRIDE { return NAME_None; }
-	virtual bool InitiallyCollapsed() const OVERRIDE { return false; }
+	virtual void SetOnRebuildChildren( FSimpleDelegate InOnRegenerateChildren ) override {};
+	virtual void GenerateHeaderRowContent( FDetailWidgetRow& NodeRow ) override;
+	virtual void GenerateChildContent( IDetailChildrenBuilder& ChildrenBuilder ) override {};
+	virtual void Tick( float DeltaTime ) override {}
+	virtual bool RequiresTick() const override { return false; }
+	virtual FName GetName() const override { return NAME_None; }
+	virtual bool InitiallyCollapsed() const override { return false; }
 
 private:
 	/** Retrieves the argument's name */
diff --git a/Engine/Source/Editor/Kismet/Private/InstancedStaticMeshSCSEditorCustomization.h b/Engine/Source/Editor/Kismet/Private/InstancedStaticMeshSCSEditorCustomization.h
--- a/Engine/Source/Editor/Kismet/Private/InstancedStaticMeshSCSEditorCustomization.h
+++ b/Engine/Source/Editor/Kismet/Private/InstancedStaticMeshSCSEditorCustomization.h
@@ -11,10 +11,10 @@ public:
 
 public:
 	/** ISCSEditorCustomization interface */
-	virtual bool HandleViewportClick(const TSharedRef<FEditorViewportClient>& InViewportClient, class FSceneView& InView, class HHitProxy* InHitProxy, FKey InKey, EInputEvent InEvent, uint32 InHitX, uint32 InHitY) OVERRIDE;
-	virtual bool HandleViewportDrag(class USceneComponent* InComponentScene, class USceneComponent* InComponentTemplate, const FVector& InDeltaTranslation, const FRotator& InDeltaRotation, const FVector& InDeltaScale, const FVector& InPivot) OVERRIDE;
-	virtual bool HandleGetWidgetLocation(class USceneComponent* InSceneComponent, FVector& OutWidgetLocation) OVERRIDE;
-	virtual bool HandleGetWidgetTransform(class USceneComponent* InSceneComponent, FMatrix& OutWidgetTransform) OVERRIDE;
+	virtual bool HandleViewportClick(const TSharedRef<FEditorViewportClient>& InViewportClient, class FSceneView& InView, class HHitProxy* InHitProxy, FKey InKey, EInputEvent InEvent, uint32 InHitX, uint32 InHitY) override;
+	virtual bool HandleViewportDrag(class USceneComponent* InComponentScene, class USceneComponent* InComponentTemplate, const FVector& InDeltaTranslation, const FRotator& InDeltaRotation, const FVector& InDeltaScale, const FVector& InPivot) override;
+	virtual bool HandleGetWidgetLocation(class USceneComponent* InSceneComponent, FVector& OutWidgetLocation) override;
+	virtual bool HandleGetWidgetTransform(class USceneComponent* InSceneComponent, FMatrix& OutWidgetTransform) override;
 
 protected:
 	/** Ensure that selection bits are in sync w/ the number of instances */
diff --git a/Engine/Source/Editor/Kismet/Private/SBlueprintActionMenu.cpp b/Engine/Source/Editor/Kismet/Private/SBlueprintActionMenu.cpp
--- a/Engine/Source/Editor/Kismet/Private/SBlueprintActionMenu.cpp
+++ b/Engine/Source/Editor/Kismet/Private/SBlueprintActionMenu.cpp
@@ -29,7 +29,7 @@ struct FBlueprintAction_PromoteVariable : public FEdGraphSchemaAction
 	}
 
 	// FEdGraphSchemaAction interface
-	virtual UEdGraphNode* PerformAction( class UEdGraph* ParentGraph, UEdGraphPin* FromPin, const FVector2D Location, bool bSelectNewNode = true) OVERRIDE
+	virtual UEdGraphNode* PerformAction( class UEdGraph* ParentGraph, UEdGraphPin* FromPin, const FVector2D Location, bool bSelectNewNode = true) override
 	{
 		if( ( ParentGraph != NULL ) && ( FromPin != NULL ) )
 		{
diff --git a/Engine/Source/Editor/Kismet/Private/SBlueprintEditorToolbar.cpp b/Engine/Source/Editor/Kismet/Private/SBlueprintEditorToolbar.cpp
--- a/Engine/Source/Editor/Kismet/Private/SBlueprintEditorToolbar.cpp
+++ b/Engine/Source/Editor/Kismet/Private/SBlueprintEditorToolbar.cpp
@@ -44,7 +44,7 @@ public:
 	}
 
 	// SWidget interface
-	virtual FVector2D ComputeDesiredSize() const OVERRIDE
+	virtual FVector2D ComputeDesiredSize() const override
 	{
 		const float Height = 20.0f;
 		const float Thickness = 16.0f;
diff --git a/Engine/Source/Editor/Kismet/Private/SBlueprintFavoritesPalette.h b/Engine/Source/Editor/Kismet/Private/SBlueprintFavoritesPalette.h
--- a/Engine/Source/Editor/Kismet/Private/SBlueprintFavoritesPalette.h
+++ b/Engine/Source/Editor/Kismet/Private/SBlueprintFavoritesPalette.h
@@ -30,12 +30,12 @@ public:
 
 private:
 	// SGraphPalette Interface
-	virtual void CollectAllActions(FGraphActionListBuilderBase& OutAllActions) OVERRIDE;
+	virtual void CollectAllActions(FGraphActionListBuilderBase& OutAllActions) override;
 	// End SGraphPalette Interface
 
 	// SBlueprintSubPalette Interface
-	virtual void BindCommands(TSharedPtr<FUICommandList> CommandListIn) const OVERRIDE;
-	virtual void GenerateContextMenuEntries(FMenuBuilder& MenuBuilder) const OVERRIDE;
+	virtual void BindCommands(TSharedPtr<FUICommandList> CommandListIn) const override;
+	virtual void GenerateContextMenuEntries(FMenuBuilder& MenuBuilder) const override;
 	// End SBlueprintSubPalette Interface
 
 	/** Flags weather we should add the "frequently used" list to the user's favorites */
diff --git a/Engine/Source/Editor/Kismet/Private/SBlueprintLibraryPalette.cpp b/Engine/Source/Editor/Kismet/Private/SBlueprintLibraryPalette.cpp
--- a/Engine/Source/Editor/Kismet/Private/SBlueprintLibraryPalette.cpp
+++ b/Engine/Source/Editor/Kismet/Private/SBlueprintLibraryPalette.cpp
@@ -146,7 +146,7 @@ public:
 	TSharedPtr<FUICommandInfo> RemoveSubFavorites;
 
 	/** Registers context menu commands for the blueprint library palette. */
-	virtual void RegisterCommands() OVERRIDE
+	virtual void RegisterCommands() override
 	{
 		UI_COMMAND(AddSingleFavorite,    "Add to Favorites",               "Adds this item to your favorites list.",                      EUserInterfaceActionType::Button, FInputGesture());
 		UI_COMMAND(AddSubFavorites,      "Add Category to Favorites",      "Adds all the nodes in this category to your favorites.",      EUserInterfaceActionType::Button, FInputGesture());
@@ -162,13 +162,13 @@ public:
 class FPaletteClassFilter : public IClassViewerFilter
 {
 public:
-	virtual bool IsClassAllowed(const FClassViewerInitializationOptions& InInitOptions, const UClass* InClass, TSharedRef< FClassViewerFilterFuncs > InFilterFuncs ) OVERRIDE
+	virtual bool IsClassAllowed(const FClassViewerInitializationOptions& InInitOptions, const UClass* InClass, TSharedRef< FClassViewerFilterFuncs > InFilterFuncs ) override
 	{
 		const UEdGraphSchema_K2* K2Schema = GetDefault<UEdGraphSchema_K2>();
 		return K2Schema->ClassHasBlueprintAccessibleMembers(InClass);
 	}
 
-	virtual bool IsUnloadedClassAllowed(const FClassViewerInitializationOptions& InInitOptions, const TSharedRef< const IUnloadedBlueprintData > InUnloadedClassData, TSharedRef< FClassViewerFilterFuncs > InFilterFuncs) OVERRIDE
+	virtual bool IsUnloadedClassAllowed(const FClassViewerInitializationOptions& InInitOptions, const TSharedRef< const IUnloadedBlueprintData > InUnloadedClassData, TSharedRef< FClassViewerFilterFuncs > InFilterFuncs) override
 	{
 		// @TODO: One day would be nice to see functions on unloaded classes...
 		return false;
diff --git a/Engine/Source/Editor/Kismet/Private/SBlueprintLibraryPalette.h b/Engine/Source/Editor/Kismet/Private/SBlueprintLibraryPalette.h
--- a/Engine/Source/Editor/Kismet/Private/SBlueprintLibraryPalette.h
+++ b/Engine/Source/Editor/Kismet/Private/SBlueprintLibraryPalette.h
@@ -29,13 +29,13 @@ public:
 
 private:
 	// SGraphPalette Interface
-	virtual void CollectAllActions(FGraphActionListBuilderBase& OutAllActions) OVERRIDE;
+	virtual void CollectAllActions(FGraphActionListBuilderBase& OutAllActions) override;
 	// End SGraphPalette Interface
 
 	// SBlueprintSubPalette Interface
-	virtual TSharedRef<SVerticalBox> ConstructHeadingWidget(FSlateBrush const* const Icon, FString const& TitleText, FString const& ToolTip) OVERRIDE;
-	virtual void BindCommands(TSharedPtr<FUICommandList> CommandListIn) const OVERRIDE;
-	virtual void GenerateContextMenuEntries(FMenuBuilder& MenuBuilder) const OVERRIDE;
+	virtual TSharedRef<SVerticalBox> ConstructHeadingWidget(FSlateBrush const* const Icon, FString const& TitleText, FString const& ToolTip) override;
+	virtual void BindCommands(TSharedPtr<FUICommandList> CommandListIn) const override;
+	virtual void GenerateContextMenuEntries(FMenuBuilder& MenuBuilder) const override;
 	// End SBlueprintSubPalette Interface
 
 	/**
diff --git a/Engine/Source/Editor/Kismet/Private/SBlueprintPalette.h b/Engine/Source/Editor/Kismet/Private/SBlueprintPalette.h
--- a/Engine/Source/Editor/Kismet/Private/SBlueprintPalette.h
+++ b/Engine/Source/Editor/Kismet/Private/SBlueprintPalette.h
@@ -34,10 +34,10 @@ public:
 
 private:
 	// SGraphPaletteItem Interface
-	virtual TSharedRef<SWidget> CreateTextSlotWidget( const FSlateFontInfo& NameFont,  FCreateWidgetForActionData* const InCreateData, bool bIsReadOnly ) OVERRIDE;
-	virtual FText GetDisplayText() const OVERRIDE;
-	virtual bool OnNameTextVerifyChanged(const FText& InNewText, FText& OutErrorMessage) OVERRIDE;
-	virtual void OnNameTextCommitted(const FText& NewText, ETextCommit::Type InTextCommit) OVERRIDE;
+	virtual TSharedRef<SWidget> CreateTextSlotWidget( const FSlateFontInfo& NameFont,  FCreateWidgetForActionData* const InCreateData, bool bIsReadOnly ) override;
+	virtual FText GetDisplayText() const override;
+	virtual bool OnNameTextVerifyChanged(const FText& InNewText, FText& OutErrorMessage) override;
+	virtual void OnNameTextCommitted(const FText& NewText, ETextCommit::Type InTextCommit) override;
 	// End of SGraphPaletteItem Interface
 
 	/**
diff --git a/Engine/Source/Editor/Kismet/Private/SBlueprintSubPalette.h b/Engine/Source/Editor/Kismet/Private/SBlueprintSubPalette.h
--- a/Engine/Source/Editor/Kismet/Private/SBlueprintSubPalette.h
+++ b/Engine/Source/Editor/Kismet/Private/SBlueprintSubPalette.h
@@ -60,9 +60,9 @@ public:
 
 protected:
 	// SGraphPalette Interface
-	virtual void RefreshActionsList(bool bPreserveExpansion) OVERRIDE;
-	virtual TSharedRef<SWidget> OnCreateWidgetForAction(FCreateWidgetForActionData* const InCreateData) OVERRIDE;
-	virtual FReply OnActionDragged(const TArray< TSharedPtr<FEdGraphSchemaAction> >& InActions, const FPointerEvent& MouseEvent) OVERRIDE;
+	virtual void RefreshActionsList(bool bPreserveExpansion) override;
+	virtual TSharedRef<SWidget> OnCreateWidgetForAction(FCreateWidgetForActionData* const InCreateData) override;
+	virtual FReply OnActionDragged(const TArray< TSharedPtr<FEdGraphSchemaAction> >& InActions, const FPointerEvent& MouseEvent) override;
 	// End of SGraphPalette Interface
 
 	/**
diff --git a/Engine/Source/Editor/Kismet/Private/SCSEditorViewportClient.h b/Engine/Source/Editor/Kismet/Private/SCSEditorViewportClient.h
--- a/Engine/Source/Editor/Kismet/Private/SCSEditorViewportClient.h
+++ b/Engine/Source/Editor/Kismet/Private/SCSEditorViewportClient.h
@@ -24,18 +24,18 @@ public:
 	virtual ~FSCSEditorViewportClient();
 
 	// FEditorViewportClient interface
-	virtual void Tick(float DeltaSeconds) OVERRIDE;
-	virtual FSceneView* CalcSceneView(FSceneViewFamily* ViewFamily) OVERRIDE;
-	virtual void Draw(const FSceneView* View,FPrimitiveDrawInterface* PDI) OVERRIDE;
-	virtual void DrawCanvas( FViewport& InViewport, FSceneView& View, FCanvas& Canvas ) OVERRIDE;
-	virtual bool InputKey(FViewport* Viewport, int32 ControllerId, FKey Key, EInputEvent Event, float AmountDepressed = 1.f, bool bGamepad=false) OVERRIDE;
-	virtual void ProcessClick(class FSceneView& View, class HHitProxy* HitProxy, FKey Key, EInputEvent Event, uint32 HitX, uint32 HitY) OVERRIDE;
-	virtual bool InputWidgetDelta( FViewport* Viewport, EAxisList::Type CurrentAxis, FVector& Drag, FRotator& Rot, FVector& Scale ) OVERRIDE;
-	virtual void TrackingStarted( const struct FInputEventState& InInputState, bool bIsDragging, bool bNudge ) OVERRIDE;
-	virtual void TrackingStopped() OVERRIDE;
-	virtual FWidget::EWidgetMode GetWidgetMode() const OVERRIDE;
-	virtual void SetWidgetMode( FWidget::EWidgetMode NewMode ) OVERRIDE;
-	virtual void SetWidgetCoordSystemSpace( ECoordSystem NewCoordSystem ) OVERRIDE;
+	virtual void Tick(float DeltaSeconds) override;
+	virtual FSceneView* CalcSceneView(FSceneViewFamily* ViewFamily) override;
+	virtual void Draw(const FSceneView* View,FPrimitiveDrawInterface* PDI) override;
+	virtual void DrawCanvas( FViewport& InViewport, FSceneView& View, FCanvas& Canvas ) override;
+	virtual bool InputKey(FViewport* Viewport, int32 ControllerId, FKey Key, EInputEvent Event, float AmountDepressed = 1.f, bool bGamepad=false) override;
+	virtual void ProcessClick(class FSceneView& View, class HHitProxy* HitProxy, FKey Key, EInputEvent Event, uint32 HitX, uint32 HitY) override;
+	virtual bool InputWidgetDelta( FViewport* Viewport, EAxisList::Type CurrentAxis, FVector& Drag, FRotator& Rot, FVector& Scale ) override;
+	virtual void TrackingStarted( const struct FInputEventState& InInputState, bool bIsDragging, bool bNudge ) override;
+	virtual void TrackingStopped() override;
+	virtual FWidget::EWidgetMode GetWidgetMode() const override;
+	virtual void SetWidgetMode( FWidget::EWidgetMode NewMode ) override;
+	virtual void SetWidgetCoordSystemSpace( ECoordSystem NewCoordSystem ) override;
 	virtual FVector GetWidgetLocation() const;
 	virtual FMatrix GetWidgetCoordSystem() const;
 	virtual ECoordSystem GetWidgetCoordSystemSpace() const { return WidgetCoordSystem; }
diff --git a/Engine/Source/Editor/Kismet/Private/SMyBlueprint.cpp b/Engine/Source/Editor/Kismet/Private/SMyBlueprint.cpp
--- a/Engine/Source/Editor/Kismet/Private/SMyBlueprint.cpp
+++ b/Engine/Source/Editor/Kismet/Private/SMyBlueprint.cpp
@@ -60,7 +60,7 @@ class FMyBlueprintCategoryDragDropAction : public FGraphEditorDragDropAction
 public:
 	DRAG_DROP_OPERATOR_TYPE(FMyBlueprintCategoryDragDropAction, FGraphEditorDragDropAction)
 
-	virtual void HoverTargetChanged() OVERRIDE
+	virtual void HoverTargetChanged() override
 	{
 		const FSlateBrush* StatusSymbol = FEditorStyle::GetBrush(TEXT("NoBrush")); 
 		FText Message = FText::FromString(DraggedCategory);
@@ -93,7 +93,7 @@ public:
 		SetSimpleFeedbackMessage(StatusSymbol, FLinearColor::White, Message);
 	}
 	
-	virtual FReply DroppedOnCategory(FString OnCategory) OVERRIDE
+	virtual FReply DroppedOnCategory(FString OnCategory) override
 	{
 		// Get the Blueprint via BlueprintEditor
 		TSharedPtr<SMyBlueprint> MyBlueprint = MyBlueprintPtr.Pin();
diff --git a/Engine/Source/Editor/Kismet/Private/SSCSEditor.cpp b/Engine/Source/Editor/Kismet/Private/SSCSEditor.cpp
--- a/Engine/Source/Editor/Kismet/Private/SSCSEditor.cpp
+++ b/Engine/Source/Editor/Kismet/Private/SSCSEditor.cpp
@@ -139,7 +139,7 @@ public:
 	EDropActionType PendingDropAction;
 	
 	/** The widget decorator to use */
-	virtual TSharedPtr<SWidget> GetDefaultDecorator() const OVERRIDE
+	virtual TSharedPtr<SWidget> GetDefaultDecorator() const override
 	{
 		return SNew(SBorder)
 			.BorderImage(FEditorStyle::GetBrush("Graph.ConnectorFeedback.Border"))
@@ -759,13 +759,13 @@ protected:
 
 	// FCustomizableTextObjectFactory implementation
 
-	virtual bool CanCreateClass(UClass* ObjectClass) const OVERRIDE
+	virtual bool CanCreateClass(UClass* ObjectClass) const override
 	{
 		// Only allow actor component types to be created
 		return ObjectClass->IsChildOf(UActorComponent::StaticClass());
 	}
 
-	virtual void ProcessConstructedObject(UObject* NewObject) OVERRIDE
+	virtual void ProcessConstructedObject(UObject* NewObject) override
 	{
 		check(NewObject);
 
diff --git a/Engine/Source/Editor/Kismet/Private/SSCSEditor.h b/Engine/Source/Editor/Kismet/Private/SSCSEditor.h
--- a/Engine/Source/Editor/Kismet/Private/SSCSEditor.h
+++ b/Engine/Source/Editor/Kismet/Private/SSCSEditor.h
@@ -212,15 +212,15 @@ public:
 	virtual ~SSCS_RowWidget();
 
 	// SMultiColumnTableRow<T> interface
-	virtual TSharedRef<SWidget> GenerateWidgetForColumn( const FName& ColumnName ) OVERRIDE;
+	virtual TSharedRef<SWidget> GenerateWidgetForColumn( const FName& ColumnName ) override;
 	// End of SMultiColumnTableRow<T>
 
 	// SWidget interface
-	virtual FReply OnMouseButtonDown(const FGeometry& MyGeometry, const FPointerEvent& MouseEvent) OVERRIDE;
-	virtual FReply OnDragDetected( const FGeometry& MyGeometry, const FPointerEvent& MouseEvent ) OVERRIDE;
-	virtual void OnDragEnter( const FGeometry& MyGeometry, const FDragDropEvent& DragDropEvent ) OVERRIDE;
-	virtual void OnDragLeave( const FDragDropEvent& DragDropEvent ) OVERRIDE;
-	virtual FReply OnDrop( const FGeometry& MyGeometry, const FDragDropEvent& DragDropEvent ) OVERRIDE;
+	virtual FReply OnMouseButtonDown(const FGeometry& MyGeometry, const FPointerEvent& MouseEvent) override;
+	virtual FReply OnDragDetected( const FGeometry& MyGeometry, const FPointerEvent& MouseEvent ) override;
+	virtual void OnDragEnter( const FGeometry& MyGeometry, const FDragDropEvent& DragDropEvent ) override;
+	virtual void OnDragLeave( const FDragDropEvent& DragDropEvent ) override;
+	virtual FReply OnDrop( const FGeometry& MyGeometry, const FDragDropEvent& DragDropEvent ) override;
 	// End of SWidget interface
 
 	/** Get the blueprint we are editing */
diff --git a/Engine/Source/Editor/Kismet/Private/SSCSEditorViewport.h b/Engine/Source/Editor/Kismet/Private/SSCSEditorViewport.h
--- a/Engine/Source/Editor/Kismet/Private/SSCSEditorViewport.h
+++ b/Engine/Source/Editor/Kismet/Private/SSCSEditorViewport.h
@@ -56,7 +56,7 @@ public:
 	 * @param  InCurrentTime  Current absolute real time
 	 * @param  InDeltaTime  Real time passed since last tick
 	 */
-	virtual void Tick( const FGeometry& AllottedGeometry, const double InCurrentTime, const float InDeltaTime ) OVERRIDE;
+	virtual void Tick( const FGeometry& AllottedGeometry, const double InCurrentTime, const float InDeltaTime ) override;
 	
 	// End of SWidget interface
 
@@ -98,8 +98,8 @@ protected:
 	EVisibility GetWidgetVisibility() const;
 
 	/** SEditorViewport interface */
-	virtual TSharedRef<class FEditorViewportClient> MakeEditorViewportClient() OVERRIDE;
-	virtual TSharedPtr<class SWidget> MakeViewportToolbar() OVERRIDE;
+	virtual TSharedRef<class FEditorViewportClient> MakeEditorViewportClient() override;
+	virtual TSharedPtr<class SWidget> MakeViewportToolbar() override;
 	virtual void BindCommands();
 private:
 	/** Pointer back to editor tool (owner) */
diff --git a/Engine/Source/Editor/Kismet/Private/UserDefinedEnumEditor.h b/Engine/Source/Editor/Kismet/Private/UserDefinedEnumEditor.h
--- a/Engine/Source/Editor/Kismet/Private/UserDefinedEnumEditor.h
+++ b/Engine/Source/Editor/Kismet/Private/UserDefinedEnumEditor.h
@@ -32,14 +32,14 @@ public:
 	virtual ~FUserDefinedEnumEditor();
 
 	/** IToolkit interface */
-	virtual FName GetToolkitFName() const OVERRIDE;
-	virtual FText GetBaseToolkitName() const OVERRIDE;
-	virtual FText GetToolkitName() const OVERRIDE;
-	virtual FString GetWorldCentricTabPrefix() const OVERRIDE;
-	virtual FLinearColor GetWorldCentricTabColorScale() const OVERRIDE;
+	virtual FName GetToolkitFName() const override;
+	virtual FText GetBaseToolkitName() const override;
+	virtual FText GetToolkitName() const override;
+	virtual FString GetWorldCentricTabPrefix() const override;
+	virtual FLinearColor GetWorldCentricTabColorScale() const override;
 
-	virtual void RegisterTabSpawners(const TSharedRef<class FTabManager>& TabManager) OVERRIDE;
-	virtual void UnregisterTabSpawners(const TSharedRef<class FTabManager>& TabManager) OVERRIDE;
+	virtual void RegisterTabSpawners(const TSharedRef<class FTabManager>& TabManager) override;
+	virtual void UnregisterTabSpawners(const TSharedRef<class FTabManager>& TabManager) override;
 
 protected:
 	TSharedRef<SDockTab> SpawnEnumeratorsTab(const FSpawnTabArgs& Args);
@@ -63,13 +63,13 @@ public:
 	~FEnumDetails();
 
 	/** IDetailCustomization interface */
-	virtual void CustomizeDetails( IDetailLayoutBuilder& DetailLayout ) OVERRIDE;
+	virtual void CustomizeDetails( IDetailLayoutBuilder& DetailLayout ) override;
 
 	/** Forces a refresh on the details customization */
 	void OnForceRefresh();
 
 	/** FEnumEditorUtils::INotifyOnEnumChanged */
-	virtual void OnChanged(const class UUserDefinedEnum* Enum) OVERRIDE;
+	virtual void OnChanged(const class UUserDefinedEnum* Enum) override;
 
 private:
 	/** Handles new enum element request */
@@ -98,13 +98,13 @@ public:
 
 private:
 	/** IDetailCustomNodeBuilder Interface*/
-	virtual void SetOnRebuildChildren( FSimpleDelegate InOnRegenerateChildren ) OVERRIDE { OnRebuildChildren = InOnRegenerateChildren; }
-	virtual void GenerateHeaderRowContent( FDetailWidgetRow& NodeRow ) OVERRIDE {}
-	virtual void GenerateChildContent( IDetailChildrenBuilder& ChildrenBuilder ) OVERRIDE;
-	virtual void Tick( float DeltaTime ) OVERRIDE {}
-	virtual bool RequiresTick() const OVERRIDE { return false; }
-	virtual FName GetName() const OVERRIDE { return NAME_None; }
-	virtual bool InitiallyCollapsed() const OVERRIDE { return false; }
+	virtual void SetOnRebuildChildren( FSimpleDelegate InOnRegenerateChildren ) override { OnRebuildChildren = InOnRegenerateChildren; }
+	virtual void GenerateHeaderRowContent( FDetailWidgetRow& NodeRow ) override {}
+	virtual void GenerateChildContent( IDetailChildrenBuilder& ChildrenBuilder ) override;
+	virtual void Tick( float DeltaTime ) override {}
+	virtual bool RequiresTick() const override { return false; }
+	virtual FName GetName() const override { return NAME_None; }
+	virtual bool InitiallyCollapsed() const override { return false; }
 
 private:
 	FSimpleDelegate OnRebuildChildren;
@@ -128,13 +128,13 @@ public:
 
 private:
 	/** IDetailCustomNodeBuilder Interface*/
-	virtual void SetOnRebuildChildren( FSimpleDelegate InOnRegenerateChildren ) OVERRIDE {};
-	virtual void GenerateHeaderRowContent( FDetailWidgetRow& NodeRow ) OVERRIDE;
-	virtual void GenerateChildContent( IDetailChildrenBuilder& ChildrenBuilder ) OVERRIDE {};
-	virtual void Tick( float DeltaTime ) OVERRIDE {}
-	virtual bool RequiresTick() const OVERRIDE { return false; }
-	virtual FName GetName() const OVERRIDE { return NAME_None; }
-	virtual bool InitiallyCollapsed() const OVERRIDE { return false; }
+	virtual void SetOnRebuildChildren( FSimpleDelegate InOnRegenerateChildren ) override {};
+	virtual void GenerateHeaderRowContent( FDetailWidgetRow& NodeRow ) override;
+	virtual void GenerateChildContent( IDetailChildrenBuilder& ChildrenBuilder ) override {};
+	virtual void Tick( float DeltaTime ) override {}
+	virtual bool RequiresTick() const override { return false; }
+	virtual FName GetName() const override { return NAME_None; }
+	virtual bool InitiallyCollapsed() const override { return false; }
 
 private:
 	/** Retrieves the enumerator's name */
diff --git a/Engine/Source/Editor/Kismet/Private/UserDefinedStructureEditor.cpp b/Engine/Source/Editor/Kismet/Private/UserDefinedStructureEditor.cpp
--- a/Engine/Source/Editor/Kismet/Private/UserDefinedStructureEditor.cpp
+++ b/Engine/Source/Editor/Kismet/Private/UserDefinedStructureEditor.cpp
@@ -223,17 +223,17 @@ public:
 	}
 
 	/** IDetailCustomNodeBuilder Interface*/
-	virtual void SetOnRebuildChildren( FSimpleDelegate InOnRegenerateChildren ) OVERRIDE 
+	virtual void SetOnRebuildChildren( FSimpleDelegate InOnRegenerateChildren ) override 
 	{
 		OnRegenerateChildren = InOnRegenerateChildren;
 	}
-	virtual void GenerateChildContent( IDetailChildrenBuilder& ChildrenBuilder ) OVERRIDE;
+	virtual void GenerateChildContent( IDetailChildrenBuilder& ChildrenBuilder ) override;
 
-	virtual void GenerateHeaderRowContent( FDetailWidgetRow& NodeRow ) OVERRIDE {}
+	virtual void GenerateHeaderRowContent( FDetailWidgetRow& NodeRow ) override {}
 
-	virtual void Tick( float DeltaTime ) OVERRIDE {}
-	virtual bool RequiresTick() const OVERRIDE { return false; }
-	virtual FName GetName() const OVERRIDE 
+	virtual void Tick( float DeltaTime ) override {}
+	virtual bool RequiresTick() const override { return false; }
+	virtual FName GetName() const override 
 	{ 
 		auto StructureDetailsSP = StructureDetails.Pin();
 		if(StructureDetailsSP.IsValid())
@@ -245,7 +245,7 @@ public:
 		}
 		return NAME_None; 
 	}
-	virtual bool InitiallyCollapsed() const OVERRIDE { return false; }
+	virtual bool InitiallyCollapsed() const override { return false; }
 
 private:
 	TWeakPtr<class FUserDefinedStructureDetails> StructureDetails;
@@ -447,7 +447,7 @@ public:
 	}
 
 	/** IDetailCustomNodeBuilder Interface*/
-	virtual void SetOnRebuildChildren( FSimpleDelegate InOnRegenerateChildren ) OVERRIDE 
+	virtual void SetOnRebuildChildren( FSimpleDelegate InOnRegenerateChildren ) override 
 	{
 		OnRegenerateChildren = InOnRegenerateChildren;
 	}
@@ -506,7 +506,7 @@ public:
 		}
 	}
 
-	virtual void GenerateHeaderRowContent( FDetailWidgetRow& NodeRow ) OVERRIDE 
+	virtual void GenerateHeaderRowContent( FDetailWidgetRow& NodeRow ) override 
 	{
 		auto K2Schema = GetDefault<UEdGraphSchema_K2>();
 
@@ -574,7 +574,7 @@ public:
 		}
 	}
 
-	virtual void GenerateChildContent( IDetailChildrenBuilder& ChildrenBuilder ) OVERRIDE 
+	virtual void GenerateChildContent( IDetailChildrenBuilder& ChildrenBuilder ) override 
 	{
 		ChildrenBuilder.AddChildContent( *LOCTEXT( "UDSVarDefaultValue", "Default Value" ).ToString() )
 		.NameContent()
@@ -640,10 +640,10 @@ public:
 					this, &FUserDefinedStructureFieldLayout::Is3dWidgetOptionVisible)));
 	}
 
-	virtual void Tick( float DeltaTime ) OVERRIDE {}
-	virtual bool RequiresTick() const OVERRIDE { return false; }
-	virtual FName GetName() const OVERRIDE { return FName(*FieldGuid.ToString()); }
-	virtual bool InitiallyCollapsed() const OVERRIDE { return true; }
+	virtual void Tick( float DeltaTime ) override {}
+	virtual bool RequiresTick() const override { return false; }
+	virtual FName GetName() const override { return FName(*FieldGuid.ToString()); }
+	virtual bool InitiallyCollapsed() const override { return true; }
 
 private:
 	TWeakPtr<class FUserDefinedStructureDetails> StructureDetails;
diff --git a/Engine/Source/Editor/Kismet/Private/UserDefinedStructureEditor.h b/Engine/Source/Editor/Kismet/Private/UserDefinedStructureEditor.h
--- a/Engine/Source/Editor/Kismet/Private/UserDefinedStructureEditor.h
+++ b/Engine/Source/Editor/Kismet/Private/UserDefinedStructureEditor.h
@@ -31,14 +31,14 @@ public:
 	virtual ~FUserDefinedStructureEditor();
 
 	/** IToolkit interface */
-	virtual FName GetToolkitFName() const OVERRIDE;
-	virtual FText GetBaseToolkitName() const OVERRIDE;
-	virtual FText GetToolkitName() const OVERRIDE;
-	virtual FString GetWorldCentricTabPrefix() const OVERRIDE;
-	virtual FLinearColor GetWorldCentricTabColorScale() const OVERRIDE;
+	virtual FName GetToolkitFName() const override;
+	virtual FText GetBaseToolkitName() const override;
+	virtual FText GetToolkitName() const override;
+	virtual FString GetWorldCentricTabPrefix() const override;
+	virtual FLinearColor GetWorldCentricTabColorScale() const override;
 
-	virtual void RegisterTabSpawners(const TSharedRef<class FTabManager>& TabManager) OVERRIDE;
-	virtual void UnregisterTabSpawners(const TSharedRef<class FTabManager>& TabManager) OVERRIDE;
+	virtual void RegisterTabSpawners(const TSharedRef<class FTabManager>& TabManager) override;
+	virtual void UnregisterTabSpawners(const TSharedRef<class FTabManager>& TabManager) override;
 
 protected:
 	TSharedRef<SDockTab> SpawnStructureTab(const FSpawnTabArgs& Args);
@@ -56,14 +56,14 @@ public:
 	~FUserDefinedStructureDetails();
 
 	/** IDetailCustomization interface */
-	virtual void CustomizeDetails( class IDetailLayoutBuilder& DetailLayout ) OVERRIDE;
+	virtual void CustomizeDetails( class IDetailLayoutBuilder& DetailLayout ) override;
 
 	UUserDefinedStruct* GetUserDefinedStruct();
 
 	struct FStructVariableDescription* FindStructureFieldByGuid(FGuid Guid);
 
 	/** FStructureEditorUtils::INotifyOnStructChanged */
-	virtual void OnChanged(const class UUserDefinedStruct* Struct) OVERRIDE;
+	virtual void OnChanged(const class UUserDefinedStruct* Struct) override;
 
 private:
 	TWeakObjectPtr<UUserDefinedStruct> UserDefinedStruct;
diff --git a/Engine/Source/Editor/Kismet/Public/BlueprintEditor.h b/Engine/Source/Editor/Kismet/Public/BlueprintEditor.h
--- a/Engine/Source/Editor/Kismet/Public/BlueprintEditor.h
+++ b/Engine/Source/Editor/Kismet/Public/BlueprintEditor.h
@@ -44,9 +44,9 @@ struct KISMET_API FSelectionDetailsSummoner : public FWorkflowTabFactory
 public:
 	FSelectionDetailsSummoner(TSharedPtr<class FAssetEditorToolkit> InHostingApp);
 
-	virtual TSharedRef<SWidget> CreateTabBody(const FWorkflowTabSpawnInfo& Info) const OVERRIDE;
+	virtual TSharedRef<SWidget> CreateTabBody(const FWorkflowTabSpawnInfo& Info) const override;
 
-	virtual FText GetTabToolTipText(const FWorkflowTabSpawnInfo& Info) const OVERRIDE
+	virtual FText GetTabToolTipText(const FWorkflowTabSpawnInfo& Info) const override
 	{
 		return LOCTEXT("SelectionDetailsTooltip", "The Details tab allows you see and edit properties of whatever is selected.");
 	}
@@ -107,7 +107,7 @@ class KISMET_API FBlueprintEditor : public IBlueprintEditor, public FGCObject, p
 
 public:
 	// IToolkit interface
-	virtual void RegisterTabSpawners(const TSharedRef<class FTabManager>& TabManager) OVERRIDE;
+	virtual void RegisterTabSpawners(const TSharedRef<class FTabManager>& TabManager) override;
 	// End of IToolkit interface
 
 public:
@@ -130,37 +130,37 @@ public:
 
 public:
 	// FAssetEditorToolkit interface
-	virtual bool OnRequestClose() OVERRIDE;
-	virtual void ToolkitBroughtToFront() OVERRIDE;
+	virtual bool OnRequestClose() override;
+	virtual void ToolkitBroughtToFront() override;
 	// End of FAssetEditorToolkit 
 
 	// IToolkit interface
-	virtual FName GetToolkitFName() const OVERRIDE;
-	virtual FText GetBaseToolkitName() const OVERRIDE;
-	virtual FText GetToolkitName() const OVERRIDE;
-	virtual FString GetWorldCentricTabPrefix() const OVERRIDE;
-	virtual FLinearColor GetWorldCentricTabColorScale() const OVERRIDE;
-	virtual bool IsBlueprintEditor() const OVERRIDE;
+	virtual FName GetToolkitFName() const override;
+	virtual FText GetBaseToolkitName() const override;
+	virtual FText GetToolkitName() const override;
+	virtual FString GetWorldCentricTabPrefix() const override;
+	virtual FLinearColor GetWorldCentricTabColorScale() const override;
+	virtual bool IsBlueprintEditor() const override;
 	// End of IToolkit interface
 
 	// FGCObject interface
-	virtual void AddReferencedObjects( FReferenceCollector& Collector ) OVERRIDE;
+	virtual void AddReferencedObjects( FReferenceCollector& Collector ) override;
 	// End of FGCObject interface
 
 	// IBlueprintEditor interface
-	virtual void RefreshEditors() OVERRIDE;
-	virtual void JumpToHyperlink(const UObject* ObjectReference, bool bRequestRename = false) OVERRIDE;
-	virtual void SummonSearchUI(bool bSetFindWithinBlueprint, FString NewSearchTerms = FString(), bool bSelectFirstResult = false) OVERRIDE;
-	virtual TArray<TSharedPtr<class FSCSEditorTreeNode> > GetSelectedSCSEditorTreeNodes() const OVERRIDE;
-	virtual TSharedPtr<class FSCSEditorTreeNode> FindAndSelectSCSEditorTreeNode(const UActorComponent* InComponent, bool IsCntrlDown) OVERRIDE;
-	virtual int32 GetNumberOfSelectedNodes() const OVERRIDE;
-	virtual void AnalyticsTrackNodeEvent( UBlueprint* Blueprint, UEdGraphNode *GraphNode, bool bNodeDelete = false ) const OVERRIDE;
+	virtual void RefreshEditors() override;
+	virtual void JumpToHyperlink(const UObject* ObjectReference, bool bRequestRename = false) override;
+	virtual void SummonSearchUI(bool bSetFindWithinBlueprint, FString NewSearchTerms = FString(), bool bSelectFirstResult = false) override;
+	virtual TArray<TSharedPtr<class FSCSEditorTreeNode> > GetSelectedSCSEditorTreeNodes() const override;
+	virtual TSharedPtr<class FSCSEditorTreeNode> FindAndSelectSCSEditorTreeNode(const UActorComponent* InComponent, bool IsCntrlDown) override;
+	virtual int32 GetNumberOfSelectedNodes() const override;
+	virtual void AnalyticsTrackNodeEvent( UBlueprint* Blueprint, UEdGraphNode *GraphNode, bool bNodeDelete = false ) const override;
 	// End of IBlueprintEditor interface
 
 	// FTickableEditorObject interface
-	virtual void Tick(float DeltaTime) OVERRIDE;
-	virtual bool IsTickable() const OVERRIDE { return true; }
-	virtual TStatId GetStatId() const OVERRIDE;
+	virtual void Tick(float DeltaTime) override;
+	virtual bool IsTickable() const override { return true; }
+	virtual TStatId GetStatId() const override;
 	// End of FTickableEditorObject interface
 
 public:
@@ -198,7 +198,7 @@ public:
 	TSharedPtr<class FBlueprintEditorToolbar> GetToolbarBuilder() {return Toolbar;}
 
 	/** @return the documentation location for this editor */
-	virtual FString GetDocumentationLink() const OVERRIDE;
+	virtual FString GetDocumentationLink() const override;
 
 	/**	Returns whether the edited blueprint has components */
 	bool CanAccessComponentsMode() const;
@@ -491,7 +491,7 @@ protected:
 	EVisibility ParentClassButtonsVisibility() const;
 
 	/** Recreates the overlay on the menu bar */
-	virtual void PostRegenerateMenusAndToolbars() OVERRIDE;
+	virtual void PostRegenerateMenusAndToolbars() override;
 
 	/** Returns the name of the Blueprint's parent class */
 	FText GetParentClassNameText() const;
@@ -636,10 +636,10 @@ protected:
 	bool CanCopyNodes() const;
 	
 	/** Paste on graph at specific location */
-	virtual void PasteNodesHere(class UEdGraph* DestinationGraph, const FVector2D& GraphLocation) OVERRIDE;
+	virtual void PasteNodesHere(class UEdGraph* DestinationGraph, const FVector2D& GraphLocation) override;
 	
 	void PasteNodes();
-	virtual bool CanPasteNodes() const OVERRIDE;
+	virtual bool CanPasteNodes() const override;
 
 	void DuplicateNodes();
 	bool CanDuplicateNodes() const;
@@ -762,8 +762,8 @@ protected:
 	void FindInBlueprints_OnClicked();
 
 	// FNotifyHook interface
-	virtual void NotifyPreChange( UProperty* PropertyAboutToChange ) OVERRIDE;
-	virtual void NotifyPostChange( const FPropertyChangedEvent& PropertyChangedEvent, UProperty* PropertyThatChanged) OVERRIDE;
+	virtual void NotifyPreChange( UProperty* PropertyAboutToChange ) override;
+	virtual void NotifyPostChange( const FPropertyChangedEvent& PropertyChangedEvent, UProperty* PropertyThatChanged) override;
 	// End of FNotifyHook interface
 
 	/** Callback when properties have finished being handled */
@@ -786,8 +786,8 @@ protected:
 	void LoadLibrariesFromAssetRegistry();
 
 	// Begin FEditorUndoClient Interface
-	virtual void	PostUndo(bool bSuccess) OVERRIDE;
-	virtual void	PostRedo(bool bSuccess) OVERRIDE;
+	virtual void	PostUndo(bool bSuccess) override;
+	virtual void	PostRedo(bool bSuccess) override;
 	// End of FEditorUndoClient
 
 	/** Get graph appearance */
diff --git a/Engine/Source/Editor/Kismet/Public/BlueprintEditorModes.h b/Engine/Source/Editor/Kismet/Public/BlueprintEditorModes.h
--- a/Engine/Source/Editor/Kismet/Public/BlueprintEditorModes.h
+++ b/Engine/Source/Editor/Kismet/Public/BlueprintEditorModes.h
@@ -43,9 +43,9 @@ class KISMET_API FBlueprintEditorApplicationMode : public FApplicationMode
 public:
 	FBlueprintEditorApplicationMode(TSharedPtr<class FBlueprintEditor> InBlueprintEditor, FName InModeName, const bool bRegisterViewport = true, const bool bRegisterDefaultsTab = true);
 
-	virtual void RegisterTabFactories(TSharedPtr<FTabManager> InTabManager) OVERRIDE;
-	virtual void PreDeactivateMode() OVERRIDE;
-	virtual void PostActivateMode() OVERRIDE;
+	virtual void RegisterTabFactories(TSharedPtr<FTabManager> InTabManager) override;
+	virtual void PreDeactivateMode() override;
+	virtual void PostActivateMode() override;
 public:
 
 protected:
@@ -67,8 +67,8 @@ class KISMET_API FBlueprintDefaultsApplicationMode : public FApplicationMode
 public:
 	FBlueprintDefaultsApplicationMode(TSharedPtr<class FBlueprintEditor> InBlueprintEditor);
 
-	virtual void RegisterTabFactories(TSharedPtr<FTabManager> InTabManager) OVERRIDE;
-	virtual void PostActivateMode() OVERRIDE;
+	virtual void RegisterTabFactories(TSharedPtr<FTabManager> InTabManager) override;
+	virtual void PostActivateMode() override;
 public:
 
 protected:
@@ -84,9 +84,9 @@ class KISMET_API FBlueprintComponentsApplicationMode : public FApplicationMode
 public:
 	FBlueprintComponentsApplicationMode(TSharedPtr<class FBlueprintEditor> InBlueprintEditor);
 
-	virtual void RegisterTabFactories(TSharedPtr<FTabManager> InTabManager) OVERRIDE;
-	virtual void PreDeactivateMode() OVERRIDE;
-	virtual void PostActivateMode() OVERRIDE;
+	virtual void RegisterTabFactories(TSharedPtr<FTabManager> InTabManager) override;
+	virtual void PreDeactivateMode() override;
+	virtual void PostActivateMode() override;
 
 	TSharedPtr<FBlueprintEditor> GetBlueprintEditor() const { return MyBlueprintEditor.Pin(); }
 	AActor* GetPreviewActor() const;
@@ -102,7 +102,7 @@ class KISMET_API FBlueprintInterfaceApplicationMode : public FApplicationMode
 public:
 	FBlueprintInterfaceApplicationMode(TSharedPtr<class FBlueprintEditor> InBlueprintEditor);
 
-	virtual void RegisterTabFactories(TSharedPtr<FTabManager> InTabManager) OVERRIDE;
+	virtual void RegisterTabFactories(TSharedPtr<FTabManager> InTabManager) override;
 public:
 
 protected:
@@ -117,7 +117,7 @@ class KISMET_API FBlueprintMacroApplicationMode : public FApplicationMode
 public:
 	FBlueprintMacroApplicationMode(TSharedPtr<class FBlueprintEditor> InBlueprintEditor);
 
-	virtual void RegisterTabFactories(TSharedPtr<FTabManager> InTabManager) OVERRIDE;
+	virtual void RegisterTabFactories(TSharedPtr<FTabManager> InTabManager) override;
 public:
 
 protected:
diff --git a/Engine/Source/Editor/Kismet/Public/WorkflowOrientedApp/WorkflowCentricApplication.h b/Engine/Source/Editor/Kismet/Public/WorkflowOrientedApp/WorkflowCentricApplication.h
--- a/Engine/Source/Editor/Kismet/Public/WorkflowOrientedApp/WorkflowCentricApplication.h
+++ b/Engine/Source/Editor/Kismet/Public/WorkflowOrientedApp/WorkflowCentricApplication.h
@@ -20,12 +20,12 @@ class KISMET_API FWorkflowCentricApplication : public FAssetEditorToolkit
 {
 public:
 	// IToolkit interface
-	virtual void RegisterTabSpawners(const TSharedRef<class FTabManager>& TabManager) OVERRIDE;
-	virtual void UnregisterTabSpawners(const TSharedRef<class FTabManager>& TabManager) OVERRIDE;
+	virtual void RegisterTabSpawners(const TSharedRef<class FTabManager>& TabManager) override;
+	virtual void UnregisterTabSpawners(const TSharedRef<class FTabManager>& TabManager) override;
 	// End of IToolkit interface
 	
 	// FAssetEditorToolkit interface
-	virtual bool OnRequestClose() OVERRIDE;
+	virtual bool OnRequestClose() override;
 	// End of FAssetEditorToolkit interface
 
 	// Returns the current mode of this application
diff --git a/Engine/Source/Editor/Kismet/Public/WorkflowOrientedApp/WorkflowUObjectDocuments.h b/Engine/Source/Editor/Kismet/Public/WorkflowOrientedApp/WorkflowUObjectDocuments.h
--- a/Engine/Source/Editor/Kismet/Public/WorkflowOrientedApp/WorkflowUObjectDocuments.h
+++ b/Engine/Source/Editor/Kismet/Public/WorkflowOrientedApp/WorkflowUObjectDocuments.h
@@ -26,7 +26,7 @@ public:
 	}
 
 	// Determine if another payload is the same as this one
-	virtual bool IsEqual(const TSharedRef<FTabPayload>& OtherPayload) const OVERRIDE
+	virtual bool IsEqual(const TSharedRef<FTabPayload>& OtherPayload) const override
 	{
 		if (OtherPayload->PayloadType == PayloadType)
 		{
@@ -36,7 +36,7 @@ public:
 		return false;
 	}
 
-	virtual bool IsValid() const OVERRIDE
+	virtual bool IsValid() const override
 	{
 		return DocumentID.IsValid();
 	}
@@ -71,19 +71,19 @@ public:
 public:
 
 	// FWorkflowTabFactory public interface
-	virtual TSharedRef<SWidget> CreateTabBody(const FWorkflowTabSpawnInfo& Info) const OVERRIDE
+	virtual TSharedRef<SWidget> CreateTabBody(const FWorkflowTabSpawnInfo& Info) const override
 	{
 		BaseClass* TypedDocumentID = FTabPayload_UObject::CastChecked<BaseClass>(Info.Payload);
 		return CreateTabBodyForObject(Info, TypedDocumentID);
 	}
 
-	virtual const FSlateBrush* GetTabIcon(const FWorkflowTabSpawnInfo& Info) const OVERRIDE
+	virtual const FSlateBrush* GetTabIcon(const FWorkflowTabSpawnInfo& Info) const override
 	{
 		BaseClass* TypedDocumentID = FTabPayload_UObject::CastChecked<BaseClass>(Info.Payload);
 		return GetTabIconForObject(Info, TypedDocumentID);
 	}
 
-	virtual bool IsPayloadSupported(TSharedRef<FTabPayload> Payload) const OVERRIDE
+	virtual bool IsPayloadSupported(TSharedRef<FTabPayload> Payload) const override
 	{
 		if (Payload->PayloadType == NAME_Object)
 		{
@@ -93,7 +93,7 @@ public:
 		return false;
 	}
 
-	virtual bool IsPayloadValid(TSharedRef<FTabPayload> Payload) const OVERRIDE
+	virtual bool IsPayloadValid(TSharedRef<FTabPayload> Payload) const override
 	{
 		if (Payload->PayloadType == NAME_Object)
 		{
@@ -111,7 +111,7 @@ protected:
 	}
 
 	// Creates the label for the tab
-	virtual TAttribute<FText> ConstructTabName(const FWorkflowTabSpawnInfo& Info) const OVERRIDE
+	virtual TAttribute<FText> ConstructTabName(const FWorkflowTabSpawnInfo& Info) const override
 	{
 		check(Info.Payload.IsValid());
 		BaseClass* TypedDocumentID = FTabPayload_UObject::CastChecked<BaseClass>(Info.Payload);
diff --git a/Engine/Source/Editor/KismetCompiler/Private/AnimBlueprintCompiler.h b/Engine/Source/Editor/KismetCompiler/Private/AnimBlueprintCompiler.h
--- a/Engine/Source/Editor/KismetCompiler/Private/AnimBlueprintCompiler.h
+++ b/Engine/Source/Editor/KismetCompiler/Private/AnimBlueprintCompiler.h
@@ -26,15 +26,15 @@ public:
 	virtual ~FAnimBlueprintCompiler();
 protected:
 	// Implementation of FKismetCompilerContext interface
-	virtual UEdGraphSchema_K2* CreateSchema() OVERRIDE;
-	virtual void MergeUbergraphPagesIn(UEdGraph* Ubergraph) OVERRIDE;
-	virtual void ProcessOneFunctionGraph(UEdGraph* SourceGraph) OVERRIDE;
-	virtual void CreateFunctionList() OVERRIDE;
-	virtual void SpawnNewClass(const FString& NewClassName) OVERRIDE;
-	virtual void CopyTermDefaultsToDefaultObject(UObject* DefaultObject) OVERRIDE;
-	virtual void PostCompileDiagnostics() OVERRIDE;
-	virtual void EnsureProperGeneratedClass(UClass*& TargetClass) OVERRIDE;
-	virtual void CleanAndSanitizeClass(UBlueprintGeneratedClass* ClassToClean, UObject*& OldCDO) OVERRIDE;
+	virtual UEdGraphSchema_K2* CreateSchema() override;
+	virtual void MergeUbergraphPagesIn(UEdGraph* Ubergraph) override;
+	virtual void ProcessOneFunctionGraph(UEdGraph* SourceGraph) override;
+	virtual void CreateFunctionList() override;
+	virtual void SpawnNewClass(const FString& NewClassName) override;
+	virtual void CopyTermDefaultsToDefaultObject(UObject* DefaultObject) override;
+	virtual void PostCompileDiagnostics() override;
+	virtual void EnsureProperGeneratedClass(UClass*& TargetClass) override;
+	virtual void CleanAndSanitizeClass(UBlueprintGeneratedClass* ClassToClean, UObject*& OldCDO) override;
 	// End of FKismetCompilerContext interface
 
 protected:
diff --git a/Engine/Source/Editor/KismetCompiler/Public/KismetCompiler.h b/Engine/Source/Editor/KismetCompiler/Public/KismetCompiler.h
--- a/Engine/Source/Editor/KismetCompiler/Public/KismetCompiler.h
+++ b/Engine/Source/Editor/KismetCompiler/Public/KismetCompiler.h
@@ -164,13 +164,13 @@ protected:
 
 protected:
 	// FGraphCompilerContext interface
-	virtual void ValidateLink(const UEdGraphPin* PinA, const UEdGraphPin* PinB) const OVERRIDE;
-	virtual void ValidatePin(const UEdGraphPin* Pin) const OVERRIDE;
-	virtual void ValidateNode(const UEdGraphNode* Node) const OVERRIDE;
-	virtual bool CanIgnoreNode(const UEdGraphNode* Node) const OVERRIDE;
-	virtual bool ShouldForceKeepNode(const UEdGraphNode* Node) const OVERRIDE;
-	virtual void PruneIsolatedNodes(const TArray<UEdGraphNode*>& RootSet, TArray<UEdGraphNode*>& GraphNodes) OVERRIDE;
-	virtual bool PinIsImportantForDependancies(const UEdGraphPin* Pin) const OVERRIDE
+	virtual void ValidateLink(const UEdGraphPin* PinA, const UEdGraphPin* PinB) const override;
+	virtual void ValidatePin(const UEdGraphPin* Pin) const override;
+	virtual void ValidateNode(const UEdGraphNode* Node) const override;
+	virtual bool CanIgnoreNode(const UEdGraphNode* Node) const override;
+	virtual bool ShouldForceKeepNode(const UEdGraphNode* Node) const override;
+	virtual void PruneIsolatedNodes(const TArray<UEdGraphNode*>& RootSet, TArray<UEdGraphNode*>& GraphNodes) override;
+	virtual bool PinIsImportantForDependancies(const UEdGraphPin* Pin) const override
 	{
 		// The execution wires do not form data dependencies, they are only important for final scheduling and that is handled thru gotos
 		return Pin->PinType.PinCategory != Schema->PC_Exec;
diff --git a/Engine/Source/Editor/KismetCompiler/Public/KismetCompilerModule.h b/Engine/Source/Editor/KismetCompiler/Public/KismetCompilerModule.h
--- a/Engine/Source/Editor/KismetCompiler/Public/KismetCompilerModule.h
+++ b/Engine/Source/Editor/KismetCompiler/Public/KismetCompilerModule.h
@@ -64,11 +64,11 @@ class FKismet2CompilerModule : public IKismetCompilerInterface
 {
 public:
 	// Implementation of the IKismetCompilerInterface
-	virtual void CompileBlueprint(class UBlueprint* Blueprint, const FKismetCompilerOptions& CompileOptions, FCompilerResultsLog& Results, class FBlueprintCompileReinstancer* ParentReinstancer = NULL, TArray<UObject*>* ObjLoaded = NULL) OVERRIDE;
-	virtual void CompileStructure(class UUserDefinedStruct* Struct, FCompilerResultsLog& Results) OVERRIDE;
-	virtual void RecoverCorruptedBlueprint(class UBlueprint* Blueprint) OVERRIDE;
-	virtual void RemoveBlueprintGeneratedClasses(class UBlueprint* Blueprint) OVERRIDE;
-	virtual TArray<FBlueprintCompileDelegate>& GetCompilers() OVERRIDE { return Compilers; }
+	virtual void CompileBlueprint(class UBlueprint* Blueprint, const FKismetCompilerOptions& CompileOptions, FCompilerResultsLog& Results, class FBlueprintCompileReinstancer* ParentReinstancer = NULL, TArray<UObject*>* ObjLoaded = NULL) override;
+	virtual void CompileStructure(class UUserDefinedStruct* Struct, FCompilerResultsLog& Results) override;
+	virtual void RecoverCorruptedBlueprint(class UBlueprint* Blueprint) override;
+	virtual void RemoveBlueprintGeneratedClasses(class UBlueprint* Blueprint) override;
+	virtual TArray<FBlueprintCompileDelegate>& GetCompilers() override { return Compilers; }
 	// End implementation
 private:
 	void CompileBlueprintInner(class UBlueprint* Blueprint, bool bPrintResultSuccess, const FKismetCompilerOptions& CompileOptions, FCompilerResultsLog& Results, TArray<UObject*>* ObjLoaded);
diff --git a/Engine/Source/Editor/KismetWidgets/Public/SScrubWidget.h b/Engine/Source/Editor/KismetWidgets/Public/SScrubWidget.h
--- a/Engine/Source/Editor/KismetWidgets/Public/SScrubWidget.h
+++ b/Engine/Source/Editor/KismetWidgets/Public/SScrubWidget.h
@@ -66,13 +66,13 @@ public:
 	void Construct( const FArguments& InArgs );
 
 	// SWidget interface
-	virtual FVector2D ComputeDesiredSize() const OVERRIDE;
-	virtual int32 OnPaint( const FGeometry& AllottedGeometry, const FSlateRect& MyClippingRect, FSlateWindowElementList& OutDrawElements, int32 LayerId, const FWidgetStyle& InWidgetStyle, bool bParentEnabled ) const OVERRIDE;
-	virtual FReply OnMouseButtonDown( const FGeometry& MyGeometry, const FPointerEvent& MouseEvent ) OVERRIDE;
-	virtual FReply OnMouseButtonUp( const FGeometry& MyGeometry, const FPointerEvent& MouseEvent ) OVERRIDE;
-	virtual FReply OnMouseMove( const FGeometry& MyGeometry, const FPointerEvent& MouseEvent ) OVERRIDE;
-	virtual FReply OnMouseWheel( const FGeometry& MyGeometry, const FPointerEvent& MouseEvent ) OVERRIDE;
-	virtual FCursorReply OnCursorQuery( const FGeometry& MyGeometry, const FPointerEvent& CursorEvent ) const OVERRIDE;
+	virtual FVector2D ComputeDesiredSize() const override;
+	virtual int32 OnPaint( const FGeometry& AllottedGeometry, const FSlateRect& MyClippingRect, FSlateWindowElementList& OutDrawElements, int32 LayerId, const FWidgetStyle& InWidgetStyle, bool bParentEnabled ) const override;
+	virtual FReply OnMouseButtonDown( const FGeometry& MyGeometry, const FPointerEvent& MouseEvent ) override;
+	virtual FReply OnMouseButtonUp( const FGeometry& MyGeometry, const FPointerEvent& MouseEvent ) override;
+	virtual FReply OnMouseMove( const FGeometry& MyGeometry, const FPointerEvent& MouseEvent ) override;
+	virtual FReply OnMouseWheel( const FGeometry& MyGeometry, const FPointerEvent& MouseEvent ) override;
+	virtual FCursorReply OnCursorQuery( const FGeometry& MyGeometry, const FPointerEvent& CursorEvent ) const override;
 
 	static int32 GetDivider( float InputMinX, float InputMaxX, FVector2D WidgetSize, float SequenceLength, int32 NumFrames );
 protected:
diff --git a/Engine/Source/Editor/LandscapeEditor/Classes/ActorFactoryLandscape.h b/Engine/Source/Editor/LandscapeEditor/Classes/ActorFactoryLandscape.h
--- a/Engine/Source/Editor/LandscapeEditor/Classes/ActorFactoryLandscape.h
+++ b/Engine/Source/Editor/LandscapeEditor/Classes/ActorFactoryLandscape.h
@@ -11,7 +11,7 @@ class UActorFactoryLandscape : public UActorFactory
 
 	// Begin UActorFactory Interface
 	//virtual bool CanCreateActorFrom( const FAssetData& AssetData, FText& OutErrorMsg ) OVERRIDE;
-	virtual AActor* SpawnActor( UObject* Asset, ULevel* InLevel, const FVector& Location, const FRotator& Rotation, EObjectFlags ObjectFlags, const FName& Name ) OVERRIDE;
+	virtual AActor* SpawnActor( UObject* Asset, ULevel* InLevel, const FVector& Location, const FRotator& Rotation, EObjectFlags ObjectFlags, const FName& Name ) override;
 	//virtual void PostSpawnActor( UObject* Asset, AActor* NewActor) OVERRIDE;
 	//virtual void PostCreateBlueprint( UObject* Asset, AActor* CDO ) OVERRIDE;
 	//virtual UObject* GetAssetFromActorInstance(AActor* ActorInstance) OVERRIDE;
diff --git a/Engine/Source/Editor/LandscapeEditor/Classes/LandscapeEditorObject.h b/Engine/Source/Editor/LandscapeEditor/Classes/LandscapeEditorObject.h
--- a/Engine/Source/Editor/LandscapeEditor/Classes/LandscapeEditorObject.h
+++ b/Engine/Source/Editor/LandscapeEditor/Classes/LandscapeEditorObject.h
@@ -489,7 +489,7 @@ class ULandscapeEditorObject : public UObject
 
 #if WITH_EDITOR
 	// Begin UObject Interface
-	void PostEditChangeProperty(FPropertyChangedEvent& PropertyChangedEvent) OVERRIDE;
+	void PostEditChangeProperty(FPropertyChangedEvent& PropertyChangedEvent) override;
 	// End UObject Interface
 #endif // WITH_EDITOR
 
diff --git a/Engine/Source/Editor/LandscapeEditor/Classes/LandscapePlaceholder.h b/Engine/Source/Editor/LandscapeEditor/Classes/LandscapePlaceholder.h
--- a/Engine/Source/Editor/LandscapeEditor/Classes/LandscapePlaceholder.h
+++ b/Engine/Source/Editor/LandscapeEditor/Classes/LandscapePlaceholder.h
@@ -10,7 +10,7 @@ class ALandscapePlaceholder : public AActor
 	GENERATED_UCLASS_BODY()
 
 	//virtual void PostEditMove(bool bFinished) OVERRIDE;
-	virtual bool TeleportTo(const FVector& DestLocation, const FRotator& DestRotation, bool bIsATest = false, bool bNoCheck = false) OVERRIDE;
+	virtual bool TeleportTo(const FVector& DestLocation, const FRotator& DestRotation, bool bIsATest = false, bool bNoCheck = false) override;
 
 	void PostEditChangeProperty(FPropertyChangedEvent& PropertyChangedEvent);
 };
diff --git a/Engine/Source/Editor/LandscapeEditor/Private/LandscapeEdModeBrushes.cpp b/Engine/Source/Editor/LandscapeEditor/Private/LandscapeEdModeBrushes.cpp
--- a/Engine/Source/Editor/LandscapeEditor/Private/LandscapeEdModeBrushes.cpp
+++ b/Engine/Source/Editor/LandscapeEditor/Private/LandscapeEdModeBrushes.cpp
@@ -64,7 +64,7 @@ public:
 	class FEdModeLandscape* EdMode;
 
 	// FGCObject interface
-	virtual void AddReferencedObjects(FReferenceCollector& Collector) OVERRIDE
+	virtual void AddReferencedObjects(FReferenceCollector& Collector) override
 	{
 		Collector.AddReferencedObject(BrushMaterial);
 
@@ -77,7 +77,7 @@ public:
 		}
 	}
 
-	virtual void LeaveBrush() OVERRIDE
+	virtual void LeaveBrush() override
 	{
 		for( TSet<ULandscapeComponent*>::TIterator It(BrushMaterialComponents); It; ++It )
 		{
@@ -93,13 +93,13 @@ public:
 		BrushMaterialComponents.Empty();
 	}
 
-	virtual void BeginStroke(float LandscapeX, float LandscapeY, class FLandscapeTool* CurrentTool) OVERRIDE
+	virtual void BeginStroke(float LandscapeX, float LandscapeY, class FLandscapeTool* CurrentTool) override
 	{
 		FLandscapeBrush::BeginStroke(LandscapeX,LandscapeY,CurrentTool);
 		LastMousePosition = FVector2D(LandscapeX, LandscapeY);
 	}
 
-	virtual void Tick(FLevelEditorViewportClient* ViewportClient,float DeltaTime) OVERRIDE
+	virtual void Tick(FLevelEditorViewportClient* ViewportClient,float DeltaTime) override
 	{
 		ULandscapeInfo* LandscapeInfo = EdMode->CurrentToolTarget.LandscapeInfo.Get();
 		ALandscapeProxy* Proxy = LandscapeInfo ? LandscapeInfo->GetLandscapeProxy() : NULL;
@@ -200,18 +200,18 @@ public:
 		}
 	}
 
-	virtual void MouseMove(float LandscapeX, float LandscapeY) OVERRIDE
+	virtual void MouseMove(float LandscapeX, float LandscapeY) override
 	{
 		LastMousePosition = FVector2D(LandscapeX, LandscapeY);
 	}
 
-	virtual bool InputKey( FLevelEditorViewportClient* InViewportClient, FViewport* InViewport, FKey InKey, EInputEvent InEvent ) OVERRIDE
+	virtual bool InputKey( FLevelEditorViewportClient* InViewportClient, FViewport* InViewport, FKey InKey, EInputEvent InEvent ) override
 	{
 		bool bUpdate = false;
 		return bUpdate;
 	}
 
-	virtual bool ApplyBrush( const TArray<FLandscapeToolMousePosition>& MousePositions, TMap<FIntPoint, float>& OutBrush, int32& X1, int32& Y1, int32& X2, int32& Y2 ) OVERRIDE
+	virtual bool ApplyBrush( const TArray<FLandscapeToolMousePosition>& MousePositions, TMap<FIntPoint, float>& OutBrush, int32& X1, int32& Y1, int32& X2, int32& Y2 ) override
 	{
 		ULandscapeInfo* LandscapeInfo = EdMode->CurrentToolTarget.LandscapeInfo.Get();
 		if (LandscapeInfo)
@@ -299,8 +299,8 @@ class FLandscapeBrushComponent : public FLandscapeBrush
 {
 	TSet<ULandscapeComponent*> BrushMaterialComponents;
 
-	virtual const TCHAR* GetBrushName() OVERRIDE { return TEXT("Component"); }
-	virtual FText GetDisplayName() OVERRIDE { return NSLOCTEXT("UnrealEd", "LandscapeMode_Brush_Component", "Component"); };
+	virtual const TCHAR* GetBrushName() override { return TEXT("Component"); }
+	virtual FText GetDisplayName() override { return NSLOCTEXT("UnrealEd", "LandscapeMode_Brush_Component", "Component"); };
 
 protected:
 	FVector2D LastMousePosition;
@@ -316,14 +316,14 @@ public:
 	}
 
 	// FGCObject interface
-	virtual void AddReferencedObjects(FReferenceCollector& Collector) OVERRIDE
+	virtual void AddReferencedObjects(FReferenceCollector& Collector) override
 	{
 		Collector.AddReferencedObject(BrushMaterial);
 	}
 
-	virtual EBrushType GetBrushType() OVERRIDE { return BT_Component; }
+	virtual EBrushType GetBrushType() override { return BT_Component; }
 
-	virtual void LeaveBrush() OVERRIDE
+	virtual void LeaveBrush() override
 	{
 		for( TSet<ULandscapeComponent*>::TIterator It(BrushMaterialComponents); It; ++It )
 		{
@@ -335,13 +335,13 @@ public:
 		BrushMaterialComponents.Empty();
 	}
 
-	virtual void BeginStroke(float LandscapeX, float LandscapeY, class FLandscapeTool* CurrentTool) OVERRIDE
+	virtual void BeginStroke(float LandscapeX, float LandscapeY, class FLandscapeTool* CurrentTool) override
 	{
 		FLandscapeBrush::BeginStroke(LandscapeX,LandscapeY,CurrentTool);
 		LastMousePosition = FVector2D(LandscapeX, LandscapeY);
 	}
 
-	virtual void Tick(FLevelEditorViewportClient* ViewportClient,float DeltaTime) OVERRIDE
+	virtual void Tick(FLevelEditorViewportClient* ViewportClient,float DeltaTime) override
 	{
 		TSet<ULandscapeComponent*> NewComponents;
 
@@ -403,12 +403,12 @@ public:
 		BrushMaterialComponents = MoveTemp(NewComponents);
 	}
 
-	virtual void MouseMove(float LandscapeX, float LandscapeY) OVERRIDE
+	virtual void MouseMove(float LandscapeX, float LandscapeY) override
 	{
 		LastMousePosition = FVector2D(LandscapeX, LandscapeY);
 	}
 
-	virtual bool InputKey( FLevelEditorViewportClient* InViewportClient, FViewport* InViewport, FKey InKey, EInputEvent InEvent ) OVERRIDE
+	virtual bool InputKey( FLevelEditorViewportClient* InViewportClient, FViewport* InViewport, FKey InKey, EInputEvent InEvent ) override
 	{
 		bool bUpdate = false;
 
@@ -416,7 +416,7 @@ public:
 		return bUpdate;
 	}
 
-	virtual bool ApplyBrush( const TArray<FLandscapeToolMousePosition>& MousePositions, TMap<FIntPoint, float>& OutBrush, int32& X1, int32& Y1, int32& X2, int32& Y2 ) OVERRIDE
+	virtual bool ApplyBrush( const TArray<FLandscapeToolMousePosition>& MousePositions, TMap<FIntPoint, float>& OutBrush, int32& X1, int32& Y1, int32& X2, int32& Y2 ) override
 	{
 		// Selection Brush only works for 
 		ULandscapeInfo* LandscapeInfo = EdMode->CurrentToolTarget.LandscapeInfo.Get();
@@ -492,8 +492,8 @@ class FLandscapeBrushGizmo : public FLandscapeBrush
 {
 	TSet<ULandscapeComponent*> BrushMaterialComponents;
 
-	const TCHAR* GetBrushName() OVERRIDE { return TEXT("Gizmo"); }
-	virtual FText GetDisplayName() OVERRIDE { return NSLOCTEXT("UnrealEd", "LandscapeMode_Brush_Gizmo", "Gizmo"); };
+	const TCHAR* GetBrushName() override { return TEXT("Gizmo"); }
+	virtual FText GetDisplayName() override { return NSLOCTEXT("UnrealEd", "LandscapeMode_Brush_Gizmo", "Gizmo"); };
 
 protected:
 	FVector2D LastMousePosition;
@@ -510,14 +510,14 @@ public:
 	}
 
 	// FGCObject interface
-	virtual void AddReferencedObjects(FReferenceCollector& Collector) OVERRIDE
+	virtual void AddReferencedObjects(FReferenceCollector& Collector) override
 	{
 		Collector.AddReferencedObject(BrushMaterial);
 	}
 
-	virtual EBrushType GetBrushType() OVERRIDE { return BT_Gizmo; }
+	virtual EBrushType GetBrushType() override { return BT_Gizmo; }
 
-	virtual void LeaveBrush() OVERRIDE
+	virtual void LeaveBrush() override
 	{
 		for( TSet<ULandscapeComponent*>::TIterator It(BrushMaterialComponents); It; ++It )
 		{
@@ -529,13 +529,13 @@ public:
 		BrushMaterialComponents.Empty();
 	}
 
-	virtual void BeginStroke(float LandscapeX, float LandscapeY, class FLandscapeTool* CurrentTool) OVERRIDE
+	virtual void BeginStroke(float LandscapeX, float LandscapeY, class FLandscapeTool* CurrentTool) override
 	{
 		FLandscapeBrush::BeginStroke(LandscapeX,LandscapeY,CurrentTool);
 		LastMousePosition = FVector2D(LandscapeX, LandscapeY);
 	}
 
-	virtual void Tick(FLevelEditorViewportClient* ViewportClient,float DeltaTime) OVERRIDE
+	virtual void Tick(FLevelEditorViewportClient* ViewportClient,float DeltaTime) override
 	{
 		if (GLandscapeEditRenderMode & ELandscapeEditRenderMode::Gizmo || GLandscapeEditRenderMode & ELandscapeEditRenderMode::Select)
 		{
@@ -606,18 +606,18 @@ public:
 		}
 	}
 
-	virtual void MouseMove(float LandscapeX, float LandscapeY) OVERRIDE
+	virtual void MouseMove(float LandscapeX, float LandscapeY) override
 	{
 		LastMousePosition = FVector2D(LandscapeX, LandscapeY);
 	}
 
-	virtual bool InputKey( FLevelEditorViewportClient* InViewportClient, FViewport* InViewport, FKey InKey, EInputEvent InEvent ) OVERRIDE
+	virtual bool InputKey( FLevelEditorViewportClient* InViewportClient, FViewport* InViewport, FKey InKey, EInputEvent InEvent ) override
 	{
 		bool bUpdate = false;
 		return bUpdate;
 	}
 
-	virtual bool ApplyBrush( const TArray<FLandscapeToolMousePosition>& MousePositions, TMap<FIntPoint, float>& OutBrush, int32& X1, int32& Y1, int32& X2, int32& Y2 ) OVERRIDE
+	virtual bool ApplyBrush( const TArray<FLandscapeToolMousePosition>& MousePositions, TMap<FIntPoint, float>& OutBrush, int32& X1, int32& Y1, int32& X2, int32& Y2 ) override
 	{
 		// Selection Brush only works for 
 		ALandscapeGizmoActiveActor* Gizmo = EdMode->CurrentGizmoActor.Get();
@@ -711,8 +711,8 @@ public:
 class FLandscapeBrushSplines : public FLandscapeBrush
 {
 public:
-	const TCHAR* GetBrushName() OVERRIDE { return TEXT("Splines"); }
-	virtual FText GetDisplayName() OVERRIDE { return NSLOCTEXT("UnrealEd", "LandscapeMode_Brush_Splines", "Splines"); };
+	const TCHAR* GetBrushName() override { return TEXT("Splines"); }
+	virtual FText GetDisplayName() override { return NSLOCTEXT("UnrealEd", "LandscapeMode_Brush_Splines", "Splines"); };
 
 	class FEdModeLandscape* EdMode;
 
@@ -725,13 +725,13 @@ public:
 	{
 	}
 
-	virtual EBrushType GetBrushType() OVERRIDE { return BT_Splines; }
+	virtual EBrushType GetBrushType() override { return BT_Splines; }
 
-	virtual void MouseMove(float LandscapeX, float LandscapeY) OVERRIDE 
+	virtual void MouseMove(float LandscapeX, float LandscapeY) override 
 	{
 	}
 
-	virtual bool ApplyBrush( const TArray<FLandscapeToolMousePosition>& MousePositions, TMap<FIntPoint, float>& OutBrush, int32& X1, int32& Y1, int32& X2, int32& Y2 ) OVERRIDE
+	virtual bool ApplyBrush( const TArray<FLandscapeToolMousePosition>& MousePositions, TMap<FIntPoint, float>& OutBrush, int32& X1, int32& Y1, int32& X2, int32& Y2 ) override
 	{
 		return false;
 	}
@@ -744,7 +744,7 @@ class FLandscapeBrushDummy : public FLandscapeBrush
 {
 public:
 	const TCHAR* GetBrushName() { return TEXT("None"); }
-	virtual FText GetDisplayName() OVERRIDE { return NSLOCTEXT("UnrealEd", "LandscapeMode_Brush_None", "None"); };
+	virtual FText GetDisplayName() override { return NSLOCTEXT("UnrealEd", "LandscapeMode_Brush_None", "None"); };
 
 	class FEdModeLandscape* EdMode;
 
@@ -757,13 +757,13 @@ public:
 	{
 	}
 
-	virtual EBrushType GetBrushType() OVERRIDE { return BT_Normal; }
+	virtual EBrushType GetBrushType() override { return BT_Normal; }
 
-	virtual void MouseMove(float LandscapeX, float LandscapeY) OVERRIDE 
+	virtual void MouseMove(float LandscapeX, float LandscapeY) override 
 	{
 	}
 
-	virtual bool ApplyBrush( const TArray<FLandscapeToolMousePosition>& MousePositions, TMap<FIntPoint, float>& OutBrush, int32& X1, int32& Y1, int32& X2, int32& Y2 ) OVERRIDE
+	virtual bool ApplyBrush( const TArray<FLandscapeToolMousePosition>& MousePositions, TMap<FIntPoint, float>& OutBrush, int32& X1, int32& Y1, int32& X2, int32& Y2 ) override
 	{
 		return false;
 	}
@@ -778,7 +778,7 @@ protected:
 	{
 	}
 
-	virtual float CalculateFalloff( float Distance, float Radius, float Falloff ) OVERRIDE
+	virtual float CalculateFalloff( float Distance, float Radius, float Falloff ) override
 	{
 		return Distance < Radius ? 1.f : 
 			Falloff > 0.f ? FMath::Max<float>(0.f, 1.f - (Distance - Radius) / Falloff) : 
@@ -791,8 +791,8 @@ public:
 		UMaterialInstanceConstant* CircleBrushMaterial_Linear = LoadObject<UMaterialInstanceConstant>(NULL, TEXT("/Engine/EditorLandscapeResources/CircleBrushMaterial_Linear.CircleBrushMaterial_Linear"), NULL, LOAD_None, NULL);
 		return new FLandscapeBrushCircle_Linear(InEdMode, CircleBrushMaterial_Linear);
 	}
-	virtual const TCHAR* GetBrushName() OVERRIDE { return TEXT("Circle_Linear"); }
-	virtual FText GetDisplayName() OVERRIDE { return NSLOCTEXT("UnrealEd", "LandscapeMode_Brush_Falloff_Linear", "Linear falloff"); };
+	virtual const TCHAR* GetBrushName() override { return TEXT("Circle_Linear"); }
+	virtual FText GetDisplayName() override { return NSLOCTEXT("UnrealEd", "LandscapeMode_Brush_Falloff_Linear", "Linear falloff"); };
 
 };
 
@@ -804,7 +804,7 @@ protected:
 	{
 	}
 
-	virtual float CalculateFalloff( float Distance, float Radius, float Falloff ) OVERRIDE
+	virtual float CalculateFalloff( float Distance, float Radius, float Falloff ) override
 	{
 		float y = FLandscapeBrushCircle_Linear::CalculateFalloff(Distance, Radius, Falloff);
 		// Smooth-step it
@@ -817,8 +817,8 @@ public:
 		UMaterialInstanceConstant* CircleBrushMaterial_Smooth = LoadObject<UMaterialInstanceConstant>(NULL, TEXT("/Engine/EditorLandscapeResources/CircleBrushMaterial_Smooth.CircleBrushMaterial_Smooth"), NULL, LOAD_None, NULL);
 		return new FLandscapeBrushCircle_Smooth(InEdMode, CircleBrushMaterial_Smooth);
 	}
-	virtual const TCHAR* GetBrushName() OVERRIDE { return TEXT("Circle_Smooth"); }
-	virtual FText GetDisplayName() OVERRIDE { return NSLOCTEXT("UnrealEd", "LandscapeMode_Brush_Falloff_Smooth", "Smooth falloff"); };
+	virtual const TCHAR* GetBrushName() override { return TEXT("Circle_Smooth"); }
+	virtual FText GetDisplayName() override { return NSLOCTEXT("UnrealEd", "LandscapeMode_Brush_Falloff_Smooth", "Smooth falloff"); };
 
 };
 
@@ -830,7 +830,7 @@ protected:
 	{
 	}
 
-	virtual float CalculateFalloff( float Distance, float Radius, float Falloff ) OVERRIDE
+	virtual float CalculateFalloff( float Distance, float Radius, float Falloff ) override
 	{
 		if( Distance <= Radius )
 		{
@@ -852,8 +852,8 @@ public:
 		UMaterialInstanceConstant* CircleBrushMaterial_Spherical = LoadObject<UMaterialInstanceConstant>(NULL, TEXT("/Engine/EditorLandscapeResources/CircleBrushMaterial_Spherical.CircleBrushMaterial_Spherical"), NULL, LOAD_None, NULL);
 		return new FLandscapeBrushCircle_Spherical(InEdMode, CircleBrushMaterial_Spherical);
 	}
-	virtual const TCHAR* GetBrushName() OVERRIDE { return TEXT("Circle_Spherical"); }
-	virtual FText GetDisplayName() OVERRIDE { return NSLOCTEXT("UnrealEd", "LandscapeMode_Brush_Falloff_Spherical", "Spherical falloff"); };
+	virtual const TCHAR* GetBrushName() override { return TEXT("Circle_Spherical"); }
+	virtual FText GetDisplayName() override { return NSLOCTEXT("UnrealEd", "LandscapeMode_Brush_Falloff_Spherical", "Spherical falloff"); };
 };
 
 class FLandscapeBrushCircle_Tip : public FLandscapeBrushCircle
@@ -864,7 +864,7 @@ protected:
 	{
 	}
 
-	virtual float CalculateFalloff( float Distance, float Radius, float Falloff ) OVERRIDE
+	virtual float CalculateFalloff( float Distance, float Radius, float Falloff ) override
 	{
 		if( Distance <= Radius )
 		{
@@ -886,8 +886,8 @@ public:
 		UMaterialInstanceConstant* CircleBrushMaterial_Tip = LoadObject<UMaterialInstanceConstant>(NULL, TEXT("/Engine/EditorLandscapeResources/CircleBrushMaterial_Tip.CircleBrushMaterial_Tip"), NULL, LOAD_None, NULL);
 		return new FLandscapeBrushCircle_Tip(InEdMode, CircleBrushMaterial_Tip);
 	}
-	virtual const TCHAR* GetBrushName() OVERRIDE { return TEXT("Circle_Tip"); }
-	virtual FText GetDisplayName() OVERRIDE { return NSLOCTEXT("UnrealEd", "LandscapeMode_Brush_Falloff_Tip", "Tip falloff"); };
+	virtual const TCHAR* GetBrushName() override { return TEXT("Circle_Tip"); }
+	virtual FText GetDisplayName() override { return NSLOCTEXT("UnrealEd", "LandscapeMode_Brush_Falloff_Tip", "Tip falloff"); };
 };
 
 
@@ -944,9 +944,9 @@ public:
 		return new FLandscapeBrushAlphaPattern(InEdMode, PatternBrushMaterial);
 	}
 
-	virtual EBrushType GetBrushType() OVERRIDE { return BT_Alpha; }
+	virtual EBrushType GetBrushType() override { return BT_Alpha; }
 
-	virtual bool ApplyBrush( const TArray<FLandscapeToolMousePosition>& MousePositions, TMap<FIntPoint, float>& OutBrush, int32& X1, int32& Y1, int32& X2, int32& Y2 ) OVERRIDE
+	virtual bool ApplyBrush( const TArray<FLandscapeToolMousePosition>& MousePositions, TMap<FIntPoint, float>& OutBrush, int32& X1, int32& Y1, int32& X2, int32& Y2 ) override
 	{
 		ULandscapeInfo* LandscapeInfo = EdMode->CurrentToolTarget.LandscapeInfo.Get();
 		float ScaleXY = LandscapeInfo->DrawScale.X;
@@ -1014,7 +1014,7 @@ public:
 		return (X1 <= X2 && Y1 <= Y2);
 	}
 
-	virtual void Tick(FLevelEditorViewportClient* ViewportClient,float DeltaTime) OVERRIDE
+	virtual void Tick(FLevelEditorViewportClient* ViewportClient,float DeltaTime) override
 	{
 		FLandscapeBrushCircle::Tick(ViewportClient,DeltaTime);
 
@@ -1048,8 +1048,8 @@ public:
 		}
 	}
 
-	virtual const TCHAR* GetBrushName() OVERRIDE { return TEXT("Pattern"); }
-	virtual FText GetDisplayName() OVERRIDE { return NSLOCTEXT("UnrealEd", "LandscapeMode_Brush_PatternAlpha", "Pattern Alpha"); };
+	virtual const TCHAR* GetBrushName() override { return TEXT("Pattern"); }
+	virtual FText GetDisplayName() override { return NSLOCTEXT("UnrealEd", "LandscapeMode_Brush_PatternAlpha", "Pattern Alpha"); };
 };
 
 
@@ -1078,7 +1078,7 @@ public:
 		return new FLandscapeBrushAlpha(InEdMode, AlphaBrushMaterial);
 	}
 
-	virtual bool ApplyBrush( const TArray<FLandscapeToolMousePosition>& MousePositions, TMap<FIntPoint, float>& OutBrush, int32& X1, int32& Y1, int32& X2, int32& Y2 ) OVERRIDE
+	virtual bool ApplyBrush( const TArray<FLandscapeToolMousePosition>& MousePositions, TMap<FIntPoint, float>& OutBrush, int32& X1, int32& Y1, int32& X2, int32& Y2 ) override
 	{
 		ULandscapeInfo* LandscapeInfo = EdMode->CurrentToolTarget.LandscapeInfo.Get();
 		if( OldMousePosition.IsZero() )
@@ -1152,7 +1152,7 @@ public:
 		return (X1 <= X2 && Y1 <= Y2);
 	}
 
-	virtual void MouseMove(float LandscapeX, float LandscapeY) OVERRIDE
+	virtual void MouseMove(float LandscapeX, float LandscapeY) override
 	{
 		FLandscapeBrushAlphaBase::MouseMove(LandscapeX, LandscapeY);
 
@@ -1171,7 +1171,7 @@ public:
 		}
 	}
 
-	virtual void Tick(FLevelEditorViewportClient* ViewportClient,float DeltaTime) OVERRIDE
+	virtual void Tick(FLevelEditorViewportClient* ViewportClient,float DeltaTime) override
 	{
 		FLandscapeBrushCircle::Tick(ViewportClient,DeltaTime);
 
@@ -1204,8 +1204,8 @@ public:
 		}
 	}
 
-	virtual const TCHAR* GetBrushName() OVERRIDE { return TEXT("Alpha"); }
-	virtual FText GetDisplayName() OVERRIDE { return NSLOCTEXT("UnrealEd", "LandscapeMode_Brush_Alpha", "Alpha"); };
+	virtual const TCHAR* GetBrushName() override { return TEXT("Alpha"); }
+	virtual FText GetDisplayName() override { return NSLOCTEXT("UnrealEd", "LandscapeMode_Brush_Alpha", "Alpha"); };
 
 };
 
diff --git a/Engine/Source/Editor/LandscapeEditor/Private/LandscapeEdModeComponentTools.cpp b/Engine/Source/Editor/LandscapeEditor/Private/LandscapeEdModeComponentTools.cpp
--- a/Engine/Source/Editor/LandscapeEditor/Private/LandscapeEdModeComponentTools.cpp
+++ b/Engine/Source/Editor/LandscapeEditor/Private/LandscapeEdModeComponentTools.cpp
@@ -1,4 +1,4 @@
-﻿// Copyright 1998-2014 Epic Games, Inc. All Rights Reserved.
+// Copyright 1998-2014 Epic Games, Inc. All Rights Reserved.
 
 #include "LandscapeEditorPrivatePCH.h"
 #include "ObjectTools.h"
@@ -167,14 +167,14 @@ public:
 		: FLandscapeToolBase<TStrokeClass>(InEdMode)
 	{}
 
-	virtual const TCHAR* GetToolName() OVERRIDE { return TEXT("Selection"); }
-	virtual FText GetDisplayName() OVERRIDE { return NSLOCTEXT("UnrealEd", "LandscapeMode_Selection", "Component Selection"); };
-	virtual void SetEditRenderType() OVERRIDE { GLandscapeEditRenderMode = ELandscapeEditRenderMode::SelectComponent | (GLandscapeEditRenderMode & ELandscapeEditRenderMode::BitMaskForMask); }
-	virtual bool SupportsMask() OVERRIDE { return false; }
+	virtual const TCHAR* GetToolName() override { return TEXT("Selection"); }
+	virtual FText GetDisplayName() override { return NSLOCTEXT("UnrealEd", "LandscapeMode_Selection", "Component Selection"); };
+	virtual void SetEditRenderType() override { GLandscapeEditRenderMode = ELandscapeEditRenderMode::SelectComponent | (GLandscapeEditRenderMode & ELandscapeEditRenderMode::BitMaskForMask); }
+	virtual bool SupportsMask() override { return false; }
 
-	virtual FLandscapeTool::EToolType GetToolType() OVERRIDE { return this->TT_Mask; }
+	virtual FLandscapeTool::EToolType GetToolType() override { return this->TT_Mask; }
 
-	virtual bool IsValidForTarget(const FLandscapeToolTarget& Target) OVERRIDE
+	virtual bool IsValidForTarget(const FLandscapeToolTarget& Target) override
 	{
 		return true; // applied to all...
 	}
@@ -188,10 +188,10 @@ public:
 		: FLandscapeToolSelect<TStrokeClass>(InEdMode)
 	{}
 
-	virtual const TCHAR* GetToolName() OVERRIDE { return TEXT("Mask"); }
-	virtual FText GetDisplayName() OVERRIDE { return NSLOCTEXT("UnrealEd", "LandscapeMode_Mask", "Region Selection"); };
-	virtual void SetEditRenderType() OVERRIDE { GLandscapeEditRenderMode = ELandscapeEditRenderMode::SelectRegion | (GLandscapeEditRenderMode & ELandscapeEditRenderMode::BitMaskForMask); }
-	virtual bool SupportsMask() OVERRIDE { return true; }
+	virtual const TCHAR* GetToolName() override { return TEXT("Mask"); }
+	virtual FText GetDisplayName() override { return NSLOCTEXT("UnrealEd", "LandscapeMode_Mask", "Region Selection"); };
+	virtual void SetEditRenderType() override { GLandscapeEditRenderMode = ELandscapeEditRenderMode::SelectRegion | (GLandscapeEditRenderMode & ELandscapeEditRenderMode::BitMaskForMask); }
+	virtual bool SupportsMask() override { return true; }
 };
 
 class FLandscapeToolStrokeVisibility
@@ -271,13 +271,13 @@ public:
 		: FLandscapeToolBase<FLandscapeToolStrokeVisibility>(InEdMode)
 	{}
 
-	virtual const TCHAR* GetToolName() OVERRIDE { return TEXT("Visibility"); }
-	virtual FText GetDisplayName() OVERRIDE { return NSLOCTEXT("UnrealEd", "LandscapeMode_Visibility", "Visibility"); };
+	virtual const TCHAR* GetToolName() override { return TEXT("Visibility"); }
+	virtual FText GetDisplayName() override { return NSLOCTEXT("UnrealEd", "LandscapeMode_Visibility", "Visibility"); };
 
-	virtual void SetEditRenderType() OVERRIDE { GLandscapeEditRenderMode = ELandscapeEditRenderMode::None | (GLandscapeEditRenderMode & ELandscapeEditRenderMode::BitMaskForMask); }
-	virtual bool SupportsMask() OVERRIDE { return false; }
+	virtual void SetEditRenderType() override { GLandscapeEditRenderMode = ELandscapeEditRenderMode::None | (GLandscapeEditRenderMode & ELandscapeEditRenderMode::BitMaskForMask); }
+	virtual bool SupportsMask() override { return false; }
 
-	virtual bool IsValidForTarget(const FLandscapeToolTarget& Target) OVERRIDE
+	virtual bool IsValidForTarget(const FLandscapeToolTarget& Target) override
 	{
 		return true; // applied to all...
 	}
@@ -735,13 +735,13 @@ public:
 		: FLandscapeToolBase<FLandscapeToolStrokeMoveToLevel>(InEdMode)	
 	{}
 
-	virtual const TCHAR* GetToolName() OVERRIDE { return TEXT("MoveToLevel"); }
-	virtual FText GetDisplayName() OVERRIDE { return NSLOCTEXT("UnrealEd", "LandscapeMode_MoveToLevel", "Move to Streaming Level"); };
+	virtual const TCHAR* GetToolName() override { return TEXT("MoveToLevel"); }
+	virtual FText GetDisplayName() override { return NSLOCTEXT("UnrealEd", "LandscapeMode_MoveToLevel", "Move to Streaming Level"); };
 
-	virtual void SetEditRenderType() OVERRIDE { GLandscapeEditRenderMode = ELandscapeEditRenderMode::SelectComponent | (GLandscapeEditRenderMode & ELandscapeEditRenderMode::BitMaskForMask); }
-	virtual bool SupportsMask() OVERRIDE { return false; }
+	virtual void SetEditRenderType() override { GLandscapeEditRenderMode = ELandscapeEditRenderMode::SelectComponent | (GLandscapeEditRenderMode & ELandscapeEditRenderMode::BitMaskForMask); }
+	virtual bool SupportsMask() override { return false; }
 
-	virtual bool IsValidForTarget(const FLandscapeToolTarget& Target) OVERRIDE
+	virtual bool IsValidForTarget(const FLandscapeToolTarget& Target) override
 	{
 		return true; // applied to all...
 	}
@@ -866,13 +866,13 @@ public:
 		: FLandscapeToolBase<FLandscapeToolStrokeAddComponent>(InEdMode)
 	{}
 
-	virtual const TCHAR* GetToolName() OVERRIDE { return TEXT("AddComponent"); }
-	virtual FText GetDisplayName() OVERRIDE { return NSLOCTEXT("UnrealEd", "LandscapeMode_AddComponent", "Add New Landscape Component"); };
+	virtual const TCHAR* GetToolName() override { return TEXT("AddComponent"); }
+	virtual FText GetDisplayName() override { return NSLOCTEXT("UnrealEd", "LandscapeMode_AddComponent", "Add New Landscape Component"); };
 
-	virtual void SetEditRenderType() OVERRIDE { GLandscapeEditRenderMode = ELandscapeEditRenderMode::None | (GLandscapeEditRenderMode & ELandscapeEditRenderMode::BitMaskForMask); }
-	virtual bool SupportsMask() OVERRIDE { return false; }
+	virtual void SetEditRenderType() override { GLandscapeEditRenderMode = ELandscapeEditRenderMode::None | (GLandscapeEditRenderMode & ELandscapeEditRenderMode::BitMaskForMask); }
+	virtual bool SupportsMask() override { return false; }
 
-	virtual bool IsValidForTarget(const FLandscapeToolTarget& Target) OVERRIDE
+	virtual bool IsValidForTarget(const FLandscapeToolTarget& Target) override
 	{
 		return true; // applied to all...
 	}
@@ -1045,13 +1045,13 @@ public:
 		: FLandscapeToolBase<FLandscapeToolStrokeDeleteComponent>(InEdMode)
 	{}
 
-	virtual const TCHAR* GetToolName() OVERRIDE { return TEXT("DeleteComponent"); }
-	virtual FText GetDisplayName() OVERRIDE { return NSLOCTEXT("UnrealEd", "LandscapeMode_DeleteComponent", "Delete Landscape Components"); };
+	virtual const TCHAR* GetToolName() override { return TEXT("DeleteComponent"); }
+	virtual FText GetDisplayName() override { return NSLOCTEXT("UnrealEd", "LandscapeMode_DeleteComponent", "Delete Landscape Components"); };
 
-	virtual void SetEditRenderType() OVERRIDE { GLandscapeEditRenderMode = ELandscapeEditRenderMode::SelectComponent | (GLandscapeEditRenderMode & ELandscapeEditRenderMode::BitMaskForMask); }
-	virtual bool SupportsMask() OVERRIDE { return false; }
+	virtual void SetEditRenderType() override { GLandscapeEditRenderMode = ELandscapeEditRenderMode::SelectComponent | (GLandscapeEditRenderMode & ELandscapeEditRenderMode::BitMaskForMask); }
+	virtual bool SupportsMask() override { return false; }
 
-	virtual bool IsValidForTarget(const FLandscapeToolTarget& Target) OVERRIDE
+	virtual bool IsValidForTarget(const FLandscapeToolTarget& Target) override
 	{
 		return true; // applied to all...
 	}
@@ -1370,21 +1370,21 @@ public:
 		: FLandscapeToolBase<FLandscapeToolStrokeCopy<ToolTarget> >(InEdMode)
 	{}
 
-	virtual const TCHAR* GetToolName() OVERRIDE { return TEXT("Copy"); }
-	virtual FText GetDisplayName() OVERRIDE { return NSLOCTEXT("UnrealEd", "LandscapeMode_Copy", "Copy"); };
+	virtual const TCHAR* GetToolName() override { return TEXT("Copy"); }
+	virtual FText GetDisplayName() override { return NSLOCTEXT("UnrealEd", "LandscapeMode_Copy", "Copy"); };
 
-	virtual void SetEditRenderType() OVERRIDE
+	virtual void SetEditRenderType() override
 	{ 
 		GLandscapeEditRenderMode = ELandscapeEditRenderMode::Gizmo | (GLandscapeEditRenderMode & ELandscapeEditRenderMode::BitMaskForMask);
 		GLandscapeEditRenderMode |= (this->EdMode && this->EdMode->CurrentToolTarget.LandscapeInfo.IsValid() && this->EdMode->CurrentToolTarget.LandscapeInfo->SelectedRegion.Num()) ? ELandscapeEditRenderMode::SelectRegion : ELandscapeEditRenderMode::SelectComponent;
 	}
 
-	virtual bool IsValidForTarget(const FLandscapeToolTarget& Target) OVERRIDE
+	virtual bool IsValidForTarget(const FLandscapeToolTarget& Target) override
 	{
 		return Target.TargetType == ToolTarget::TargetType;
 	}
 
-	virtual bool BeginTool( FLevelEditorViewportClient* ViewportClient, const FLandscapeToolTarget& InTarget, const FVector& InHitLocation ) OVERRIDE
+	virtual bool BeginTool( FLevelEditorViewportClient* ViewportClient, const FLandscapeToolTarget& InTarget, const FVector& InHitLocation ) override
 	{
 		check(this->MousePositions.Num()==0);
 		this->bToolActive = true;
@@ -1399,7 +1399,7 @@ public:
 		return true;
 	}
 
-	virtual void EndTool(FLevelEditorViewportClient* ViewportClient) OVERRIDE 
+	virtual void EndTool(FLevelEditorViewportClient* ViewportClient) override 
 	{
 		if( this->bToolActive && this->MousePositions.Num() )
 		{
@@ -1413,7 +1413,7 @@ public:
 		this->EdMode->GizmoBrush->EndStroke();
 	}
 
-	virtual bool MouseMove( FLevelEditorViewportClient* ViewportClient, FViewport* Viewport, int x, int y ) OVERRIDE
+	virtual bool MouseMove( FLevelEditorViewportClient* ViewportClient, FViewport* Viewport, int x, int y ) override
 	{
 		return true;
 	}	
@@ -1668,10 +1668,10 @@ public:
 		,	bUseGizmoRegion(false)
 	{}
 
-	virtual const TCHAR* GetToolName() OVERRIDE { return TEXT("Paste"); }
-	virtual FText GetDisplayName() OVERRIDE { return NSLOCTEXT("UnrealEd", "LandscapeMode_Region", "Region Copy/Paste"); };
+	virtual const TCHAR* GetToolName() override { return TEXT("Paste"); }
+	virtual FText GetDisplayName() override { return NSLOCTEXT("UnrealEd", "LandscapeMode_Region", "Region Copy/Paste"); };
 
-	virtual void SetEditRenderType() OVERRIDE 
+	virtual void SetEditRenderType() override 
 	{ 
 		GLandscapeEditRenderMode = ELandscapeEditRenderMode::Gizmo | (GLandscapeEditRenderMode & ELandscapeEditRenderMode::BitMaskForMask);
 		GLandscapeEditRenderMode |= (this->EdMode && this->EdMode->CurrentToolTarget.LandscapeInfo.IsValid() && this->EdMode->CurrentToolTarget.LandscapeInfo->SelectedRegion.Num()) ? ELandscapeEditRenderMode::SelectRegion : ELandscapeEditRenderMode::SelectComponent;
@@ -1682,12 +1682,12 @@ public:
 		bUseGizmoRegion = InbUseGizmoRegion;
 	}
 
-	virtual bool IsValidForTarget(const FLandscapeToolTarget& Target) OVERRIDE
+	virtual bool IsValidForTarget(const FLandscapeToolTarget& Target) override
 	{
 		return Target.TargetType == ToolTarget::TargetType;
 	}
 
-	virtual bool BeginTool( FLevelEditorViewportClient* ViewportClient, const FLandscapeToolTarget& InTarget, const FVector& InHitLocation ) OVERRIDE
+	virtual bool BeginTool( FLevelEditorViewportClient* ViewportClient, const FLandscapeToolTarget& InTarget, const FVector& InHitLocation ) override
 	{
 		this->bToolActive = true;
 		this->ToolStroke = new FLandscapeToolStrokePaste<ToolTarget>(this->EdMode, InTarget);
@@ -1708,7 +1708,7 @@ public:
 		return true;
 	}
 
-	virtual void EndTool(FLevelEditorViewportClient* ViewportClient) OVERRIDE
+	virtual void EndTool(FLevelEditorViewportClient* ViewportClient) override
 	{
 		if( this->bToolActive && this->MousePositions.Num() )
 		{
@@ -1729,7 +1729,7 @@ public:
 		}
 	}
 
-	virtual bool MouseMove( FLevelEditorViewportClient* ViewportClient, FViewport* Viewport, int32 x, int32 y ) OVERRIDE
+	virtual bool MouseMove( FLevelEditorViewportClient* ViewportClient, FViewport* Viewport, int32 x, int32 y ) override
 	{
 		if (bUseGizmoRegion)
 		{
@@ -1763,7 +1763,7 @@ public:
 
 	// Just hybrid of Copy and Paste tool
 	// Copy tool doesn't use any view information, so just do it as one function
-	virtual void Process(int32 Index, int32 Arg) OVERRIDE
+	virtual void Process(int32 Index, int32 Arg) override
 	{
 		switch(Index)
 		{
@@ -1809,13 +1809,13 @@ public:
 	{
 	}
 
-	virtual const TCHAR* GetToolName() OVERRIDE { return TEXT("NewLandscape"); }
-	virtual FText GetDisplayName() OVERRIDE { return NSLOCTEXT("UnrealEd", "LandscapeMode_NewLandscape", "New Landscape"); };
+	virtual const TCHAR* GetToolName() override { return TEXT("NewLandscape"); }
+	virtual FText GetDisplayName() override { return NSLOCTEXT("UnrealEd", "LandscapeMode_NewLandscape", "New Landscape"); };
 
-	virtual void SetEditRenderType() OVERRIDE { GLandscapeEditRenderMode = ELandscapeEditRenderMode::None | (GLandscapeEditRenderMode & ELandscapeEditRenderMode::BitMaskForMask); }
-	virtual bool SupportsMask() OVERRIDE { return false; }
+	virtual void SetEditRenderType() override { GLandscapeEditRenderMode = ELandscapeEditRenderMode::None | (GLandscapeEditRenderMode & ELandscapeEditRenderMode::BitMaskForMask); }
+	virtual bool SupportsMask() override { return false; }
 
-	virtual bool IsValidForTarget(const FLandscapeToolTarget& Target) OVERRIDE
+	virtual bool IsValidForTarget(const FLandscapeToolTarget& Target) override
 	{
 		return true; // ignores target...
 	}
@@ -1864,13 +1864,13 @@ public:
 	{
 	}
 
-	virtual const TCHAR* GetToolName() OVERRIDE { return TEXT("ResizeLandscape"); }
-	virtual FText GetDisplayName() OVERRIDE { return LOCTEXT("LandscapeMode_ResizeLandscape", "Change Landscape Component Size"); };
+	virtual const TCHAR* GetToolName() override { return TEXT("ResizeLandscape"); }
+	virtual FText GetDisplayName() override { return LOCTEXT("LandscapeMode_ResizeLandscape", "Change Landscape Component Size"); };
 
-	virtual void SetEditRenderType() OVERRIDE { GLandscapeEditRenderMode = ELandscapeEditRenderMode::None | (GLandscapeEditRenderMode & ELandscapeEditRenderMode::BitMaskForMask); }
-	virtual bool SupportsMask() OVERRIDE { return false; }
+	virtual void SetEditRenderType() override { GLandscapeEditRenderMode = ELandscapeEditRenderMode::None | (GLandscapeEditRenderMode & ELandscapeEditRenderMode::BitMaskForMask); }
+	virtual bool SupportsMask() override { return false; }
 
-	virtual bool IsValidForTarget(const FLandscapeToolTarget& Target) OVERRIDE
+	virtual bool IsValidForTarget(const FLandscapeToolTarget& Target) override
 	{
 		return true; // ignores target...
 	}
diff --git a/Engine/Source/Editor/LandscapeEditor/Private/LandscapeEdModeErosionTools.cpp b/Engine/Source/Editor/LandscapeEditor/Private/LandscapeEdModeErosionTools.cpp
--- a/Engine/Source/Editor/LandscapeEditor/Private/LandscapeEdModeErosionTools.cpp
+++ b/Engine/Source/Editor/LandscapeEditor/Private/LandscapeEdModeErosionTools.cpp
@@ -42,7 +42,7 @@ public:
 	:	FLandscapeToolBase<TStrokeClass>(InEdMode)
 	{}
 
-	virtual bool IsValidForTarget(const FLandscapeToolTarget& Target) OVERRIDE 
+	virtual bool IsValidForTarget(const FLandscapeToolTarget& Target) override 
 	{
 		// Erosion is applied to all layers
 		return true;
@@ -60,7 +60,7 @@ public:
 	:	FLandscapeToolStrokeErosionBase(InEdMode, InTarget)
 	{}
 
-	virtual void Apply(FLevelEditorViewportClient* ViewportClient, FLandscapeBrush* Brush, const ULandscapeEditorObject* UISettings, const TArray<FLandscapeToolMousePosition>& MousePositions) OVERRIDE
+	virtual void Apply(FLevelEditorViewportClient* ViewportClient, FLandscapeBrush* Brush, const ULandscapeEditorObject* UISettings, const TArray<FLandscapeToolMousePosition>& MousePositions) override
 	{
 		if( !LandscapeInfo )
 		{
@@ -260,8 +260,8 @@ public:
 	:	FLandscapeToolErosionBase(InEdMode)
 	{}
 
-	virtual const TCHAR* GetToolName() OVERRIDE { return TEXT("Erosion"); }
-	virtual FText GetDisplayName() OVERRIDE { return NSLOCTEXT("UnrealEd", "LandscapeMode_Erosion", "Erosion"); };
+	virtual const TCHAR* GetToolName() override { return TEXT("Erosion"); }
+	virtual FText GetDisplayName() override { return NSLOCTEXT("UnrealEd", "LandscapeMode_Erosion", "Erosion"); };
 
 };
 
@@ -276,7 +276,7 @@ public:
 	:	FLandscapeToolStrokeErosionBase(InEdMode, InTarget)
 	{}
 
-	virtual void Apply(FLevelEditorViewportClient* ViewportClient, FLandscapeBrush* Brush, const ULandscapeEditorObject* UISettings, const TArray<FLandscapeToolMousePosition>& MousePositions) OVERRIDE
+	virtual void Apply(FLevelEditorViewportClient* ViewportClient, FLandscapeBrush* Brush, const ULandscapeEditorObject* UISettings, const TArray<FLandscapeToolMousePosition>& MousePositions) override
 	{
 		if( !LandscapeInfo )
 		{
@@ -460,8 +460,8 @@ public:
 	:	FLandscapeToolErosionBase(InEdMode)
 	{}
 
-	virtual const TCHAR* GetToolName() OVERRIDE { return TEXT("HydraulicErosion"); }
-	virtual FText GetDisplayName() OVERRIDE { return NSLOCTEXT("UnrealEd", "LandscapeMode_HydraErosion", "Hydraulic Erosion"); };
+	virtual const TCHAR* GetToolName() override { return TEXT("HydraulicErosion"); }
+	virtual FText GetDisplayName() override { return NSLOCTEXT("UnrealEd", "LandscapeMode_HydraErosion", "Hydraulic Erosion"); };
 
 };
 
diff --git a/Engine/Source/Editor/LandscapeEditor/Private/LandscapeEdModePaintTools.cpp b/Engine/Source/Editor/LandscapeEditor/Private/LandscapeEdModePaintTools.cpp
--- a/Engine/Source/Editor/LandscapeEditor/Private/LandscapeEdModePaintTools.cpp
+++ b/Engine/Source/Editor/LandscapeEditor/Private/LandscapeEdModePaintTools.cpp
@@ -85,7 +85,7 @@ public:
 	:	FLandscapeToolStrokePaintBase<ToolTarget>(InEdMode, InTarget)
 	{}
 
-	virtual void Apply(FLevelEditorViewportClient* ViewportClient, FLandscapeBrush* Brush, const ULandscapeEditorObject* UISettings, const TArray<FLandscapeToolMousePosition>& MousePositions) OVERRIDE
+	virtual void Apply(FLevelEditorViewportClient* ViewportClient, FLandscapeBrush* Brush, const ULandscapeEditorObject* UISettings, const TArray<FLandscapeToolMousePosition>& MousePositions) override
 	{
 		// Get list of verts to update
 		TMap<FIntPoint, float> BrushInfo;
@@ -304,8 +304,8 @@ public:
 	:	FLandscapeToolPaintBase<ToolTarget, FLandscapeToolStrokePaint<ToolTarget> >(InEdMode)
 	{}
 
-	virtual const TCHAR* GetToolName() OVERRIDE { return TEXT("Paint"); }
-	virtual FText GetDisplayName() OVERRIDE { return NSLOCTEXT("UnrealEd", "LandscapeMode_Paint", "Paint"); };
+	virtual const TCHAR* GetToolName() override { return TEXT("Paint"); }
+	virtual FText GetDisplayName() override { return NSLOCTEXT("UnrealEd", "LandscapeMode_Paint", "Paint"); };
 
 };
 
@@ -321,7 +321,7 @@ public:
 	:	FLandscapeToolStrokePaintBase<ToolTarget>(InEdMode, InTarget)
 	{}
 
-	virtual void Apply(FLevelEditorViewportClient* ViewportClient, FLandscapeBrush* Brush, const ULandscapeEditorObject* UISettings, const TArray<FLandscapeToolMousePosition>& MousePositions) OVERRIDE
+	virtual void Apply(FLevelEditorViewportClient* ViewportClient, FLandscapeBrush* Brush, const ULandscapeEditorObject* UISettings, const TArray<FLandscapeToolMousePosition>& MousePositions) override
 	{
 		if (!this->LandscapeInfo) return;
 
@@ -403,8 +403,8 @@ public:
 	:	FLandscapeToolPaintBase<ToolTarget, FLandscapeToolStrokeSmooth<ToolTarget> >(InEdMode)
 	{}
 
-	virtual const TCHAR* GetToolName() OVERRIDE { return TEXT("Smooth"); }
-	virtual FText GetDisplayName() OVERRIDE { return NSLOCTEXT("UnrealEd", "LandscapeMode_Smooth", "Smooth"); };
+	virtual const TCHAR* GetToolName() override { return TEXT("Smooth"); }
+	virtual FText GetDisplayName() override { return NSLOCTEXT("UnrealEd", "LandscapeMode_Smooth", "Smooth"); };
 
 };
 
@@ -436,7 +436,7 @@ public:
 		}
 	}
 
-	virtual void Apply(FLevelEditorViewportClient* ViewportClient, FLandscapeBrush* Brush, const ULandscapeEditorObject* UISettings, const TArray<FLandscapeToolMousePosition>& MousePositions) OVERRIDE
+	virtual void Apply(FLevelEditorViewportClient* ViewportClient, FLandscapeBrush* Brush, const ULandscapeEditorObject* UISettings, const TArray<FLandscapeToolMousePosition>& MousePositions) override
 	{
 		if (!this->LandscapeInfo) return;
 
@@ -569,8 +569,8 @@ public:
 		check(PlaneMesh);
 	}
 
-	virtual const TCHAR* GetToolName() OVERRIDE { return TEXT("Flatten"); }
-	virtual FText GetDisplayName() OVERRIDE { return NSLOCTEXT("UnrealEd", "LandscapeMode_Flatten", "Flatten"); };
+	virtual const TCHAR* GetToolName() override { return TEXT("Flatten"); }
+	virtual FText GetDisplayName() override { return NSLOCTEXT("UnrealEd", "LandscapeMode_Flatten", "Flatten"); };
 
 	virtual void Tick(FLevelEditorViewportClient* ViewportClient,float DeltaTime)
 	{
@@ -580,7 +580,7 @@ public:
 		MeshComponent->SetVisibility(bShowGrid);
 	}
 
-	virtual bool MouseMove(FLevelEditorViewportClient* ViewportClient, FViewport* Viewport, int32 x, int32 y) OVERRIDE
+	virtual bool MouseMove(FLevelEditorViewportClient* ViewportClient, FViewport* Viewport, int32 x, int32 y) override
 	{
 		bool bResult = FLandscapeToolPaintBase<ToolTarget, FLandscapeToolStrokeFlatten<ToolTarget>>::MouseMove(ViewportClient, Viewport, x, y);
 
@@ -639,7 +639,7 @@ public:
 	:	FLandscapeToolStrokePaintBase<ToolTarget>(InEdMode, InTarget)
 	{}
 
-	virtual void Apply(FLevelEditorViewportClient* ViewportClient, FLandscapeBrush* Brush, const ULandscapeEditorObject* UISettings, const TArray<FLandscapeToolMousePosition>& MousePositions) OVERRIDE
+	virtual void Apply(FLevelEditorViewportClient* ViewportClient, FLandscapeBrush* Brush, const ULandscapeEditorObject* UISettings, const TArray<FLandscapeToolMousePosition>& MousePositions) override
 	{
 		if (!this->LandscapeInfo) return;
 
@@ -726,8 +726,8 @@ public:
 	:	FLandscapeToolPaintBase<ToolTarget, FLandscapeToolStrokeNoise<ToolTarget> >(InEdMode)
 	{}
 
-	virtual const TCHAR* GetToolName() OVERRIDE { return TEXT("Noise"); }
-	virtual FText GetDisplayName() OVERRIDE { return NSLOCTEXT("UnrealEd", "LandscapeMode_Noise", "Noise"); };
+	virtual const TCHAR* GetToolName() override { return TEXT("Noise"); }
+	virtual FText GetDisplayName() override { return NSLOCTEXT("UnrealEd", "LandscapeMode_Noise", "Noise"); };
 };
 
 
diff --git a/Engine/Source/Editor/LandscapeEditor/Private/LandscapeEdModeRampTool.cpp b/Engine/Source/Editor/LandscapeEditor/Private/LandscapeEdModeRampTool.cpp
--- a/Engine/Source/Editor/LandscapeEditor/Private/LandscapeEdModeRampTool.cpp
+++ b/Engine/Source/Editor/LandscapeEditor/Private/LandscapeEdModeRampTool.cpp
@@ -99,25 +99,25 @@ public:
 		check(SpriteTexture);
 	}
 
-	virtual const TCHAR* GetToolName() OVERRIDE { return TEXT("Ramp"); }
-	virtual FText GetDisplayName() OVERRIDE { return NSLOCTEXT("UnrealEd", "LandscapeMode_Ramp", "Ramp"); };
+	virtual const TCHAR* GetToolName() override { return TEXT("Ramp"); }
+	virtual FText GetDisplayName() override { return NSLOCTEXT("UnrealEd", "LandscapeMode_Ramp", "Ramp"); };
 
-	virtual void SetEditRenderType() OVERRIDE { GLandscapeEditRenderMode = ELandscapeEditRenderMode::None | (GLandscapeEditRenderMode & ELandscapeEditRenderMode::BitMaskForMask); }
-	virtual bool SupportsMask() OVERRIDE { return false; }
+	virtual void SetEditRenderType() override { GLandscapeEditRenderMode = ELandscapeEditRenderMode::None | (GLandscapeEditRenderMode & ELandscapeEditRenderMode::BitMaskForMask); }
+	virtual bool SupportsMask() override { return false; }
 
-	virtual bool IsValidForTarget(const FLandscapeToolTarget& Target) OVERRIDE
+	virtual bool IsValidForTarget(const FLandscapeToolTarget& Target) override
 	{
 		return Target.TargetType == ELandscapeToolTargetType::Heightmap;
 	}
 
-	virtual void EnterTool() OVERRIDE
+	virtual void EnterTool() override
 	{
 		NumPoints = 0;
 		SelectedPoint = INDEX_NONE;
 		GEditorModeTools().SetWidgetMode(FWidget::WM_Translate);
 	}
 
-	virtual bool BeginTool(FLevelEditorViewportClient* ViewportClient, const FLandscapeToolTarget& Target, const FVector& InHitLocation) OVERRIDE
+	virtual bool BeginTool(FLevelEditorViewportClient* ViewportClient, const FLandscapeToolTarget& Target, const FVector& InHitLocation) override
 	{
 		if (NumPoints < 2)
 		{
@@ -142,12 +142,12 @@ public:
 		return true;
 	}
 
-	virtual void EndTool(FLevelEditorViewportClient* ViewportClient) OVERRIDE
+	virtual void EndTool(FLevelEditorViewportClient* ViewportClient) override
 	{
 		bMovingPoint = false;
 	}
 
-	virtual bool MouseMove(FLevelEditorViewportClient* ViewportClient, FViewport* Viewport, int32 x, int32 y) OVERRIDE
+	virtual bool MouseMove(FLevelEditorViewportClient* ViewportClient, FViewport* Viewport, int32 x, int32 y) override
 	{
 		if (bMovingPoint)
 		{
@@ -167,7 +167,7 @@ public:
 		return true;
 	}
 
-	virtual bool HandleClick(HHitProxy* HitProxy, const FViewportClick& Click) OVERRIDE
+	virtual bool HandleClick(HHitProxy* HitProxy, const FViewportClick& Click) override
 	{
 		if (HitProxy)
 		{
@@ -185,7 +185,7 @@ public:
 		return false;
 	}
 
-	virtual bool InputKey(FLevelEditorViewportClient* InViewportClient, FViewport* InViewport, FKey InKey, EInputEvent InEvent) OVERRIDE
+	virtual bool InputKey(FLevelEditorViewportClient* InViewportClient, FViewport* InViewport, FKey InKey, EInputEvent InEvent) override
 	{
 		if (InKey == EKeys::Enter && InEvent == IE_Pressed)
 		{
@@ -267,7 +267,7 @@ public:
 		return false;
 	}
 
-	virtual bool InputDelta(FLevelEditorViewportClient* InViewportClient, FViewport* InViewport, FVector& InDrag, FRotator& InRot, FVector& InScale) OVERRIDE
+	virtual bool InputDelta(FLevelEditorViewportClient* InViewportClient, FViewport* InViewport, FVector& InDrag, FRotator& InRot, FVector& InScale) override
 	{
 		if (SelectedPoint != INDEX_NONE && InViewportClient->GetCurrentWidgetAxis() != EAxisList::None)
 		{
@@ -282,7 +282,7 @@ public:
 		return false;
 	}
 
-	virtual void Render(const FSceneView* View, FViewport* Viewport, FPrimitiveDrawInterface* PDI) OVERRIDE
+	virtual void Render(const FSceneView* View, FViewport* Viewport, FPrimitiveDrawInterface* PDI) override
 	{
 		if (NumPoints > 0)
 		{
@@ -357,12 +357,12 @@ public:
 		}
 	}
 
-	virtual bool OverrideSelection() const OVERRIDE
+	virtual bool OverrideSelection() const override
 	{
 		return true;
 	}
 
-	virtual bool IsSelectionAllowed( AActor* InActor, bool bInSelection ) const OVERRIDE
+	virtual bool IsSelectionAllowed( AActor* InActor, bool bInSelection ) const override
 	{
 		// Only filter selection not deselection
 		if (bInSelection)
@@ -373,7 +373,7 @@ public:
 		return true;
 	}
 
-	virtual bool UsesTransformWidget() const OVERRIDE
+	virtual bool UsesTransformWidget() const override
 	{
 		if (SelectedPoint != INDEX_NONE)
 		{
@@ -383,7 +383,7 @@ public:
 		return false;
 	}
 
-	virtual EAxisList::Type GetWidgetAxisToDraw(FWidget::EWidgetMode CheckMode) const OVERRIDE
+	virtual EAxisList::Type GetWidgetAxisToDraw(FWidget::EWidgetMode CheckMode) const override
 	{
 		if (SelectedPoint != INDEX_NONE)
 		{
@@ -400,7 +400,7 @@ public:
 		return EAxisList::None;
 	}
 
-	virtual FVector GetWidgetLocation() const OVERRIDE
+	virtual FVector GetWidgetLocation() const override
 	{
 		if (SelectedPoint != INDEX_NONE)
 		{
@@ -413,7 +413,7 @@ public:
 		return FVector::ZeroVector;
 	}
 
-	virtual FMatrix GetWidgetRotation() const OVERRIDE
+	virtual FMatrix GetWidgetRotation() const override
 	{
 		if (SelectedPoint != INDEX_NONE)
 		{
