Fixed ocasional crash when painting foliage on a landscape

diff --git a/Engine/Source/Editor/LandscapeEditor/Private/LandscapeEdModeTools.h b/Engine/Source/Editor/LandscapeEditor/Private/LandscapeEdModeTools.h
--- a/Engine/Source/Editor/LandscapeEditor/Private/LandscapeEdModeTools.h
+++ b/Engine/Source/Editor/LandscapeEditor/Private/LandscapeEdModeTools.h
@@ -589,7 +589,10 @@ struct FHeightmapAccessor
 				{
 					ULandscapeHeightfieldCollisionComponent* CollisionComponent = CollisionComponents[Index];
 					AInstancedFoliageActor* IFA = AInstancedFoliageActor::GetInstancedFoliageActorForLevel(CollisionComponent->GetComponentLevel());
-					IFA->SnapInstancesForLandscape(CollisionComponent, PreUpdateLocalBoxes[Index].TransformBy(LandscapeInfo->GetLandscapeProxy()->LandscapeActorToWorld().ToMatrixWithScale()).ExpandBy(1.f));
+					if (IFA)
+					{
+						IFA->SnapInstancesForLandscape(CollisionComponent, PreUpdateLocalBoxes[Index].TransformBy(LandscapeInfo->GetLandscapeProxy()->LandscapeActorToWorld().ToMatrixWithScale()).ExpandBy(1.f));
+					}
 				}
 			}
 			else
