Fix the landscape dropdown in the landscape editor not updating to match the active landscape

diff --git a/Engine/Source/Editor/LandscapeEditor/Private/LandscapeEditorDetails.cpp b/Engine/Source/Editor/LandscapeEditor/Private/LandscapeEditorDetails.cpp
--- a/Engine/Source/Editor/LandscapeEditor/Private/LandscapeEditorDetails.cpp
+++ b/Engine/Source/Editor/LandscapeEditor/Private/LandscapeEditorDetails.cpp
@@ -83,7 +83,7 @@ void FLandscapeEditorDetails::CustomizeDetails(IDetailLayoutBuilder& DetailBuild
 		.ButtonContent()
 		[
 			SNew(STextBlock)
-			.Text(GetTargetLandscapeName())
+			.Text_Static(&FLandscapeEditorDetails::GetTargetLandscapeName)
 		]
 	];
 
