Landscape : Fixed for deleting/adding component problem, where deleted component or neighbors have LODBias/ForcedLOD case

diff --git a/Engine/Source/Editor/LandscapeEditor/Private/LandscapeEdModeComponentTools.cpp b/Engine/Source/Editor/LandscapeEditor/Private/LandscapeEdModeComponentTools.cpp
--- a/Engine/Source/Editor/LandscapeEditor/Private/LandscapeEdModeComponentTools.cpp
+++ b/Engine/Source/Editor/LandscapeEditor/Private/LandscapeEdModeComponentTools.cpp
@@ -834,6 +834,31 @@ public:
 
 			for ( int32 Idx = 0; Idx < NewComponents.Num(); Idx++ )
 			{
+				// Update LODBias from neighbors
+				FIntPoint ComponentBase = NewComponents[Idx]->GetSectionBase() / NewComponents[Idx]->ComponentSizeQuads;
+				FIntPoint LandscapeKey[8] =
+				{
+					ComponentBase + FIntPoint(-1, -1),
+					ComponentBase + FIntPoint(+0, -1),
+					ComponentBase + FIntPoint(+1, -1),
+					ComponentBase + FIntPoint(-1, +0),
+					ComponentBase + FIntPoint(+1, +0),
+					ComponentBase + FIntPoint(-1, +1),
+					ComponentBase + FIntPoint(+0, +1),
+					ComponentBase + FIntPoint(+1, +1)
+				};
+
+				for (int32 NeighborIdx = 0; NeighborIdx < 8; ++NeighborIdx)
+				{
+					ULandscapeComponent* Comp = LandscapeInfo->XYtoComponentMap.FindRef(LandscapeKey[NeighborIdx]);
+					if (Comp)
+					{
+						NewComponents[Idx]->NeighborLOD[NeighborIdx] = Comp->ForcedLOD >= 0 ? Comp->ForcedLOD : 255;
+						NewComponents[Idx]->NeighborLODBias[NeighborIdx] = Comp->LODBias + 128;
+					}
+				}
+
+				// Update Collision
 				NewComponents[Idx]->UpdateCachedBounds();
 				NewComponents[Idx]->UpdateBounds();
 				NewComponents[Idx]->MarkRenderStateDirty();
@@ -954,6 +979,34 @@ public:
 				Proxy->Modify();
 				//Comp->Modify();
 
+				// Reset neighbors LOD information
+				FIntPoint ComponentBase = Comp->GetSectionBase() / Comp->ComponentSizeQuads;
+				FIntPoint LandscapeKey[8] =
+				{
+					ComponentBase + FIntPoint(-1, -1),
+					ComponentBase + FIntPoint(+0, -1),
+					ComponentBase + FIntPoint(+1, -1),
+					ComponentBase + FIntPoint(-1, +0),
+					ComponentBase + FIntPoint(+1, +0),
+					ComponentBase + FIntPoint(-1, +1),
+					ComponentBase + FIntPoint(+0, +1),
+					ComponentBase + FIntPoint(+1, +1)
+				};
+
+				for (int32 Idx = 0; Idx < 8; ++Idx)
+				{
+					ULandscapeComponent* NeighborComp = LandscapeInfo->XYtoComponentMap.FindRef(LandscapeKey[Idx]);
+					if (NeighborComp)
+					{
+						NeighborComp->Modify();
+						NeighborComp->NeighborLOD[7 - Idx] = 255; // Use 255 as unspecified value
+						NeighborComp->NeighborLODBias[7 - Idx] = 128;
+
+						NeighborComp->InvalidateLightingCache();
+						FComponentReregisterContext ReregisterContext(NeighborComp);
+					}
+				}
+
 				// Remove Selected Region in deleted Component
 				for (int32 Y = 0; Y < Comp->ComponentSizeQuads; ++Y )
 				{
