Missing include file breaking non-unity.

diff --git a/Engine/Source/Editor/LandscapeEditor/Private/LandscapeEdModeComponentTools.cpp b/Engine/Source/Editor/LandscapeEditor/Private/LandscapeEdModeComponentTools.cpp
--- a/Engine/Source/Editor/LandscapeEditor/Private/LandscapeEdModeComponentTools.cpp
+++ b/Engine/Source/Editor/LandscapeEditor/Private/LandscapeEdModeComponentTools.cpp
@@ -13,6 +13,7 @@
 #include "Editor/PropertyEditor/Public/PropertyEditorModule.h"
 #include "LandscapeEdModeTools.h"
 #include "Foliage/InstancedFoliageActor.h"
+#include "ComponentReregisterContext.h"
 
 #define LOCTEXT_NAMESPACE "Landscape"
 
