Remove FLandscapeToolSet
Rename ALandscapeProxy::DataLayer to VisibilityLayer
Fix visibility tool not showing the visibility layer as its target

diff --git a/Engine/Source/Editor/EditorStyle/Private/SlateEditorStyle.cpp b/Engine/Source/Editor/EditorStyle/Private/SlateEditorStyle.cpp
--- a/Engine/Source/Editor/EditorStyle/Private/SlateEditorStyle.cpp
+++ b/Engine/Source/Editor/EditorStyle/Private/SlateEditorStyle.cpp
@@ -4540,96 +4540,97 @@ void FSlateEditorStyle::FStyle::SetupLandscapeEditorStyle()
 	// Landscape Editor
 	{
 		// Modes
-		Set( "LandscapeEditor.ManageMode", new IMAGE_BRUSH( "Icons/icon_Landscape_Mode_Manage_40x", Icon40x40 ) );
-		Set( "LandscapeEditor.SculptMode", new IMAGE_BRUSH( "Icons/icon_Landscape_Mode_Sculpt_40x", Icon40x40 ) );
-		Set( "LandscapeEditor.PaintMode", new IMAGE_BRUSH( "Icons/icon_Landscape_Mode_Paint_40x", Icon40x40 ) );
-		Set( "LandscapeEditor.ManageMode.Small", new IMAGE_BRUSH( "Icons/icon_Landscape_Mode_Manage_20x", Icon20x20 ) );
-		Set( "LandscapeEditor.SculptMode.Small", new IMAGE_BRUSH( "Icons/icon_Landscape_Mode_Sculpt_20x", Icon20x20 ) );
-		Set( "LandscapeEditor.PaintMode.Small", new IMAGE_BRUSH( "Icons/icon_Landscape_Mode_Paint_20x", Icon20x20 ) );
+		Set("LandscapeEditor.ManageMode", new IMAGE_BRUSH("Icons/icon_Landscape_Mode_Manage_40x", Icon40x40));
+		Set("LandscapeEditor.SculptMode", new IMAGE_BRUSH("Icons/icon_Landscape_Mode_Sculpt_40x", Icon40x40));
+		Set("LandscapeEditor.PaintMode",  new IMAGE_BRUSH("Icons/icon_Landscape_Mode_Paint_40x",  Icon40x40));
+		Set("LandscapeEditor.ManageMode.Small", new IMAGE_BRUSH("Icons/icon_Landscape_Mode_Manage_20x", Icon20x20));
+		Set("LandscapeEditor.SculptMode.Small", new IMAGE_BRUSH("Icons/icon_Landscape_Mode_Sculpt_20x", Icon20x20));
+		Set("LandscapeEditor.PaintMode.Small",  new IMAGE_BRUSH("Icons/icon_Landscape_Mode_Paint_20x",  Icon20x20));
 
 		// Tools
-		Set( "LandscapeEditor.NewLandscape", new IMAGE_BRUSH( "Icons/icon_Landscape_New_Landscape_40x", Icon40x40 ) );
-		Set( "LandscapeEditor.NewLandscape.Small", new IMAGE_BRUSH( "Icons/icon_Landscape_New_Landscape_20x", Icon20x20 ) );
-		Set( "LandscapeEditor.ResizeLandscape", new IMAGE_BRUSH( "Icons/icon_Landscape_Tool_Copy_40x", Icon40x40 ) );
-		Set( "LandscapeEditor.ResizeLandscape.Small", new IMAGE_BRUSH( "Icons/icon_Landscape_Tool_Copy_20x", Icon20x20 ) );
-
-		Set( "LandscapeEditor.SculptTool", new IMAGE_BRUSH( "Icons/icon_Landscape_Tool_Sculpt_40x", Icon40x40 ) );
-		Set( "LandscapeEditor.PaintTool", new IMAGE_BRUSH( "Icons/icon_Landscape_Tool_Paint_40x", Icon40x40 ) );
-		Set( "LandscapeEditor.SmoothTool", new IMAGE_BRUSH( "Icons/icon_Landscape_Tool_Smooth_40x", Icon40x40 ) );
-		Set( "LandscapeEditor.FlattenTool", new IMAGE_BRUSH( "Icons/icon_Landscape_Tool_Flatten_40x", Icon40x40 ) );
-		Set( "LandscapeEditor.RampTool", new IMAGE_BRUSH( "Icons/icon_Landscape_Tool_Ramp_40x", Icon40x40 ) );
-		Set( "LandscapeEditor.ErosionTool", new IMAGE_BRUSH( "Icons/icon_Landscape_Tool_Erosion_40x", Icon40x40 ) );
-		Set( "LandscapeEditor.HydroErosionTool", new IMAGE_BRUSH( "Icons/icon_Landscape_Tool_HydraulicErosion_40x", Icon40x40 ) );
-		Set( "LandscapeEditor.NoiseTool", new IMAGE_BRUSH( "Icons/icon_Landscape_Tool_Noise_40x", Icon40x40 ) );
-		Set( "LandscapeEditor.RetopologizeTool", new IMAGE_BRUSH( "Icons/icon_Landscape_Tool_Retopologize_40x", Icon40x40 ) );
-		Set( "LandscapeEditor.VisibilityTool", new IMAGE_BRUSH( "Icons/icon_Landscape_Tool_Visibility_40x", Icon40x40 ) );
-		Set( "LandscapeEditor.SculptTool.Small", new IMAGE_BRUSH( "Icons/icon_Landscape_Tool_Sculpt_20x", Icon20x20 ) );
-		Set( "LandscapeEditor.PaintTool.Small", new IMAGE_BRUSH( "Icons/icon_Landscape_Tool_Paint_20x", Icon20x20 ) );
-		Set( "LandscapeEditor.SmoothTool.Small", new IMAGE_BRUSH( "Icons/icon_Landscape_Tool_Smooth_20x", Icon20x20 ) );
-		Set( "LandscapeEditor.FlattenTool.Small", new IMAGE_BRUSH( "Icons/icon_Landscape_Tool_Flatten_20x", Icon20x20 ) );
-		Set( "LandscapeEditor.RampTool.Small", new IMAGE_BRUSH( "Icons/icon_Landscape_Tool_Ramp_20x", Icon20x20 ) );
-		Set( "LandscapeEditor.ErosionTool.Small", new IMAGE_BRUSH( "Icons/icon_Landscape_Tool_Erosion_20x", Icon20x20 ) );
-		Set( "LandscapeEditor.HydroErosionTool.Small", new IMAGE_BRUSH( "Icons/icon_Landscape_Tool_HydraulicErosion_20x", Icon20x20 ) );
-		Set( "LandscapeEditor.NoiseTool.Small", new IMAGE_BRUSH( "Icons/icon_Landscape_Tool_Noise_20x", Icon20x20 ) );
-		Set( "LandscapeEditor.RetopologizeTool.Small", new IMAGE_BRUSH( "Icons/icon_Landscape_Tool_Retopologize_20x", Icon20x20 ) );
-		Set( "LandscapeEditor.VisibilityTool.Small", new IMAGE_BRUSH( "Icons/icon_Landscape_Tool_Visibility_20x", Icon20x20 ) );
-
-		Set( "LandscapeEditor.SelectComponentTool", new IMAGE_BRUSH( "Icons/icon_Landscape_Tool_Selection_40x", Icon40x40 ) );
-		Set( "LandscapeEditor.AddComponentTool", new IMAGE_BRUSH( "Icons/icon_Landscape_Tool_AddComponent_40x", Icon40x40 ) );
-		Set( "LandscapeEditor.DeleteComponentTool", new IMAGE_BRUSH( "Icons/icon_Landscape_Tool_DeleteComponent_40x", Icon40x40 ) );
-		Set( "LandscapeEditor.MoveToLevelTool", new IMAGE_BRUSH( "Icons/icon_Landscape_Tool_MoveToLevel_40x", Icon40x40 ) );
-		Set( "LandscapeEditor.SelectComponentTool.Small", new IMAGE_BRUSH( "Icons/icon_Landscape_Tool_Selection_20x", Icon20x20 ) );
-		Set( "LandscapeEditor.AddComponentTool.Small", new IMAGE_BRUSH( "Icons/icon_Landscape_Tool_AddComponent_20x", Icon20x20 ) );
-		Set( "LandscapeEditor.DeleteComponentTool.Small", new IMAGE_BRUSH( "Icons/icon_Landscape_Tool_DeleteComponent_20x", Icon20x20 ) );
-		Set( "LandscapeEditor.MoveToLevelTool.Small", new IMAGE_BRUSH( "Icons/icon_Landscape_Tool_MoveToLevel_20x", Icon20x20 ) );
-
-		Set( "LandscapeEditor.RegionSelectTool", new IMAGE_BRUSH( "Icons/icon_Landscape_Tool_Mask_40x", Icon40x40 ) );
-		Set( "LandscapeEditor.RegionCopyPasteTool", new IMAGE_BRUSH( "Icons/icon_Landscape_Tool_Copy_40x", Icon40x40 ) );
-		Set( "LandscapeEditor.RegionSelectTool.Small", new IMAGE_BRUSH( "Icons/icon_Landscape_Tool_Mask_20x", Icon20x20 ) );
-		Set( "LandscapeEditor.RegionCopyPasteTool.Small", new IMAGE_BRUSH( "Icons/icon_Landscape_Tool_Copy_20x", Icon20x20 ) );
-
-		Set( "LandscapeEditor.SplineTool", new IMAGE_BRUSH( "Icons/icon_Landscape_Tool_Splines_40x", Icon40x40 ) );
-		Set( "LandscapeEditor.SplineTool.Small", new IMAGE_BRUSH( "Icons/icon_Landscape_Tool_Splines_20x", Icon20x20 ) );
+		Set("LandscapeEditor.NewLandscape",       new IMAGE_BRUSH("Icons/icon_Landscape_New_Landscape_40x", Icon40x40));
+		Set("LandscapeEditor.NewLandscape.Small", new IMAGE_BRUSH("Icons/icon_Landscape_New_Landscape_20x", Icon20x20));
+		Set("LandscapeEditor.ResizeLandscape",       new IMAGE_BRUSH("Icons/icon_Landscape_Tool_Copy_40x", Icon40x40));
+		Set("LandscapeEditor.ResizeLandscape.Small", new IMAGE_BRUSH("Icons/icon_Landscape_Tool_Copy_20x", Icon20x20));
+
+		Set("LandscapeEditor.SculptTool",       new IMAGE_BRUSH("Icons/icon_Landscape_Tool_Sculpt_40x",           Icon40x40));
+		Set("LandscapeEditor.PaintTool",        new IMAGE_BRUSH("Icons/icon_Landscape_Tool_Paint_40x",            Icon40x40));
+		Set("LandscapeEditor.SmoothTool",       new IMAGE_BRUSH("Icons/icon_Landscape_Tool_Smooth_40x",           Icon40x40));
+		Set("LandscapeEditor.FlattenTool",      new IMAGE_BRUSH("Icons/icon_Landscape_Tool_Flatten_40x",          Icon40x40));
+		Set("LandscapeEditor.RampTool",         new IMAGE_BRUSH("Icons/icon_Landscape_Tool_Ramp_40x",             Icon40x40));
+		Set("LandscapeEditor.ErosionTool",      new IMAGE_BRUSH("Icons/icon_Landscape_Tool_Erosion_40x",          Icon40x40));
+		Set("LandscapeEditor.HydroErosionTool", new IMAGE_BRUSH("Icons/icon_Landscape_Tool_HydraulicErosion_40x", Icon40x40));
+		Set("LandscapeEditor.NoiseTool",        new IMAGE_BRUSH("Icons/icon_Landscape_Tool_Noise_40x",            Icon40x40));
+		Set("LandscapeEditor.RetopologizeTool", new IMAGE_BRUSH("Icons/icon_Landscape_Tool_Retopologize_40x",     Icon40x40));
+		Set("LandscapeEditor.VisibilityTool",   new IMAGE_BRUSH("Icons/icon_Landscape_Tool_Visibility_40x",       Icon40x40));
+		Set("LandscapeEditor.SculptTool.Small",       new IMAGE_BRUSH("Icons/icon_Landscape_Tool_Sculpt_20x",           Icon20x20));
+		Set("LandscapeEditor.PaintTool.Small",        new IMAGE_BRUSH("Icons/icon_Landscape_Tool_Paint_20x",            Icon20x20));
+		Set("LandscapeEditor.SmoothTool.Small",       new IMAGE_BRUSH("Icons/icon_Landscape_Tool_Smooth_20x",           Icon20x20));
+		Set("LandscapeEditor.FlattenTool.Small",      new IMAGE_BRUSH("Icons/icon_Landscape_Tool_Flatten_20x",          Icon20x20));
+		Set("LandscapeEditor.RampTool.Small",         new IMAGE_BRUSH("Icons/icon_Landscape_Tool_Ramp_20x",             Icon20x20));
+		Set("LandscapeEditor.ErosionTool.Small",      new IMAGE_BRUSH("Icons/icon_Landscape_Tool_Erosion_20x",          Icon20x20));
+		Set("LandscapeEditor.HydroErosionTool.Small", new IMAGE_BRUSH("Icons/icon_Landscape_Tool_HydraulicErosion_20x", Icon20x20));
+		Set("LandscapeEditor.NoiseTool.Small",        new IMAGE_BRUSH("Icons/icon_Landscape_Tool_Noise_20x",            Icon20x20));
+		Set("LandscapeEditor.RetopologizeTool.Small", new IMAGE_BRUSH("Icons/icon_Landscape_Tool_Retopologize_20x",     Icon20x20));
+		Set("LandscapeEditor.VisibilityTool.Small",   new IMAGE_BRUSH("Icons/icon_Landscape_Tool_Visibility_20x",       Icon20x20));
+
+		Set("LandscapeEditor.SelectComponentTool", new IMAGE_BRUSH("Icons/icon_Landscape_Tool_Selection_40x",       Icon40x40));
+		Set("LandscapeEditor.AddComponentTool",    new IMAGE_BRUSH("Icons/icon_Landscape_Tool_AddComponent_40x",    Icon40x40));
+		Set("LandscapeEditor.DeleteComponentTool", new IMAGE_BRUSH("Icons/icon_Landscape_Tool_DeleteComponent_40x", Icon40x40));
+		Set("LandscapeEditor.MoveToLevelTool",     new IMAGE_BRUSH("Icons/icon_Landscape_Tool_MoveToLevel_40x",     Icon40x40));
+		Set("LandscapeEditor.SelectComponentTool.Small", new IMAGE_BRUSH("Icons/icon_Landscape_Tool_Selection_20x",       Icon20x20));
+		Set("LandscapeEditor.AddComponentTool.Small",    new IMAGE_BRUSH("Icons/icon_Landscape_Tool_AddComponent_20x",    Icon20x20));
+		Set("LandscapeEditor.DeleteComponentTool.Small", new IMAGE_BRUSH("Icons/icon_Landscape_Tool_DeleteComponent_20x", Icon20x20));
+		Set("LandscapeEditor.MoveToLevelTool.Small",     new IMAGE_BRUSH("Icons/icon_Landscape_Tool_MoveToLevel_20x",     Icon20x20));
+
+		Set("LandscapeEditor.RegionSelectTool",    new IMAGE_BRUSH("Icons/icon_Landscape_Tool_Mask_40x", Icon40x40));
+		Set("LandscapeEditor.RegionCopyPasteTool", new IMAGE_BRUSH("Icons/icon_Landscape_Tool_Copy_40x", Icon40x40));
+		Set("LandscapeEditor.RegionSelectTool.Small",    new IMAGE_BRUSH("Icons/icon_Landscape_Tool_Mask_20x", Icon20x20));
+		Set("LandscapeEditor.RegionCopyPasteTool.Small", new IMAGE_BRUSH("Icons/icon_Landscape_Tool_Copy_20x", Icon20x20));
+
+		Set("LandscapeEditor.SplineTool", new IMAGE_BRUSH("Icons/icon_Landscape_Tool_Splines_40x", Icon40x40));
+		Set("LandscapeEditor.SplineTool.Small", new IMAGE_BRUSH("Icons/icon_Landscape_Tool_Splines_20x", Icon20x20));
 
 		// Brush Sets
-		Set( "LandscapeEditor.CircleBrush", new IMAGE_BRUSH( "Icons/icon_Landscape_Brush_Circle_smooth_40x", Icon40x40 ) );
-		Set( "LandscapeEditor.AlphaBrush", new IMAGE_BRUSH( "Icons/icon_Landscape_Brush_Alpha_40x", Icon40x40 ) );
-		Set( "LandscapeEditor.AlphaBrush_Pattern", new IMAGE_BRUSH( "Icons/icon_Landscape_Brush_Pattern_40x", Icon40x40 ) );
-		Set( "LandscapeEditor.ComponentBrush", new IMAGE_BRUSH( "Icons/icon_Landscape_Brush_Component_40x", Icon40x40 ) );
-		Set( "LandscapeEditor.GizmoBrush", new IMAGE_BRUSH( "Icons/icon_Landscape_Brush_Gizmo_40x", Icon40x40 ) );
-		Set( "LandscapeEditor.CircleBrush.Small", new IMAGE_BRUSH( "Icons/icon_Landscape_Brush_Circle_smooth_20x", Icon20x20 ) );
-		Set( "LandscapeEditor.AlphaBrush.Small", new IMAGE_BRUSH( "Icons/icon_Landscape_Brush_Alpha_20x", Icon20x20 ) );
-		Set( "LandscapeEditor.AlphaBrush_Pattern.Small", new IMAGE_BRUSH( "Icons/icon_Landscape_Brush_Pattern_20x", Icon20x20 ) );
-		Set( "LandscapeEditor.ComponentBrush.Small", new IMAGE_BRUSH( "Icons/icon_Landscape_Brush_Component_20x", Icon20x20 ) );
-		Set( "LandscapeEditor.GizmoBrush.Small", new IMAGE_BRUSH( "Icons/icon_Landscape_Brush_Gizmo_20x", Icon20x20 ) );
+		Set("LandscapeEditor.CircleBrush",        new IMAGE_BRUSH("Icons/icon_Landscape_Brush_Circle_smooth_40x", Icon40x40));
+		Set("LandscapeEditor.AlphaBrush",         new IMAGE_BRUSH("Icons/icon_Landscape_Brush_Alpha_40x",         Icon40x40));
+		Set("LandscapeEditor.AlphaBrush_Pattern", new IMAGE_BRUSH("Icons/icon_Landscape_Brush_Pattern_40x",       Icon40x40));
+		Set("LandscapeEditor.ComponentBrush",     new IMAGE_BRUSH("Icons/icon_Landscape_Brush_Component_40x",     Icon40x40));
+		Set("LandscapeEditor.GizmoBrush",         new IMAGE_BRUSH("Icons/icon_Landscape_Brush_Gizmo_40x",         Icon40x40));
+		Set("LandscapeEditor.CircleBrush.Small",        new IMAGE_BRUSH("Icons/icon_Landscape_Brush_Circle_smooth_20x", Icon20x20));
+		Set("LandscapeEditor.AlphaBrush.Small",         new IMAGE_BRUSH("Icons/icon_Landscape_Brush_Alpha_20x",         Icon20x20));
+		Set("LandscapeEditor.AlphaBrush_Pattern.Small", new IMAGE_BRUSH("Icons/icon_Landscape_Brush_Pattern_20x",       Icon20x20));
+		Set("LandscapeEditor.ComponentBrush.Small",     new IMAGE_BRUSH("Icons/icon_Landscape_Brush_Component_20x",     Icon20x20));
+		Set("LandscapeEditor.GizmoBrush.Small",         new IMAGE_BRUSH("Icons/icon_Landscape_Brush_Gizmo_20x",         Icon20x20));
 
 		// Brushes
-		Set( "LandscapeEditor.CircleBrush_Smooth", new IMAGE_BRUSH( "Icons/icon_Landscape_Brush_Circle_smooth_40x", Icon40x40 ) );
-		Set( "LandscapeEditor.CircleBrush_Linear", new IMAGE_BRUSH( "Icons/icon_Landscape_Brush_Circle_linear_40x", Icon40x40 ) );
-		Set( "LandscapeEditor.CircleBrush_Spherical", new IMAGE_BRUSH( "Icons/icon_Landscape_Brush_Circle_spherical_40x", Icon40x40 ) );
-		Set( "LandscapeEditor.CircleBrush_Tip", new IMAGE_BRUSH( "Icons/icon_Landscape_Brush_Circle_tip_40x", Icon40x40 ) );
-		Set( "LandscapeEditor.CircleBrush_Smooth.Small", new IMAGE_BRUSH( "Icons/icon_Landscape_Brush_Circle_smooth_20x", Icon20x20 ) );
-		Set( "LandscapeEditor.CircleBrush_Linear.Small", new IMAGE_BRUSH( "Icons/icon_Landscape_Brush_Circle_linear_20x", Icon20x20 ) );
-		Set( "LandscapeEditor.CircleBrush_Spherical.Small", new IMAGE_BRUSH( "Icons/icon_Landscape_Brush_Circle_spherical_20x", Icon20x20 ) );
-		Set( "LandscapeEditor.CircleBrush_Tip.Small", new IMAGE_BRUSH( "Icons/icon_Landscape_Brush_Circle_tip_20x", Icon20x20 ) );
-
-		Set( "LandscapeEditor.Brushes.Alpha.UseRChannel", new IMAGE_BRUSH( "Icons/icon_Landscape_Brush_Alpha_UseRChannel_20x", Icon20x20 ) );
-		Set( "LandscapeEditor.Brushes.Alpha.UseGChannel", new IMAGE_BRUSH( "Icons/icon_Landscape_Brush_Alpha_UseGChannel_20x", Icon20x20 ) );
-		Set( "LandscapeEditor.Brushes.Alpha.UseBChannel", new IMAGE_BRUSH( "Icons/icon_Landscape_Brush_Alpha_UseBChannel_20x", Icon20x20 ) );
-		Set( "LandscapeEditor.Brushes.Alpha.UseAChannel", new IMAGE_BRUSH( "Icons/icon_Landscape_Brush_Alpha_UseAChannel_20x", Icon20x20 ) );
+		Set("LandscapeEditor.CircleBrush_Smooth",    new IMAGE_BRUSH("Icons/icon_Landscape_Brush_Circle_smooth_40x",    Icon40x40));
+		Set("LandscapeEditor.CircleBrush_Linear",    new IMAGE_BRUSH("Icons/icon_Landscape_Brush_Circle_linear_40x",    Icon40x40));
+		Set("LandscapeEditor.CircleBrush_Spherical", new IMAGE_BRUSH("Icons/icon_Landscape_Brush_Circle_spherical_40x", Icon40x40));
+		Set("LandscapeEditor.CircleBrush_Tip",       new IMAGE_BRUSH("Icons/icon_Landscape_Brush_Circle_tip_40x",       Icon40x40));
+		Set("LandscapeEditor.CircleBrush_Smooth.Small",    new IMAGE_BRUSH("Icons/icon_Landscape_Brush_Circle_smooth_20x",    Icon20x20));
+		Set("LandscapeEditor.CircleBrush_Linear.Small",    new IMAGE_BRUSH("Icons/icon_Landscape_Brush_Circle_linear_20x",    Icon20x20));
+		Set("LandscapeEditor.CircleBrush_Spherical.Small", new IMAGE_BRUSH("Icons/icon_Landscape_Brush_Circle_spherical_20x", Icon20x20));
+		Set("LandscapeEditor.CircleBrush_Tip.Small",       new IMAGE_BRUSH("Icons/icon_Landscape_Brush_Circle_tip_20x",       Icon20x20));
+
+		Set("LandscapeEditor.Brushes.Alpha.UseRChannel", new IMAGE_BRUSH("Icons/icon_Landscape_Brush_Alpha_UseRChannel_20x", Icon20x20));
+		Set("LandscapeEditor.Brushes.Alpha.UseGChannel", new IMAGE_BRUSH("Icons/icon_Landscape_Brush_Alpha_UseGChannel_20x", Icon20x20));
+		Set("LandscapeEditor.Brushes.Alpha.UseBChannel", new IMAGE_BRUSH("Icons/icon_Landscape_Brush_Alpha_UseBChannel_20x", Icon20x20));
+		Set("LandscapeEditor.Brushes.Alpha.UseAChannel", new IMAGE_BRUSH("Icons/icon_Landscape_Brush_Alpha_UseAChannel_20x", Icon20x20));
 
 		// Target List
-		Set( "LandscapeEditor.TargetList.RowBackground",        new FSlateNoResource() );
-		Set( "LandscapeEditor.TargetList.RowBackgroundHovered", new BOX_BRUSH( "Common/RoundedSelection_16x", 4.0f/16.0f, FLinearColor(1.0f, 1.0f, 1.0f, 0.1f) ) );
-		Set( "LandscapeEditor.TargetList.RowSelected",          new BOX_BRUSH( "Common/RoundedSelection_16x", 4.0f/16.0f, SelectionColor_Pressed ) );
-		Set( "LandscapeEditor.TargetList.RowSelectedHovered",   new BOX_BRUSH( "Common/RoundedSelection_16x", 4.0f/16.0f, SelectionColor ) );
+		Set("LandscapeEditor.TargetList.RowBackground",        new FSlateNoResource());
+		Set("LandscapeEditor.TargetList.RowBackgroundHovered", new BOX_BRUSH("Common/RoundedSelection_16x", 4.0f/16.0f, FLinearColor(1.0f, 1.0f, 1.0f, 0.1f)));
+		Set("LandscapeEditor.TargetList.RowSelected",          new BOX_BRUSH("Common/RoundedSelection_16x", 4.0f/16.0f, SelectionColor_Pressed));
+		Set("LandscapeEditor.TargetList.RowSelectedHovered",   new BOX_BRUSH("Common/RoundedSelection_16x", 4.0f/16.0f, SelectionColor));
 		
-		Set( "LandscapeEditor.Target_Heightmap", new IMAGE_BRUSH( "Icons/icon_Landscape_Target_Heightmap_32x", Icon48x48 ) );
-		Set( "LandscapeEditor.Target_Invalid", new IMAGE_BRUSH( "Icons/icon_Landscape_Target_Invalid_32x", Icon48x48 ) );
+		Set("LandscapeEditor.Target_Heightmap",  new IMAGE_BRUSH("Icons/icon_Landscape_Target_Heightmap_48x",  Icon48x48));
+		Set("LandscapeEditor.Target_Visibility", new IMAGE_BRUSH("Icons/icon_Landscape_Target_Visibility_48x", Icon48x48));
+		Set("LandscapeEditor.Target_Invalid",    new IMAGE_BRUSH("Icons/icon_Landscape_Target_Invalid_48x",    Icon48x48));
 
-		Set( "LandscapeEditor.Target_Create", new IMAGE_BRUSH( "Icons/icon_Landscape_Target_Create_12x", Icon12x12 ) );
-		Set( "LandscapeEditor.Target_MakePublic", new IMAGE_BRUSH( "Icons/assign_right_12x", Icon12x12 ) );
-		Set( "LandscapeEditor.Target_Delete", new IMAGE_BRUSH( "Icons/Cross_12x", Icon12x12 ) );
+		Set("LandscapeEditor.Target_Create",     new IMAGE_BRUSH("Icons/icon_Landscape_Target_Create_12x", Icon12x12));
+		Set("LandscapeEditor.Target_MakePublic", new IMAGE_BRUSH("Icons/assign_right_12x",                 Icon12x12));
+		Set("LandscapeEditor.Target_Delete",     new IMAGE_BRUSH("Icons/Cross_12x",                        Icon12x12));
 	}
 #endif
 }
diff --git a/Engine/Source/Editor/LandscapeEditor/Classes/LandscapeEditorObject.h b/Engine/Source/Editor/LandscapeEditor/Classes/LandscapeEditorObject.h
--- a/Engine/Source/Editor/LandscapeEditor/Classes/LandscapeEditorObject.h
+++ b/Engine/Source/Editor/LandscapeEditor/Classes/LandscapeEditorObject.h
@@ -195,13 +195,13 @@ class ULandscapeEditorObject : public UObject
 {
 	GENERATED_UCLASS_BODY()
 
-	class FEdModeLandscape* ParentMode;
+	FEdModeLandscape* ParentMode;
 
 
 	// Common Tool Settings:
 
 	// Strength of the tool
-	UPROPERTY(Category="Tool Settings", EditAnywhere, meta=(ShowForTools="ToolSet_Paint,ToolSet_Sculpt,ToolSet_Smooth,ToolSet_Flatten,ToolSet_Erosion,ToolSet_HydraErosion,ToolSet_Noise,ToolSet_Mask,ToolSet_CopyPaste", ClampMin="0", ClampMax="10", UIMin="0", UIMax="1"))
+	UPROPERTY(Category="Tool Settings", EditAnywhere, meta=(ShowForTools="Paint,Sculpt,Smooth,Flatten,Erosion,HydraErosion,Noise,Mask,CopyPaste", ClampMin="0", ClampMax="10", UIMin="0", UIMax="1"))
 	float ToolStrength;
 
 	// Enable to make tools blend towards a target value
@@ -209,7 +209,7 @@ class ULandscapeEditorObject : public UObject
 	bool bUseWeightTargetValue;
 
 	// Enable to make tools blend towards a target value
-	UPROPERTY(Category="Tool Settings", EditAnywhere, meta=(DisplayName="Use Target Value", editcondition="bUseWeightTargetValue", ShowForTools="ToolSet_Paint,ToolSet_Sculpt,ToolSet_Noise", ClampMin="0", ClampMax="10", UIMin="0", UIMax="1"))
+	UPROPERTY(Category="Tool Settings", EditAnywhere, meta=(DisplayName="Use Target Value", editcondition="bUseWeightTargetValue", ShowForTools="Paint,Sculpt,Noise", ClampMin="0", ClampMax="10", UIMin="0", UIMax="1"))
 	float WeightTargetValue;
 
 	// I have no idea what this is for but it's used by the noise and erosion tools, and isn't exposed to the UI
@@ -219,15 +219,15 @@ class ULandscapeEditorObject : public UObject
 	// Flatten Tool:
 
 	// Whether to flatten by lowering, raising, or both
-	UPROPERTY(Category="Tool Settings", EditAnywhere, meta=(ShowForTools="ToolSet_Flatten"))
+	UPROPERTY(Category="Tool Settings", EditAnywhere, meta=(ShowForTools="Flatten"))
 	TEnumAsByte<ELandscapeToolFlattenMode::Type> FlattenMode;
 
 	// Flattens to the angle of the clicked point, instead of horizontal
-	UPROPERTY(Category="Tool Settings", EditAnywhere, meta=(ShowForTools="ToolSet_Flatten", ShowForTargetTypes="Heightmap"))
+	UPROPERTY(Category="Tool Settings", EditAnywhere, meta=(ShowForTools="Flatten", ShowForTargetTypes="Heightmap"))
 	bool bUseSlopeFlatten;
 
 	// Constantly picks new values to flatten towards when dragging around, instead of only using the first clicked point
-	UPROPERTY(Category="Tool Settings", EditAnywhere, meta=(ShowForTools="ToolSet_Flatten", ShowForTargetTypes="Heightmap"))
+	UPROPERTY(Category="Tool Settings", EditAnywhere, meta=(ShowForTools="Flatten", ShowForTargetTypes="Heightmap"))
 	bool bPickValuePerApply;
 
 	// Enable to flatten towards a target height
@@ -235,27 +235,27 @@ class ULandscapeEditorObject : public UObject
 	bool bUseFlattenTarget;
 
 	// Target height to flatten towards (in Unreal Units)
-	UPROPERTY(Category="Tool Settings", EditAnywhere, meta=(ShowForTools="ToolSet_Flatten", ShowForTargetTypes="Heightmap", editcondition="bUseFlattenTarget", UIMin="-32768", UIMax="32768"))
+	UPROPERTY(Category="Tool Settings", EditAnywhere, meta=(ShowForTools="Flatten", ShowForTargetTypes="Heightmap", editcondition="bUseFlattenTarget", UIMin="-32768", UIMax="32768"))
 	float FlattenTarget;
 
 	// Whether to show the preview grid for the flatten target height
-	UPROPERTY(Category = "Tool Settings", EditAnywhere, AdvancedDisplay, meta = (DisplayName = "Show Preview Grid", ShowForTools = "ToolSet_Flatten", ShowForTargetTypes = "Heightmap", editcondition = "bUseFlattenTarget", HideEditConditionToggle, UIMin = "-32768", UIMax = "32768"))
+	UPROPERTY(Category = "Tool Settings", EditAnywhere, AdvancedDisplay, meta = (DisplayName = "Show Preview Grid", ShowForTools = "Flatten", ShowForTargetTypes = "Heightmap", editcondition = "bUseFlattenTarget", HideEditConditionToggle, UIMin = "-32768", UIMax = "32768"))
 	bool bShowFlattenTargetPreview;
 
 	// Ramp Tool:
 
 	// Width of ramp
-	UPROPERTY(Category="Tool Settings", EditAnywhere, meta=(ShowForTools="ToolSet_Ramp", ClampMin="1", UIMin="1", UIMax="8192", SliderExponent=3))
+	UPROPERTY(Category="Tool Settings", EditAnywhere, meta=(ShowForTools="Ramp", ClampMin="1", UIMin="1", UIMax="8192", SliderExponent=3))
 	float RampWidth;
 
 	// Falloff on side of ramp
-	UPROPERTY(Category="Tool Settings", EditAnywhere, meta=(DisplayName="Side Falloff", ShowForTools="ToolSet_Ramp", ClampMin="0", ClampMax="1", UIMin="0", UIMax="1"))
+	UPROPERTY(Category="Tool Settings", EditAnywhere, meta=(DisplayName="Side Falloff", ShowForTools="Ramp", ClampMin="0", ClampMax="1", UIMin="0", UIMax="1"))
 	float RampSideFalloff;
 
 	// Smooth Tool:
 
 	// Scale multiplier for the smoothing filter kernel
-	UPROPERTY(Category="Tool Settings", EditAnywhere, meta=(DisplayName="Filter Kernel Scale", ShowForTools="ToolSet_Smooth", ClampMin="0", ClampMax="1", UIMin="0", UIMax="1"))
+	UPROPERTY(Category="Tool Settings", EditAnywhere, meta=(DisplayName="Filter Kernel Scale", ShowForTools="Smooth", ClampMin="0", ClampMax="1", UIMin="0", UIMax="1"))
 	float SmoothFilterKernelScale;
 
 	// If checked, performs a detail preserving smooth using the specified detail smoothing value
@@ -263,51 +263,51 @@ class ULandscapeEditorObject : public UObject
 	bool bDetailSmooth;
 
 	// Larger detail smoothing values remove more details, while smaller values preserve more details
-	UPROPERTY(Category="Tool Settings", EditAnywhere, meta=(DisplayName="Detail Smooth", editcondition="bDetailSmooth", ShowForTools="ToolSet_Smooth", ClampMin="0", ClampMax="0.99"))
+	UPROPERTY(Category="Tool Settings", EditAnywhere, meta=(DisplayName="Detail Smooth", editcondition="bDetailSmooth", ShowForTools="Smooth", ClampMin="0", ClampMax="0.99"))
 	float DetailScale;
 
 	// Erosion Tool:
 
 	// The minimum height difference necessary for the erosion effects to be applied. Smaller values will result in more erosion being applied
-	UPROPERTY(Category="Tool Settings", EditAnywhere, meta=(DisplayName="Threshold", ShowForTools="ToolSet_Erosion", ClampMin="0", ClampMax="256", UIMin="0", UIMax="128"))
+	UPROPERTY(Category="Tool Settings", EditAnywhere, meta=(DisplayName="Threshold", ShowForTools="Erosion", ClampMin="0", ClampMax="256", UIMin="0", UIMax="128"))
 	int32 ErodeThresh;
 
 	// The thickness of the surface for the layer weight erosion effect
-	UPROPERTY(Category="Tool Settings", EditAnywhere, meta=(DisplayName="Surface Thickness", ShowForTools="ToolSet_Erosion", ClampMin="128", ClampMax="1024", UIMin="128", UIMax="512"))
+	UPROPERTY(Category="Tool Settings", EditAnywhere, meta=(DisplayName="Surface Thickness", ShowForTools="Erosion", ClampMin="128", ClampMax="1024", UIMin="128", UIMax="512"))
 	int32 ErodeSurfaceThickness;
 
 	// Number of erosion iterations, more means more erosion but is slower
-	UPROPERTY(Category="Tool Settings", EditAnywhere, meta=(DisplayName="Iterations", ShowForTools="ToolSet_Erosion", ClampMin="1", ClampMax="300", UIMin="1", UIMax="150"))
+	UPROPERTY(Category="Tool Settings", EditAnywhere, meta=(DisplayName="Iterations", ShowForTools="Erosion", ClampMin="1", ClampMax="300", UIMin="1", UIMax="150"))
 	int32 ErodeIterationNum;
 
 	// Whether to erode by lowering, raising, or both
-	UPROPERTY(Category="Tool Settings", EditAnywhere, meta=(DisplayName="Noise Mode", ShowForTools="ToolSet_Erosion"))
+	UPROPERTY(Category="Tool Settings", EditAnywhere, meta=(DisplayName="Noise Mode", ShowForTools="Erosion"))
 	TEnumAsByte<ELandscapeToolErosionMode::Type> ErosionNoiseMode;
 
 	// The size of the perlin noise filter used
-	UPROPERTY(Category="Tool Settings", EditAnywhere, meta=(DisplayName="Noise Scale", ShowForTools="ToolSet_Erosion", ClampMin="1", ClampMax="512", UIMin="1.1", UIMax="256"))
+	UPROPERTY(Category="Tool Settings", EditAnywhere, meta=(DisplayName="Noise Scale", ShowForTools="Erosion", ClampMin="1", ClampMax="512", UIMin="1.1", UIMax="256"))
 	float ErosionNoiseScale;
 
 	// Hydraulic Erosion Tool:
 
 	// The amount of rain to apply to the surface. Larger values will result in more erosion
-	UPROPERTY(Category="Tool Settings", EditAnywhere, meta=(ShowForTools="ToolSet_HydraErosion", ClampMin="1", ClampMax="512", UIMin="1", UIMax="256"))
+	UPROPERTY(Category="Tool Settings", EditAnywhere, meta=(ShowForTools="HydraErosion", ClampMin="1", ClampMax="512", UIMin="1", UIMax="256"))
 	int32 RainAmount;
 
 	// The amount of sediment that the water can carry. Larger values will result in more erosion
-	UPROPERTY(Category="Tool Settings", EditAnywhere, meta=(DisplayName="Sediment Cap.", ShowForTools="ToolSet_HydraErosion", ClampMin="0.1", ClampMax="1.0"))
+	UPROPERTY(Category="Tool Settings", EditAnywhere, meta=(DisplayName="Sediment Cap.", ShowForTools="HydraErosion", ClampMin="0.1", ClampMax="1.0"))
 	float SedimentCapacity;
 
 	// Number of erosion iterations, more means more erosion but is slower
-	UPROPERTY(Category="Tool Settings", EditAnywhere, meta=(DisplayName="Iterations", ShowForTools="ToolSet_HydraErosion", ClampMin="1", ClampMax="300", UIMin="1", UIMax="150"))
+	UPROPERTY(Category="Tool Settings", EditAnywhere, meta=(DisplayName="Iterations", ShowForTools="HydraErosion", ClampMin="1", ClampMax="300", UIMin="1", UIMax="150"))
 	int32 HErodeIterationNum;
 
 	// Initial Rain Distribution
-	UPROPERTY(Category="Tool Settings", EditAnywhere, meta=(DisplayName="Initial Rain Distribution", ShowForTools="ToolSet_HydraErosion"))
+	UPROPERTY(Category="Tool Settings", EditAnywhere, meta=(DisplayName="Initial Rain Distribution", ShowForTools="HydraErosion"))
 	TEnumAsByte<ELandscapeToolHydroErosionMode::Type> RainDistMode;
 
 	// The size of the noise filter for applying initial rain to the surface
-	UPROPERTY(Category="Tool Settings", EditAnywhere, meta=(ShowForTools="ToolSet_HydraErosion", ClampMin="1", ClampMax="512", UIMin="1.1", UIMax="256"))
+	UPROPERTY(Category="Tool Settings", EditAnywhere, meta=(ShowForTools="HydraErosion", ClampMin="1", ClampMax="512", UIMin="1.1", UIMax="256"))
 	float RainDistScale;
 
 	// If checked, performs a detail-preserving smooth to the erosion effect using the specified detail smoothing value
@@ -315,55 +315,55 @@ class ULandscapeEditorObject : public UObject
 	bool bHErosionDetailSmooth;
 
 	// Larger detail smoothing values remove more details, while smaller values preserve more details
-	UPROPERTY(Category="Tool Settings", EditAnywhere, meta=(DisplayName="Detail Smooth", editcondition="bHErosionDetailSmooth", ShowForTools="ToolSet_HydraErosion", ClampMin="0", ClampMax="0.99"))
+	UPROPERTY(Category="Tool Settings", EditAnywhere, meta=(DisplayName="Detail Smooth", editcondition="bHErosionDetailSmooth", ShowForTools="HydraErosion", ClampMin="0", ClampMax="0.99"))
 	float HErosionDetailScale;
 
 	// Noise Tool:
 
 	// Whether to apply noise that raises, lowers, or both
-	UPROPERTY(Category="Tool Settings", EditAnywhere, meta=(DisplayName="Noise Mode", ShowForTools="ToolSet_Noise"))
+	UPROPERTY(Category="Tool Settings", EditAnywhere, meta=(DisplayName="Noise Mode", ShowForTools="Noise"))
 	TEnumAsByte<ELandscapeToolNoiseMode::Type> NoiseMode;
 
 	// The size of the perlin noise filter used
-	UPROPERTY(Category="Tool Settings", EditAnywhere, meta=(DisplayName="Noise Scale", ShowForTools="ToolSet_Noise", ClampMin="1", ClampMax="512", UIMin="1.1", UIMax="256"))
+	UPROPERTY(Category="Tool Settings", EditAnywhere, meta=(DisplayName="Noise Scale", ShowForTools="Noise", ClampMin="1", ClampMax="512", UIMin="1.1", UIMax="256"))
 	float NoiseScale;
 
 	// Mask Tool:
 
 	// Uses selected region as a mask for other tools
-	UPROPERTY(Category="Tool Settings", EditAnywhere, meta=(DisplayName="Use Region as Mask", ShowForTools="ToolSet_Mask", ShowForMask))
+	UPROPERTY(Category="Tool Settings", EditAnywhere, meta=(DisplayName="Use Region as Mask", ShowForTools="Mask", ShowForMask))
 	bool bUseSelectedRegion;
 
 	// If enabled, protects the selected region from changes
 	// If disabled, only allows changes in the selected region
-	UPROPERTY(Category="Tool Settings", EditAnywhere, meta=(DisplayName="Negative Mask", ShowForTools="ToolSet_Mask", ShowForMask))
+	UPROPERTY(Category="Tool Settings", EditAnywhere, meta=(DisplayName="Negative Mask", ShowForTools="Mask", ShowForMask))
 	bool bUseNegativeMask;
 
 	// Copy/Paste Tool:
 
 	// Whether to paste will only raise, only lower, or both
-	UPROPERTY(Category="Tool Settings", EditAnywhere, meta=(ShowForTools="ToolSet_CopyPaste"))
+	UPROPERTY(Category="Tool Settings", EditAnywhere, meta=(ShowForTools="CopyPaste"))
 	TEnumAsByte<ELandscapeToolPasteMode::Type> PasteMode;
 
 	// If set, copies/pastes all layers, otherwise only copy/pastes the layer selected in the targets panel
-	UPROPERTY(Category="Tool Settings", EditAnywhere, meta=(DisplayName="Gizmo copy/paste all layers", ShowForTools="ToolSet_CopyPaste"))
+	UPROPERTY(Category="Tool Settings", EditAnywhere, meta=(DisplayName="Gizmo copy/paste all layers", ShowForTools="CopyPaste"))
 	bool bApplyToAllTargets;
 
 	// Makes sure the gizmo is snapped perfectly to the landscape so that the sample points line up, which makes copy/paste less blurry. Irrelevant if gizmo is scaled
-	UPROPERTY(Category="Tool Settings", EditAnywhere, meta=(DisplayName="Snap Gizmo to Landscape grid", ShowForTools="ToolSet_CopyPaste"))
+	UPROPERTY(Category="Tool Settings", EditAnywhere, meta=(DisplayName="Snap Gizmo to Landscape grid", ShowForTools="CopyPaste"))
 	bool bSnapGizmo;
 
 	// Smooths the edges of the gizmo data into the landscape. Without this, the edges of the pasted data will be sharp
-	UPROPERTY(Category="Tool Settings", EditAnywhere, meta=(DisplayName="Use Smooth Gizmo Brush", ShowForTools="ToolSet_CopyPaste"))
+	UPROPERTY(Category="Tool Settings", EditAnywhere, meta=(DisplayName="Use Smooth Gizmo Brush", ShowForTools="CopyPaste"))
 	bool bSmoothGizmoBrush;
 
-	UPROPERTY(Category="Tool Settings", EditAnywhere, AdvancedDisplay, meta=(DisplayName="Heightmap", ShowForTools="ToolSet_CopyPaste"))
+	UPROPERTY(Category="Tool Settings", EditAnywhere, AdvancedDisplay, meta=(DisplayName="Heightmap", ShowForTools="CopyPaste"))
 	FString GizmoHeightmapFilenameString;
 
-	UPROPERTY(Category="Tool Settings", EditAnywhere, AdvancedDisplay, meta=(DisplayName="Heightmap Size", ShowForTools="ToolSet_CopyPaste"))
+	UPROPERTY(Category="Tool Settings", EditAnywhere, AdvancedDisplay, meta=(DisplayName="Heightmap Size", ShowForTools="CopyPaste"))
 	FIntPoint GizmoImportSize;
 
-	UPROPERTY(Category="Tool Settings", EditAnywhere, AdvancedDisplay, meta=(DisplayName="Layers", ShowForTools="ToolSet_CopyPaste"))
+	UPROPERTY(Category="Tool Settings", EditAnywhere, AdvancedDisplay, meta=(DisplayName="Layers", ShowForTools="CopyPaste"))
 	TArray<FGizmoImportLayer> GizmoImportLayers;
 
 	TArray<FGizmoHistory> GizmoHistories;
@@ -372,19 +372,19 @@ class ULandscapeEditorObject : public UObject
 	// Resize Landscape Tool
 
 	// Number of quads per landscape component section
-	UPROPERTY(Category="Change Component Size", EditAnywhere, meta=(DisplayName="Section Size", ShowForTools="ToolSet_ResizeLandscape"))
+	UPROPERTY(Category="Change Component Size", EditAnywhere, meta=(DisplayName="Section Size", ShowForTools="ResizeLandscape"))
 	int32 ResizeLandscape_QuadsPerSection;
 
 	// Number of sections per landscape component
-	UPROPERTY(Category="Change Component Size", EditAnywhere, meta=(DisplayName="Sections Per Component", ShowForTools="ToolSet_ResizeLandscape"))
+	UPROPERTY(Category="Change Component Size", EditAnywhere, meta=(DisplayName="Sections Per Component", ShowForTools="ResizeLandscape"))
 	int32 ResizeLandscape_SectionsPerComponent;
 
 	// Number of components in resulting landscape
-	UPROPERTY(Category="Change Component Size", EditAnywhere, meta=(DisplayName="Number of Components", ShowForTools="ToolSet_ResizeLandscape"))
+	UPROPERTY(Category="Change Component Size", EditAnywhere, meta=(DisplayName="Number of Components", ShowForTools="ResizeLandscape"))
 	FIntPoint ResizeLandscape_ComponentCount;
 
 	// 
-	UPROPERTY(Category="Change Component Size", EditAnywhere, meta=(DisplayName="Resize Mode", ShowForTools="ToolSet_ResizeLandscape"))
+	UPROPERTY(Category="Change Component Size", EditAnywhere, meta=(DisplayName="Resize Mode", ShowForTools="ResizeLandscape"))
 	TEnumAsByte<ELandscapeConvertMode::Type> ResizeLandscape_ConvertMode;
 
 	int32 ResizeLandscape_Original_QuadsPerSection;
@@ -394,28 +394,28 @@ class ULandscapeEditorObject : public UObject
 	// New Landscape "Tool"
 
 	// Material initially applied to the landscape. Setting a material here exposes properties for setting up layer info based on the landscape blend nodes in the material
-	UPROPERTY(Category="New Landscape", EditAnywhere, meta=(DisplayName="Material", ShowForTools="ToolSet_NewLandscape"))
+	UPROPERTY(Category="New Landscape", EditAnywhere, meta=(DisplayName="Material", ShowForTools="NewLandscape"))
 	TWeakObjectPtr<UMaterialInterface> NewLandscape_Material;
 
-	UPROPERTY(Category="New Landscape", EditAnywhere, meta=(DisplayName="Section Size", ShowForTools="ToolSet_NewLandscape"))
+	UPROPERTY(Category="New Landscape", EditAnywhere, meta=(DisplayName="Section Size", ShowForTools="NewLandscape"))
 	int32 NewLandscape_QuadsPerSection;
-	UPROPERTY(Category="New Landscape", EditAnywhere, meta=(DisplayName="Sections Per Component", ShowForTools="ToolSet_NewLandscape"))
+	UPROPERTY(Category="New Landscape", EditAnywhere, meta=(DisplayName="Sections Per Component", ShowForTools="NewLandscape"))
 	int32 NewLandscape_SectionsPerComponent;
-	UPROPERTY(Category="New Landscape", EditAnywhere, meta=(DisplayName="Number of Components", ShowForTools="ToolSet_NewLandscape"))
+	UPROPERTY(Category="New Landscape", EditAnywhere, meta=(DisplayName="Number of Components", ShowForTools="NewLandscape"))
 	FIntPoint NewLandscape_ComponentCount;
 
-	UPROPERTY(Category="New Landscape", EditAnywhere, meta=(DisplayName="Location", ShowForTools="ToolSet_NewLandscape"))
+	UPROPERTY(Category="New Landscape", EditAnywhere, meta=(DisplayName="Location", ShowForTools="NewLandscape"))
 	FVector NewLandscape_Location;
-	UPROPERTY(Category="New Landscape", EditAnywhere, meta=(DisplayName="Rotation", ShowForTools="ToolSet_NewLandscape"))
+	UPROPERTY(Category="New Landscape", EditAnywhere, meta=(DisplayName="Rotation", ShowForTools="NewLandscape"))
 	FRotator NewLandscape_Rotation;
-	UPROPERTY(Category="New Landscape", EditAnywhere, meta=(DisplayName="Scale", ShowForTools="ToolSet_NewLandscape"))
+	UPROPERTY(Category="New Landscape", EditAnywhere, meta=(DisplayName="Scale", ShowForTools="NewLandscape"))
 	FVector NewLandscape_Scale;
 
 	// Import
 
-	UPROPERTY(Category="New Landscape", VisibleAnywhere, meta=(ShowForTools="ToolSet_NewLandscape"))
+	UPROPERTY(Category="New Landscape", VisibleAnywhere, meta=(ShowForTools="NewLandscape"))
 	TEnumAsByte<ELandscapeImportHeightmapError::Type> ImportLandscape_HeightmapError;
-	UPROPERTY(Category="New Landscape", EditAnywhere, meta=(DisplayName="Heightmap File", ShowForTools="ToolSet_NewLandscape"))
+	UPROPERTY(Category="New Landscape", EditAnywhere, meta=(DisplayName="Heightmap File", ShowForTools="NewLandscape"))
 	FString ImportLandscape_HeightmapFilename;
 	UPROPERTY()
 	int32 ImportLandscape_Width;
@@ -424,7 +424,7 @@ class ULandscapeEditorObject : public UObject
 	UPROPERTY()
 	TArray<uint16> ImportLandscape_Data;
 
-	UPROPERTY(Category="New Landscape", EditAnywhere, EditFixedSize, meta=(DisplayName="Layers", ShowForTools="ToolSet_NewLandscape"))
+	UPROPERTY(Category="New Landscape", EditAnywhere, EditFixedSize, meta=(DisplayName="Layers", ShowForTools="NewLandscape"))
 	TArray<FLandscapeImportLayer> ImportLandscape_Layers;
 
 	// Common Brush Settings:
@@ -438,7 +438,7 @@ class ULandscapeEditorObject : public UObject
 	float BrushFalloff;
 
 	// Selects the Clay Brush painting mode
-	UPROPERTY(Category="Brush Settings", EditAnywhere, meta=(ShowForTools="ToolSet_Sculpt", ShowForBrushes="BrushSet_Circle,BrushSet_Alpha,BrushSet_Pattern"))
+	UPROPERTY(Category="Brush Settings", EditAnywhere, meta=(ShowForTools="Sculpt", ShowForBrushes="BrushSet_Circle,BrushSet_Alpha,BrushSet_Pattern"))
 	bool bUseClayBrush;
 
 	// Alpha/Pattern Brush:
@@ -547,7 +547,7 @@ class ULandscapeEditorObject : public UObject
 
 	void SetbSnapGizmo(bool InbSnapGizmo);
 
-	void SetParent(class FEdModeLandscape* LandscapeParent)
+	void SetParent(FEdModeLandscape* LandscapeParent)
 	{
 		ParentMode = LandscapeParent;
 	}
diff --git a/Engine/Source/Editor/LandscapeEditor/Private/ActorFactoryLandscape.cpp b/Engine/Source/Editor/LandscapeEditor/Private/ActorFactoryLandscape.cpp
--- a/Engine/Source/Editor/LandscapeEditor/Private/ActorFactoryLandscape.cpp
+++ b/Engine/Source/Editor/LandscapeEditor/Private/ActorFactoryLandscape.cpp
@@ -6,7 +6,7 @@
 
 #define LOCTEXT_NAMESPACE "Landscape"
 
-UActorFactoryLandscape::UActorFactoryLandscape(const class FPostConstructInitializeProperties& PCIP)
+UActorFactoryLandscape::UActorFactoryLandscape(const FPostConstructInitializeProperties& PCIP)
 	: Super(PCIP)
 {
 	DisplayName = LOCTEXT("Landscape", "Landscape");
@@ -22,7 +22,7 @@ AActor* UActorFactoryLandscape::SpawnActor(UObject* Asset, ULevel* InLevel, cons
 	EdMode->UISettings->NewLandscape_Location = Location;
 	EdMode->UISettings->NewLandscape_Rotation = Rotation;
 
-	EdMode->SetCurrentTool("ToolSet_NewLandscape");
+	EdMode->SetCurrentTool("NewLandscape");
 
 	FActorSpawnParameters SpawnInfo;
 	SpawnInfo.OverrideLevel = InLevel;
@@ -31,7 +31,7 @@ AActor* UActorFactoryLandscape::SpawnActor(UObject* Asset, ULevel* InLevel, cons
 	return InLevel->OwningWorld->SpawnActor(ALandscapePlaceholder::StaticClass(), &Location, &Rotation, SpawnInfo);
 }
 
-ALandscapePlaceholder::ALandscapePlaceholder(const class FPostConstructInitializeProperties& PCIP)
+ALandscapePlaceholder::ALandscapePlaceholder(const FPostConstructInitializeProperties& PCIP)
 	: Super(PCIP)
 {
 	// Structure to hold one-time initialization
@@ -73,7 +73,7 @@ bool ALandscapePlaceholder::TeleportTo(const FVector& DestLocation, const FRotat
 	EdMode->UISettings->NewLandscape_Location = GetActorLocation();
 	EdMode->UISettings->NewLandscape_Rotation = GetActorRotation();
 
-	EdMode->SetCurrentTool("ToolSet_NewLandscape");
+	EdMode->SetCurrentTool("NewLandscape");
 
 	return bResult;
 }
