Remove FLandscapeToolSet

diff --git a/Engine/Source/Editor/LandscapeEditor/Private/LandscapeEdMode.cpp b/Engine/Source/Editor/LandscapeEditor/Private/LandscapeEdMode.cpp
--- a/Engine/Source/Editor/LandscapeEditor/Private/LandscapeEdMode.cpp
+++ b/Engine/Source/Editor/LandscapeEditor/Private/LandscapeEdMode.cpp
@@ -38,7 +38,7 @@ DEFINE_LOG_CATEGORY(LogLandscapeEdMode);
 
 struct HNewLandscapeGrabHandleProxy : public HHitProxy
 {
-	DECLARE_HIT_PROXY( );
+	DECLARE_HIT_PROXY();
 
 	ELandscapeEdge::Type Edge;
 
@@ -76,10 +76,10 @@ IMPLEMENT_HIT_PROXY(HNewLandscapeGrabHandleProxy, HHitProxy)
 void ALandscape::SplitHeightmap(ULandscapeComponent* Comp, bool bMoveToCurrentLevel /*= false*/)
 {
 	ULandscapeInfo* Info = Comp->GetLandscapeInfo();
-	int32 ComponentSizeVerts = Comp->NumSubsections * (Comp->SubsectionSizeQuads+1);
+	int32 ComponentSizeVerts = Comp->NumSubsections * (Comp->SubsectionSizeQuads + 1);
 	// make sure the heightmap UVs are powers of two.
-	int32 HeightmapSizeU = (1<<FMath::CeilLogTwo( ComponentSizeVerts ));
-	int32 HeightmapSizeV = (1<<FMath::CeilLogTwo( ComponentSizeVerts ));
+	int32 HeightmapSizeU = (1 << FMath::CeilLogTwo(ComponentSizeVerts));
+	int32 HeightmapSizeV = (1 << FMath::CeilLogTwo(ComponentSizeVerts));
 
 	UTexture2D* HeightmapTexture = NULL;
 	TArray<FColor*> HeightmapTextureMipData;
@@ -88,10 +88,10 @@ void ALandscape::SplitHeightmap(ULandscapeComponent* Comp, bool bMoveToCurrentLe
 		// Read old data and split
 		FLandscapeEditDataInterface LandscapeEdit(Info);
 		TArray<uint8> HeightData;
-		HeightData.AddZeroed((1+Comp->ComponentSizeQuads)*(1+Comp->ComponentSizeQuads)*sizeof(uint16));
+		HeightData.AddZeroed((1 + Comp->ComponentSizeQuads)*(1 + Comp->ComponentSizeQuads)*sizeof(uint16));
 		// Because of edge problem, normal would be just copy from old component data
 		TArray<uint8> NormalData;
-		NormalData.AddZeroed((1+Comp->ComponentSizeQuads)*(1+Comp->ComponentSizeQuads)*sizeof(uint16));
+		NormalData.AddZeroed((1 + Comp->ComponentSizeQuads)*(1 + Comp->ComponentSizeQuads)*sizeof(uint16));
 		LandscapeEdit.GetHeightDataFast(Comp->GetSectionBase().X, Comp->GetSectionBase().Y, Comp->GetSectionBase().X + Comp->ComponentSizeQuads, Comp->GetSectionBase().Y + Comp->ComponentSizeQuads, (uint16*)HeightData.GetTypedData(), 0, (uint16*)NormalData.GetTypedData());
 
 		// Construct the heightmap textures
@@ -101,10 +101,10 @@ void ALandscape::SplitHeightmap(ULandscapeComponent* Comp, bool bMoveToCurrentLe
 		int32 MipSubsectionSizeQuads = Comp->SubsectionSizeQuads;
 		int32 MipSizeU = HeightmapSizeU;
 		int32 MipSizeV = HeightmapSizeV;
-		while( MipSizeU > 1 && MipSizeV > 1 && MipSubsectionSizeQuads >= 1 )
+		while (MipSizeU > 1 && MipSizeV > 1 && MipSubsectionSizeQuads >= 1)
 		{
 			FColor* HeightmapTextureData = (FColor*)HeightmapTexture->Source.LockMip(HeightmapTextureMipData.Num());
-			FMemory::Memzero( HeightmapTextureData, MipSizeU*MipSizeV*sizeof(FColor) );
+			FMemory::Memzero(HeightmapTextureData, MipSizeU*MipSizeV*sizeof(FColor));
 			HeightmapTextureMipData.Add(HeightmapTextureData);
 
 			MipSizeU >>= 1;
@@ -113,7 +113,7 @@ void ALandscape::SplitHeightmap(ULandscapeComponent* Comp, bool bMoveToCurrentLe
 			MipSubsectionSizeQuads = ((MipSubsectionSizeQuads + 1) >> 1) - 1;
 		}
 
-		Comp->HeightmapScaleBias = FVector4( 1.f / (float)HeightmapSizeU, 1.f / (float)HeightmapSizeV, 0.f, 0.f);
+		Comp->HeightmapScaleBias = FVector4(1.f / (float)HeightmapSizeU, 1.f / (float)HeightmapSizeV, 0.f, 0.f);
 		Comp->HeightmapTexture = HeightmapTexture;
 
 		if (Comp->MaterialInstance)
@@ -121,7 +121,7 @@ void ALandscape::SplitHeightmap(ULandscapeComponent* Comp, bool bMoveToCurrentLe
 			Comp->MaterialInstance->SetTextureParameterValueEditorOnly(FName(TEXT("Heightmap")), HeightmapTexture);
 		}
 
-		for( int32 i=0;i<HeightmapTextureMipData.Num();i++ )
+		for (int32 i = 0; i < HeightmapTextureMipData.Num(); i++)
 		{
 			HeightmapTexture->Source.UnlockMip(i);
 		}
@@ -146,18 +146,18 @@ void FLandscapeTool::SetEditRenderType()
 //
 
 /** Constructor */
-FEdModeLandscape::FEdModeLandscape() 
-:	FEdMode()
-,	CurrentGizmoActor(NULL)
-,	LandscapeRenderAddCollision(NULL)
-,	bToolActive(false)
-,	GizmoMaterial(NULL)
-,	CopyPasteToolSet(NULL)
-,	SplinesToolSet(NULL)
-,	NewLandscapePreviewMode(ENewLandscapePreviewMode::None)
-,	DraggingEdge(ELandscapeEdge::None)
-,	DraggingEdge_Remainder(0)
-,	CachedLandscapeMaterial(NULL)
+FEdModeLandscape::FEdModeLandscape()
+	: FEdMode()
+	, CurrentGizmoActor(nullptr)
+	, LandscapeRenderAddCollision(nullptr)
+	, bToolActive(false)
+	, GizmoMaterial(nullptr)
+	, CopyPasteTool(nullptr)
+	, SplinesTool(nullptr)
+	, NewLandscapePreviewMode(ENewLandscapePreviewMode::None)
+	, DraggingEdge(ELandscapeEdge::None)
+	, DraggingEdge_Remainder(0)
+	, CachedLandscapeMaterial(nullptr)
 {
 	GizmoMaterial = LoadObject<UMaterial>(NULL, TEXT("/Engine/EditorLandscapeResources/GizmoMaterial.GizmoMaterial"), NULL, LOAD_None, NULL);
 
@@ -169,29 +169,29 @@ FEdModeLandscape::FEdModeLandscape()
 
 	// Initialize modes
 	InitializeToolModes();
-	CurrentToolMode = NULL;
+	CurrentToolMode = nullptr;
 
 	// Initialize tools.
-	IntializeToolSet_Paint();
-	IntializeToolSet_Smooth();
-	IntializeToolSet_Flatten();
-	InitializeToolSet_Ramp();
-	IntializeToolSet_Erosion();
-	IntializeToolSet_HydraErosion();
-	IntializeToolSet_Noise();
-	IntializeToolSet_Retopologize();
-	InitializeToolSet_NewLandscape();
-	InitializeToolSet_ResizeLandscape();
-	IntializeToolSet_Select();
-	IntializeToolSet_AddComponent();
-	IntializeToolSet_DeleteComponent();
-	IntializeToolSet_MoveToLevel();
-	IntializeToolSet_Mask();
-	IntializeToolSet_CopyPaste();
-	IntializeToolSet_Visibility();
-	IntializeToolSet_Splines();
-
-	CurrentToolSet = NULL;
+	InitializeTool_Paint();
+	InitializeTool_Smooth();
+	InitializeTool_Flatten();
+	InitializeTool_Ramp();
+	InitializeTool_Erosion();
+	InitializeTool_HydraErosion();
+	InitializeTool_Noise();
+	InitializeTool_Retopologize();
+	InitializeTool_NewLandscape();
+	InitializeTool_ResizeLandscape();
+	InitializeTool_Select();
+	InitializeTool_AddComponent();
+	InitializeTool_DeleteComponent();
+	InitializeTool_MoveToLevel();
+	InitializeTool_Mask();
+	InitializeTool_CopyPaste();
+	InitializeTool_Visibility();
+	InitializeTool_Splines();
+
+	CurrentTool = nullptr;
 	CurrentToolIndex = INDEX_NONE;
 
 	// Initialize brushes
@@ -200,9 +200,9 @@ FEdModeLandscape::FEdModeLandscape()
 	CurrentBrush = LandscapeBrushSets[0].Brushes[0];
 	CurrentBrushSetIndex = 0;
 
-	CurrentToolTarget.LandscapeInfo = NULL;
+	CurrentToolTarget.LandscapeInfo = nullptr;
 	CurrentToolTarget.TargetType = ELandscapeToolTargetType::Heightmap;
-	CurrentToolTarget.LayerInfo = NULL;
+	CurrentToolTarget.LayerInfo = nullptr;
 
 	UISettings = ConstructObject<ULandscapeEditorObject>(ULandscapeEditorObject::StaticClass());
 	UISettings->SetParent(this);
@@ -213,7 +213,7 @@ FEdModeLandscape::FEdModeLandscape()
 FEdModeLandscape::~FEdModeLandscape()
 {
 	// Destroy tools.
-	LandscapeToolSets.Empty();
+	LandscapeTools.Empty();
 
 	// Destroy brushes
 	LandscapeBrushSets.Empty();
@@ -228,10 +228,10 @@ FEdModeLandscape::~FEdModeLandscape()
 }
 
 /** FGCObject interface */
-void FEdModeLandscape::AddReferencedObjects( FReferenceCollector& Collector )
+void FEdModeLandscape::AddReferencedObjects(FReferenceCollector& Collector)
 {
 	// Call parent implementation
-	FEdMode::AddReferencedObjects( Collector );
+	FEdMode::AddReferencedObjects(Collector);
 
 	Collector.AddReferencedObject(UISettings);
 
@@ -245,69 +245,35 @@ void FEdModeLandscape::AddReferencedObjects( FReferenceCollector& Collector )
 
 void FEdModeLandscape::InitializeToolModes()
 {
-	FLandscapeToolMode* ToolMode_Manage = new(LandscapeToolModes) FLandscapeToolMode(TEXT("ToolMode_Manage"), ELandscapeToolTargetTypeMask::NA);
-	ToolMode_Manage->ValidToolSets.Add(TEXT("ToolSet_NewLandscape"));
-	ToolMode_Manage->ValidToolSets.Add(TEXT("ToolSet_Select"));
-	ToolMode_Manage->ValidToolSets.Add(TEXT("ToolSet_AddComponent"));
-	ToolMode_Manage->ValidToolSets.Add(TEXT("ToolSet_DeleteComponent"));
-	ToolMode_Manage->ValidToolSets.Add(TEXT("ToolSet_MoveToLevel"));
-	ToolMode_Manage->ValidToolSets.Add(TEXT("ToolSet_ResizeLandscape"));
-	ToolMode_Manage->ValidToolSets.Add(TEXT("ToolSet_Splines"));
-	ToolMode_Manage->CurrentToolSetName = TEXT("ToolSet_Select");
-
-	FLandscapeToolMode* ToolMode_Sculpt = new(LandscapeToolModes) FLandscapeToolMode(TEXT("ToolMode_Sculpt"), ELandscapeToolTargetTypeMask::Heightmap | ELandscapeToolTargetTypeMask::Visibility);
-	ToolMode_Sculpt->ValidToolSets.Add(TEXT("ToolSet_Sculpt"));
-	ToolMode_Sculpt->ValidToolSets.Add(TEXT("ToolSet_Smooth"));
-	ToolMode_Sculpt->ValidToolSets.Add(TEXT("ToolSet_Flatten"));
-	ToolMode_Sculpt->ValidToolSets.Add(TEXT("ToolSet_Ramp"));
-	ToolMode_Sculpt->ValidToolSets.Add(TEXT("ToolSet_Noise"));
-	ToolMode_Sculpt->ValidToolSets.Add(TEXT("ToolSet_Erosion"));
-	ToolMode_Sculpt->ValidToolSets.Add(TEXT("ToolSet_HydraErosion"));
-	ToolMode_Sculpt->ValidToolSets.Add(TEXT("ToolSet_Retopologize"));
-	ToolMode_Sculpt->ValidToolSets.Add(TEXT("ToolSet_Visibility"));
-	ToolMode_Sculpt->ValidToolSets.Add(TEXT("ToolSet_Mask"));
-	ToolMode_Sculpt->ValidToolSets.Add(TEXT("ToolSet_CopyPaste"));
-
-	FLandscapeToolMode* ToolMode_Paint = new(LandscapeToolModes) FLandscapeToolMode(TEXT("ToolMode_Paint"), ELandscapeToolTargetTypeMask::Weightmap);
-	ToolMode_Paint->ValidToolSets.Add(TEXT("ToolSet_Paint"));
-	ToolMode_Paint->ValidToolSets.Add(TEXT("ToolSet_Smooth"));
-	ToolMode_Paint->ValidToolSets.Add(TEXT("ToolSet_Flatten"));
-	ToolMode_Paint->ValidToolSets.Add(TEXT("ToolSet_Noise"));
-	ToolMode_Paint->ValidToolSets.Add(TEXT("ToolSet_Visibility"));
-}
-
-void FEdModeLandscape::CopyDataToGizmo()
-{
-	// For Copy operation...
-	if (CopyPasteToolSet /*&& CopyPasteToolSet == CurrentToolSet*/ )
-	{
-		if (CopyPasteToolSet->SetToolForTarget(CurrentToolTarget) && CopyPasteToolSet->GetTool())
-		{
-			CopyPasteToolSet->GetTool()->Process(0, 0);
-		}
-	}
-	if (CurrentGizmoActor.IsValid())
-	{
-		GEditor->SelectNone(false, true);
-		GEditor->SelectActor(CurrentGizmoActor.Get(), true, true, true );
-	}
-}
-
-void FEdModeLandscape::PasteDataFromGizmo()
-{
-	// For Paste for Gizmo Region operation...
-	if (CopyPasteToolSet /*&& CopyPasteToolSet == CurrentToolSet*/ )
-	{
-		if (CopyPasteToolSet->SetToolForTarget(CurrentToolTarget) && CopyPasteToolSet->GetTool())
-		{
-			CopyPasteToolSet->GetTool()->Process(1, 0);
-		}
-	}
-	if (CurrentGizmoActor.IsValid())
-	{
-		GEditor->SelectNone(false, true);
-		GEditor->SelectActor(CurrentGizmoActor.Get(), true, true, true );
-	}
+	FLandscapeToolMode* ToolMode_Manage = new(LandscapeToolModes)FLandscapeToolMode(TEXT("ToolMode_Manage"), ELandscapeToolTargetTypeMask::NA);
+	ToolMode_Manage->ValidTools.Add(TEXT("NewLandscape"));
+	ToolMode_Manage->ValidTools.Add(TEXT("Select"));
+	ToolMode_Manage->ValidTools.Add(TEXT("AddComponent"));
+	ToolMode_Manage->ValidTools.Add(TEXT("DeleteComponent"));
+	ToolMode_Manage->ValidTools.Add(TEXT("MoveToLevel"));
+	ToolMode_Manage->ValidTools.Add(TEXT("ResizeLandscape"));
+	ToolMode_Manage->ValidTools.Add(TEXT("Splines"));
+	ToolMode_Manage->CurrentToolName = TEXT("Select");
+
+	FLandscapeToolMode* ToolMode_Sculpt = new(LandscapeToolModes)FLandscapeToolMode(TEXT("ToolMode_Sculpt"), ELandscapeToolTargetTypeMask::Heightmap | ELandscapeToolTargetTypeMask::Visibility);
+	ToolMode_Sculpt->ValidTools.Add(TEXT("Sculpt"));
+	ToolMode_Sculpt->ValidTools.Add(TEXT("Smooth"));
+	ToolMode_Sculpt->ValidTools.Add(TEXT("Flatten"));
+	ToolMode_Sculpt->ValidTools.Add(TEXT("Ramp"));
+	ToolMode_Sculpt->ValidTools.Add(TEXT("Noise"));
+	ToolMode_Sculpt->ValidTools.Add(TEXT("Erosion"));
+	ToolMode_Sculpt->ValidTools.Add(TEXT("HydraErosion"));
+	ToolMode_Sculpt->ValidTools.Add(TEXT("Retopologize"));
+	ToolMode_Sculpt->ValidTools.Add(TEXT("Visibility"));
+	ToolMode_Sculpt->ValidTools.Add(TEXT("Mask"));
+	ToolMode_Sculpt->ValidTools.Add(TEXT("CopyPaste"));
+
+	FLandscapeToolMode* ToolMode_Paint = new(LandscapeToolModes)FLandscapeToolMode(TEXT("ToolMode_Paint"), ELandscapeToolTargetTypeMask::Weightmap);
+	ToolMode_Paint->ValidTools.Add(TEXT("Paint"));
+	ToolMode_Paint->ValidTools.Add(TEXT("Smooth"));
+	ToolMode_Paint->ValidTools.Add(TEXT("Flatten"));
+	ToolMode_Paint->ValidTools.Add(TEXT("Noise"));
+	ToolMode_Paint->ValidTools.Add(TEXT("Visibility"));
 }
 
 bool FEdModeLandscape::UsesToolkits() const
@@ -389,7 +355,7 @@ void FEdModeLandscape::Enter()
 			// Restore Sampled Data if exist...
 			if (CurrentGizmoActor->CachedScaleXY > 0.f)
 			{
-				int32 SizeX = FMath::CeilToInt(CurrentGizmoActor->CachedWidth  / CurrentGizmoActor->CachedScaleXY);
+				int32 SizeX = FMath::CeilToInt(CurrentGizmoActor->CachedWidth / CurrentGizmoActor->CachedScaleXY);
 				int32 SizeY = FMath::CeilToInt(CurrentGizmoActor->CachedHeight / CurrentGizmoActor->CachedScaleXY);
 				for (int32 Y = 0; Y < CurrentGizmoActor->SampleSizeY; ++Y)
 				{
@@ -404,9 +370,9 @@ void FEdModeLandscape::Enter()
 						float FracY = TexY - LY;
 
 						FGizmoSelectData* Data00 = CurrentGizmoActor->SelectedData.Find(ALandscape::MakeKey(LX, LY));
-						FGizmoSelectData* Data10 = CurrentGizmoActor->SelectedData.Find(ALandscape::MakeKey(LX+1, LY));
-						FGizmoSelectData* Data01 = CurrentGizmoActor->SelectedData.Find(ALandscape::MakeKey(LX, LY+1));
-						FGizmoSelectData* Data11 = CurrentGizmoActor->SelectedData.Find(ALandscape::MakeKey(LX+1, LY+1));
+						FGizmoSelectData* Data10 = CurrentGizmoActor->SelectedData.Find(ALandscape::MakeKey(LX + 1, LY));
+						FGizmoSelectData* Data01 = CurrentGizmoActor->SelectedData.Find(ALandscape::MakeKey(LX, LY + 1));
+						FGizmoSelectData* Data11 = CurrentGizmoActor->SelectedData.Find(ALandscape::MakeKey(LX + 1, LY + 1));
 
 						TexData[X + Y*ALandscapeGizmoActiveActor::DataTexSize] = FMath::Lerp(
 							FMath::Lerp(Data00 ? Data00->Ratio : 0, Data10 ? Data10->Ratio : 0, FracX),
@@ -444,7 +410,7 @@ void FEdModeLandscape::Enter()
 	if (!Toolkit.IsValid())
 	{
 		// @todo: Remove this assumption when we make modes per level editor instead of global
-		auto ToolkitHost = FModuleManager::LoadModuleChecked< FLevelEditorModule >( "LevelEditor" ).GetFirstLevelEditor();
+		auto ToolkitHost = FModuleManager::LoadModuleChecked< FLevelEditorModule >("LevelEditor").GetFirstLevelEditor();
 		Toolkit = MakeShareable(new FLandscapeToolKit);
 		Toolkit->Init(ToolkitHost);
 	}
@@ -453,7 +419,7 @@ void FEdModeLandscape::Enter()
 	// user exits this mode.
 	const bool bWantRealTime = true;
 	const bool bRememberCurrentState = true;
-	ForceRealTimeViewports( bWantRealTime, bRememberCurrentState );
+	ForceRealTimeViewports(bWantRealTime, bRememberCurrentState);
 
 	CurrentBrush->EnterBrush();
 	if (GizmoBrush)
@@ -464,18 +430,18 @@ void FEdModeLandscape::Enter()
 	if (LandscapeList.Num() == 0)
 	{
 		SetCurrentToolMode("ToolMode_Manage", false);
-		SetCurrentTool("ToolSet_NewLandscape");
+		SetCurrentTool("NewLandscape");
 	}
 	else
 	{
 		if (CurrentToolMode == NULL)
 		{
 			SetCurrentToolMode("ToolMode_Sculpt", false);
-			SetCurrentTool("ToolSet_Sculpt");
+			SetCurrentTool("Sculpt");
 		}
 		else
 		{
-			SetCurrentTool(CurrentToolMode->CurrentToolSetName);
+			SetCurrentTool(CurrentToolMode->CurrentToolName);
 		}
 	}
 }
@@ -490,7 +456,7 @@ void FEdModeLandscape::Exit()
 	// Restore real-time viewport state if we changed it
 	const bool bWantRealTime = false;
 	const bool bRememberCurrentState = false;
-	ForceRealTimeViewports( bWantRealTime, bRememberCurrentState );
+	ForceRealTimeViewports(bWantRealTime, bRememberCurrentState);
 
 	if (Toolkit.IsValid())
 	{
@@ -504,15 +470,15 @@ void FEdModeLandscape::Exit()
 		GizmoBrush->LeaveBrush();
 	}
 
-	if (CurrentToolSet)
+	if (CurrentTool)
 	{
-		CurrentToolSet->PreviousBrushIndex = CurrentBrushSetIndex;
-		if (CurrentToolSet->GetTool())
+		CurrentTool->PreviousBrushIndex = CurrentBrushSetIndex;
+		if (CurrentTool)
 		{
-			CurrentToolSet->GetTool()->ExitTool();
+			CurrentTool->ExitTool();
 		}
 	}
-	CurrentToolSet = NULL;
+	CurrentTool = NULL;
 	// Leave CurrentToolIndex set so we can restore the active tool on re-opening the landscape editor
 
 	LandscapeList.Empty();
@@ -529,7 +495,7 @@ void FEdModeLandscape::Exit()
 
 	CurrentGizmoActor = NULL;
 
-	GEditor->SelectNone( false, true );
+	GEditor->SelectNone(false, true);
 
 	// Clear all GizmoActors if there is no Landscape in World
 	bool bIsLandscapeExist = false;
@@ -563,9 +529,9 @@ void FEdModeLandscape::Exit()
 }
 
 /** FEdMode: Called once per frame */
-void FEdModeLandscape::Tick(FEditorViewportClient* ViewportClient,float DeltaTime)
+void FEdModeLandscape::Tick(FEditorViewportClient* ViewportClient, float DeltaTime)
 {
-	FEdMode::Tick(ViewportClient,DeltaTime);
+	FEdMode::Tick(ViewportClient, DeltaTime);
 
 	if (GEditor->PlayWorld != NULL)
 	{
@@ -576,7 +542,7 @@ void FEdModeLandscape::Tick(FEditorViewportClient* ViewportClient,float DeltaTim
 	{
 		bool bStaleTargetLandscapeInfo = CurrentToolTarget.LandscapeInfo.IsStale();
 		bool bStaleTargetLandscape = CurrentToolTarget.LandscapeInfo.IsValid() && CurrentToolTarget.LandscapeInfo->LandscapeActor.IsStale();
-		
+
 		if (bStaleTargetLandscapeInfo || bStaleTargetLandscape)
 		{
 			UpdateLandscapeList();
@@ -585,22 +551,22 @@ void FEdModeLandscape::Tick(FEditorViewportClient* ViewportClient,float DeltaTim
 		if (CurrentToolTarget.LandscapeInfo.IsValid())
 		{
 			ALandscapeProxy* LandscapeProxy = CurrentToolTarget.LandscapeInfo->GetLandscapeProxy();
-			if (LandscapeProxy == NULL || 
+			if (LandscapeProxy == NULL ||
 				LandscapeProxy->GetLandscapeMaterial() != CachedLandscapeMaterial)
 			{
 				UpdateTargetList();
 			}
 		}
 
-		if( CurrentToolSet && CurrentToolSet->GetTool() )
+		if (CurrentTool && CurrentTool)
 		{
-			CurrentToolSet->GetTool()->Tick(ViewportClient,DeltaTime);
+			CurrentTool->Tick(ViewportClient, DeltaTime);
 		}
-		if( CurrentBrush )
+		if (CurrentBrush)
 		{
-			CurrentBrush->Tick(ViewportClient,DeltaTime);
+			CurrentBrush->Tick(ViewportClient, DeltaTime);
 		}
-		if ( CurrentBrush!=GizmoBrush && CurrentGizmoActor.IsValid() && GizmoBrush && (GLandscapeEditRenderMode & ELandscapeEditRenderMode::Gizmo) )
+		if (CurrentBrush != GizmoBrush && CurrentGizmoActor.IsValid() && GizmoBrush && (GLandscapeEditRenderMode & ELandscapeEditRenderMode::Gizmo))
 		{
 			GizmoBrush->Tick(ViewportClient, DeltaTime);
 		}
@@ -609,13 +575,13 @@ void FEdModeLandscape::Tick(FEditorViewportClient* ViewportClient,float DeltaTim
 
 
 /** FEdMode: Called when the mouse is moved over the viewport */
-bool FEdModeLandscape::MouseMove( FEditorViewportClient* ViewportClient, FViewport* Viewport, int32 MouseX, int32 MouseY )
+bool FEdModeLandscape::MouseMove(FEditorViewportClient* ViewportClient, FViewport* Viewport, int32 MouseX, int32 MouseY)
 {
 	if (bToolActive && !Viewport->KeyState(EKeys::LeftMouseButton))
 	{
-		if (CurrentToolSet && CurrentToolSet->GetTool())
+		if (CurrentTool && CurrentTool)
 		{
-			CurrentToolSet->GetTool()->EndTool(ViewportClient);
+			CurrentTool->EndTool(ViewportClient);
 			bToolActive = false;
 		}
 	}
@@ -628,9 +594,9 @@ bool FEdModeLandscape::MouseMove( FEditorViewportClient* ViewportClient, FViewpo
 	bool Result = false;
 	if (NewLandscapePreviewMode == ENewLandscapePreviewMode::None)
 	{
-		if( CurrentToolSet && CurrentToolSet->GetTool() )
+		if (CurrentTool && CurrentTool)
 		{
-			Result = CurrentToolSet && CurrentToolSet->GetTool()->MouseMove(ViewportClient, Viewport, MouseX, MouseY);
+			Result = CurrentTool && CurrentTool->MouseMove(ViewportClient, Viewport, MouseX, MouseY);
 			//ViewportClient->Invalidate( false, false );
 		}
 	}
@@ -653,7 +619,7 @@ bool FEdModeLandscape::DisallowMouseDeltaTracking() const
  *
  * @return	true if input was handled
  */
-bool FEdModeLandscape::CapturedMouseMove( FEditorViewportClient* ViewportClient, FViewport* Viewport, int32 MouseX, int32 MouseY )
+bool FEdModeLandscape::CapturedMouseMove(FEditorViewportClient* ViewportClient, FViewport* Viewport, int32 MouseX, int32 MouseY)
 {
 	return MouseMove(ViewportClient, Viewport, MouseX, MouseY);
 }
@@ -689,20 +655,20 @@ bool FEdModeLandscape::EndTracking(FEditorViewportClient* InViewportClient, FVie
 
 namespace
 {
-	bool RayIntersectTriangle(const FVector& Start, const FVector& End, const FVector& A, const FVector& B, const FVector& C, FVector& IntersectPoint )
+	bool RayIntersectTriangle(const FVector& Start, const FVector& End, const FVector& A, const FVector& B, const FVector& C, FVector& IntersectPoint)
 	{
 		const FVector BA = A - B;
 		const FVector CB = B - C;
 		const FVector TriNormal = BA ^ CB;
 
-		bool bCollide = FMath::SegmentPlaneIntersection(Start, End, FPlane(A, TriNormal), IntersectPoint );
+		bool bCollide = FMath::SegmentPlaneIntersection(Start, End, FPlane(A, TriNormal), IntersectPoint);
 		if (!bCollide)
 		{
 			return false;
 		}
 
 		FVector BaryCentric = FMath::ComputeBaryCentric2D(IntersectPoint, A, B, C);
-		if (BaryCentric.X > 0.f && BaryCentric.Y > 0.f && BaryCentric.Z > 0.f )
+		if (BaryCentric.X > 0.f && BaryCentric.Y > 0.f && BaryCentric.Z > 0.f)
 		{
 			return true;
 		}
@@ -711,43 +677,43 @@ namespace
 };
 
 /** Trace under the mouse cursor and return the landscape hit and the hit location (in landscape quad space) */
-bool FEdModeLandscape::LandscapeMouseTrace( FEditorViewportClient* ViewportClient, float& OutHitX, float& OutHitY )
+bool FEdModeLandscape::LandscapeMouseTrace(FEditorViewportClient* ViewportClient, float& OutHitX, float& OutHitY)
 {
 	int32 MouseX = ViewportClient->Viewport->GetMouseX();
 	int32 MouseY = ViewportClient->Viewport->GetMouseY();
 
-	return LandscapeMouseTrace( ViewportClient, MouseX, MouseY, OutHitX, OutHitY );
+	return LandscapeMouseTrace(ViewportClient, MouseX, MouseY, OutHitX, OutHitY);
 }
 
-bool FEdModeLandscape::LandscapeMouseTrace( FEditorViewportClient* ViewportClient, FVector& OutHitLocation )
+bool FEdModeLandscape::LandscapeMouseTrace(FEditorViewportClient* ViewportClient, FVector& OutHitLocation)
 {
 	int32 MouseX = ViewportClient->Viewport->GetMouseX();
 	int32 MouseY = ViewportClient->Viewport->GetMouseY();
 
-	return LandscapeMouseTrace( ViewportClient, MouseX, MouseY, OutHitLocation );
+	return LandscapeMouseTrace(ViewportClient, MouseX, MouseY, OutHitLocation);
 }
 
 /** Trace under the specified coordinates and return the landscape hit and the hit location (in landscape quad space) */
-bool FEdModeLandscape::LandscapeMouseTrace( FEditorViewportClient* ViewportClient, int32 MouseX, int32 MouseY, float& OutHitX, float& OutHitY )
+bool FEdModeLandscape::LandscapeMouseTrace(FEditorViewportClient* ViewportClient, int32 MouseX, int32 MouseY, float& OutHitX, float& OutHitY)
 {
 	FVector HitLocation;
-	bool bResult = LandscapeMouseTrace( ViewportClient, MouseX, MouseY, HitLocation );
+	bool bResult = LandscapeMouseTrace(ViewportClient, MouseX, MouseY, HitLocation);
 	OutHitX = HitLocation.X;
 	OutHitY = HitLocation.Y;
 	return bResult;
 }
 
-bool FEdModeLandscape::LandscapeMouseTrace( FEditorViewportClient* ViewportClient, int32 MouseX, int32 MouseY, FVector& OutHitLocation )
+bool FEdModeLandscape::LandscapeMouseTrace(FEditorViewportClient* ViewportClient, int32 MouseX, int32 MouseY, FVector& OutHitLocation)
 {
 	// Cache a copy of the world pointer	
 	UWorld* World = ViewportClient->GetWorld();
-	
+
 	// Compute a world space ray from the screen space mouse coordinates
-	FSceneViewFamilyContext ViewFamily( FSceneViewFamilyContext::ConstructionValues( ViewportClient->Viewport, ViewportClient->GetScene(), ViewportClient->EngineShowFlags )
-		.SetRealtimeUpdate(ViewportClient->IsRealtime()) );
+	FSceneViewFamilyContext ViewFamily(FSceneViewFamilyContext::ConstructionValues(ViewportClient->Viewport, ViewportClient->GetScene(), ViewportClient->EngineShowFlags)
+		.SetRealtimeUpdate(ViewportClient->IsRealtime()));
 
-	FSceneView* View = ViewportClient->CalcSceneView( &ViewFamily );
-	FViewportCursorLocation MouseViewportRay( View, ViewportClient, MouseX, MouseY );
+	FSceneView* View = ViewportClient->CalcSceneView(&ViewFamily);
+	FViewportCursorLocation MouseViewportRay(View, ViewportClient, MouseX, MouseY);
 
 	FVector Start = MouseViewportRay.GetOrigin();
 	FVector End = Start + WORLD_MAX * MouseViewportRay.GetDirection();
@@ -758,16 +724,16 @@ bool FEdModeLandscape::LandscapeMouseTrace( FEditorViewportClient* ViewportClien
 	// Trace only ECC_Visibility channel, so we do hit only Editor specific shape
 	World->LineTraceMulti(Results, Start, End, FCollisionQueryParams(TraceTag, true), FCollisionObjectQueryParams(ECollisionChannel::ECC_Visibility));
 
-	for( int32 i=0; i<Results.Num(); i++ )
+	for (int32 i = 0; i < Results.Num(); i++)
 	{
 		const FHitResult& Hit = Results[i];
 		ULandscapeHeightfieldCollisionComponent* CollisionComponent = Cast<ULandscapeHeightfieldCollisionComponent>(Hit.Component.Get());
-		if( CollisionComponent )
+		if (CollisionComponent)
 		{
 			ALandscapeProxy* HitLandscape = CollisionComponent->GetLandscapeProxy();
-			if( HitLandscape && 
-				CurrentToolTarget.LandscapeInfo.IsValid() && 
-				CurrentToolTarget.LandscapeInfo->LandscapeGuid == HitLandscape->GetLandscapeGuid() )
+			if (HitLandscape &&
+				CurrentToolTarget.LandscapeInfo.IsValid() &&
+				CurrentToolTarget.LandscapeInfo->LandscapeGuid == HitLandscape->GetLandscapeGuid())
 			{
 				OutHitLocation = HitLandscape->LandscapeActorToWorld().InverseTransformPosition(Hit.Location);
 				return true;
@@ -776,33 +742,33 @@ bool FEdModeLandscape::LandscapeMouseTrace( FEditorViewportClient* ViewportClien
 	}
 
 	// For Add Landscape Component Mode
-	if (CurrentToolSet->GetToolSetName() == FName(TEXT("ToolSet_AddComponent")) && 
+	if (CurrentTool->GetToolName() == FName("AddComponent") &&
 		CurrentToolTarget.LandscapeInfo.IsValid())
 	{
 		bool bCollided = false;
 		FVector IntersectPoint;
 		LandscapeRenderAddCollision = NULL;
 		// Need to optimize collision for AddLandscapeComponent...?
-		for (auto It = CurrentToolTarget.LandscapeInfo->XYtoAddCollisionMap.CreateIterator(); It; ++It )
+		for (auto It = CurrentToolTarget.LandscapeInfo->XYtoAddCollisionMap.CreateIterator(); It; ++It)
 		{
 			FLandscapeAddCollision& AddCollision = It.Value();
 			// Triangle 1
-			bCollided = RayIntersectTriangle(Start, End, AddCollision.Corners[0], AddCollision.Corners[3], AddCollision.Corners[1], IntersectPoint );
+			bCollided = RayIntersectTriangle(Start, End, AddCollision.Corners[0], AddCollision.Corners[3], AddCollision.Corners[1], IntersectPoint);
 			if (bCollided)
 			{
 				LandscapeRenderAddCollision = &AddCollision;
 				break;
 			}
 			// Triangle 2
-			bCollided = RayIntersectTriangle(Start, End, AddCollision.Corners[0], AddCollision.Corners[2], AddCollision.Corners[3], IntersectPoint );
+			bCollided = RayIntersectTriangle(Start, End, AddCollision.Corners[0], AddCollision.Corners[2], AddCollision.Corners[3], IntersectPoint);
 			if (bCollided)
 			{
 				LandscapeRenderAddCollision = &AddCollision;
 				break;
 			}
 		}
-		
-		if (bCollided && 
+
+		if (bCollided &&
 			CurrentToolTarget.LandscapeInfo.IsValid())
 		{
 			ALandscapeProxy* Proxy = CurrentToolTarget.LandscapeInfo.Get()->GetCurrentLevelLandscapeProxy(true);
@@ -813,7 +779,7 @@ bool FEdModeLandscape::LandscapeMouseTrace( FEditorViewportClient* ViewportClien
 			}
 		}
 	}
-	
+
 	return false;
 }
 
@@ -831,12 +797,12 @@ bool FEdModeLandscape::LandscapePlaneTrace(FEditorViewportClient* ViewportClient
 	UWorld* World = ViewportClient->GetWorld();
 	// Compute a world space ray from the screen space mouse coordinates
 	FSceneViewFamilyContext ViewFamily(FSceneViewFamily::ConstructionValues(
-		ViewportClient->Viewport, 
+		ViewportClient->Viewport,
 		ViewportClient->GetScene(),
-		ViewportClient->EngineShowFlags )
-		.SetRealtimeUpdate( ViewportClient->IsRealtime() ));
-	FSceneView* View = ViewportClient->CalcSceneView( &ViewFamily );
-	FViewportCursorLocation MouseViewportRay( View, ViewportClient, MouseX, MouseY );
+		ViewportClient->EngineShowFlags)
+		.SetRealtimeUpdate(ViewportClient->IsRealtime()));
+	FSceneView* View = ViewportClient->CalcSceneView(&ViewFamily);
+	FViewportCursorLocation MouseViewportRay(View, ViewportClient, MouseX, MouseY);
 
 	FVector Start = MouseViewportRay.GetOrigin();
 	FVector End = Start + WORLD_MAX * MouseViewportRay.GetDirection();
@@ -855,7 +821,7 @@ namespace
 		{
 			int32 MinX = MAX_int32, MinY = MAX_int32, MaxX = MIN_int32, MaxY = MIN_int32;
 			LandscapeInfo->GetSelectedExtent(MinX, MinY, MaxX, MaxY);
-			return (MinX != MAX_int32 && ( (MaxX - MinX) * (MaxY - MinY) ));
+			return (MinX != MAX_int32 && ((MaxX - MinX) * (MaxY - MinY)));
 		}
 		return false;
 	}
@@ -867,9 +833,9 @@ EEditAction::Type FEdModeLandscape::GetActionEditDuplicate()
 
 	if (NewLandscapePreviewMode == ENewLandscapePreviewMode::None)
 	{
-		if (CurrentToolSet != NULL && CurrentToolSet->GetTool() != NULL)
+		if (CurrentTool != NULL && CurrentTool != NULL)
 		{
-			Result = CurrentToolSet->GetTool()->GetActionEditDuplicate();
+			Result = CurrentTool->GetActionEditDuplicate();
 		}
 	}
 
@@ -882,9 +848,9 @@ EEditAction::Type FEdModeLandscape::GetActionEditDelete()
 
 	if (NewLandscapePreviewMode == ENewLandscapePreviewMode::None)
 	{
-		if (CurrentToolSet != NULL && CurrentToolSet->GetTool() != NULL)
+		if (CurrentTool != NULL && CurrentTool != NULL)
 		{
-			Result = CurrentToolSet->GetTool()->GetActionEditDelete();
+			Result = CurrentTool->GetActionEditDelete();
 		}
 
 		if (Result == EEditAction::Skip)
@@ -917,9 +883,9 @@ EEditAction::Type FEdModeLandscape::GetActionEditCut()
 
 	if (NewLandscapePreviewMode == ENewLandscapePreviewMode::None)
 	{
-		if (CurrentToolSet != NULL && CurrentToolSet->GetTool() != NULL)
+		if (CurrentTool != NULL && CurrentTool != NULL)
 		{
-			Result = CurrentToolSet->GetTool()->GetActionEditCut();
+			Result = CurrentTool->GetActionEditCut();
 		}
 	}
 
@@ -939,9 +905,9 @@ EEditAction::Type FEdModeLandscape::GetActionEditCopy()
 
 	if (NewLandscapePreviewMode == ENewLandscapePreviewMode::None)
 	{
-		if (CurrentToolSet != NULL && CurrentToolSet->GetTool() != NULL)
+		if (CurrentTool != NULL && CurrentTool != NULL)
 		{
-			Result = CurrentToolSet->GetTool()->GetActionEditCopy();
+			Result = CurrentTool->GetActionEditCopy();
 		}
 
 		if (Result == EEditAction::Skip)
@@ -965,9 +931,9 @@ EEditAction::Type FEdModeLandscape::GetActionEditPaste()
 
 	if (NewLandscapePreviewMode == ENewLandscapePreviewMode::None)
 	{
-		if (CurrentToolSet != NULL && CurrentToolSet->GetTool() != NULL)
+		if (CurrentTool != NULL && CurrentTool != NULL)
 		{
-			Result = CurrentToolSet->GetTool()->GetActionEditPaste();
+			Result = CurrentTool->GetActionEditPaste();
 		}
 
 		if (Result == EEditAction::Skip)
@@ -996,9 +962,9 @@ bool FEdModeLandscape::ProcessEditDuplicate()
 
 	if (NewLandscapePreviewMode == ENewLandscapePreviewMode::None)
 	{
-		if (CurrentToolSet != NULL && CurrentToolSet->GetTool() != NULL)
+		if (CurrentTool != NULL && CurrentTool != NULL)
 		{
-			Result = CurrentToolSet->GetTool()->ProcessEditDuplicate();
+			Result = CurrentTool->ProcessEditDuplicate();
 		}
 	}
 
@@ -1016,9 +982,9 @@ bool FEdModeLandscape::ProcessEditDelete()
 
 	if (NewLandscapePreviewMode == ENewLandscapePreviewMode::None)
 	{
-		if (CurrentToolSet != NULL && CurrentToolSet->GetTool() != NULL)
+		if (CurrentTool != NULL && CurrentTool != NULL)
 		{
-			Result = CurrentToolSet->GetTool()->ProcessEditDelete();
+			Result = CurrentTool->ProcessEditDelete();
 		}
 	}
 
@@ -1036,9 +1002,9 @@ bool FEdModeLandscape::ProcessEditCut()
 
 	if (NewLandscapePreviewMode == ENewLandscapePreviewMode::None)
 	{
-		if (CurrentToolSet != NULL && CurrentToolSet->GetTool() != NULL)
+		if (CurrentTool != NULL && CurrentTool != NULL)
 		{
-			Result = CurrentToolSet->GetTool()->ProcessEditCut();
+			Result = CurrentTool->ProcessEditCut();
 		}
 	}
 
@@ -1056,9 +1022,9 @@ bool FEdModeLandscape::ProcessEditCopy()
 
 	if (NewLandscapePreviewMode == ENewLandscapePreviewMode::None)
 	{
-		if (CurrentToolSet != NULL && CurrentToolSet->GetTool() != NULL)
+		if (CurrentTool != NULL && CurrentTool != NULL)
 		{
-			Result = CurrentToolSet->GetTool()->ProcessEditCopy();
+			Result = CurrentTool->ProcessEditCopy();
 		}
 
 		if (!Result)
@@ -1066,14 +1032,14 @@ bool FEdModeLandscape::ProcessEditCopy()
 			bool IsSlowTask = IsSlowSelect(CurrentGizmoActor->TargetLandscapeInfo);
 			if (IsSlowTask)
 			{
-				GWarn->BeginSlowTask( LOCTEXT("BeginFitGizmoAndCopy", "Fit Gizmo to Selected Region and Copy Data..."), true);
+				GWarn->BeginSlowTask(LOCTEXT("BeginFitGizmoAndCopy", "Fit Gizmo to Selected Region and Copy Data..."), true);
 			}
 
-			FScopedTransaction Transaction( LOCTEXT("LandscapeGizmo_Copy", "Copy landscape data to Gizmo") );
+			FScopedTransaction Transaction(LOCTEXT("LandscapeGizmo_Copy", "Copy landscape data to Gizmo"));
 			CurrentGizmoActor->Modify();
-			CurrentGizmoActor->FitToSelection(); 
+			CurrentGizmoActor->FitToSelection();
 			CopyDataToGizmo();
-			SetCurrentTool(FName(TEXT("ToolSet_CopyPaste")));
+			SetCurrentTool(FName("CopyPaste"));
 
 			if (IsSlowTask)
 			{
@@ -1098,9 +1064,9 @@ bool FEdModeLandscape::ProcessEditPaste()
 
 	if (NewLandscapePreviewMode == ENewLandscapePreviewMode::None)
 	{
-		if (CurrentToolSet != NULL && CurrentToolSet->GetTool() != NULL)
+		if (CurrentTool != NULL && CurrentTool != NULL)
 		{
-			Result = CurrentToolSet->GetTool()->ProcessEditPaste();
+			Result = CurrentTool->ProcessEditPaste();
 		}
 
 		if (!Result)
@@ -1108,10 +1074,10 @@ bool FEdModeLandscape::ProcessEditPaste()
 			bool IsSlowTask = IsSlowSelect(CurrentGizmoActor->TargetLandscapeInfo);
 			if (IsSlowTask)
 			{
-				GWarn->BeginSlowTask( LOCTEXT("BeginPasteGizmoDataTask", "Paste Gizmo Data..."), true);
+				GWarn->BeginSlowTask(LOCTEXT("BeginPasteGizmoDataTask", "Paste Gizmo Data..."), true);
 			}
 			PasteDataFromGizmo();
-			SetCurrentTool(FName(TEXT("ToolSet_CopyPaste")));
+			SetCurrentTool(FName("CopyPaste"));
 			if (IsSlowTask)
 			{
 				GWarn->EndSlowTask();
@@ -1137,7 +1103,7 @@ bool FEdModeLandscape::HandleClick(FEditorViewportClient* InViewportClient, HHit
 	}
 
 	// Override Click Input for Splines Tool
-	if (CurrentToolSet && CurrentToolSet->GetTool() && CurrentToolSet->GetTool()->HandleClick(HitProxy, Click))
+	if (CurrentTool && CurrentTool && CurrentTool->HandleClick(HitProxy, Click))
 	{
 		return true;
 	}
@@ -1146,7 +1112,7 @@ bool FEdModeLandscape::HandleClick(FEditorViewportClient* InViewportClient, HHit
 }
 
 /** FEdMode: Called when a key is pressed */
-bool FEdModeLandscape::InputKey( FEditorViewportClient* ViewportClient, FViewport* Viewport, FKey Key, EInputEvent Event )
+bool FEdModeLandscape::InputKey(FEditorViewportClient* ViewportClient, FViewport* Viewport, FKey Key, EInputEvent Event)
 {
 	if (GEditor->PlayWorld != NULL)
 	{
@@ -1166,7 +1132,7 @@ bool FEdModeLandscape::InputKey( FEditorViewportClient* ViewportClient, FViewpor
 				HHitProxy*	HitProxy = Viewport->GetHitProxy(HitX, HitY);
 				if (HitProxy)
 				{
-					if (HitProxy->IsA(HNewLandscapeGrabHandleProxy::StaticGetType()) )
+					if (HitProxy->IsA(HNewLandscapeGrabHandleProxy::StaticGetType()))
 					{
 						HNewLandscapeGrabHandleProxy* EdgeProxy = (HNewLandscapeGrabHandleProxy*)HitProxy;
 						DraggingEdge = EdgeProxy->Edge;
@@ -1191,12 +1157,12 @@ bool FEdModeLandscape::InputKey( FEditorViewportClient* ViewportClient, FViewpor
 	else
 	{
 		// Override Key Input for Selection Brush
-		if( CurrentBrush && CurrentBrush->InputKey(ViewportClient, Viewport, Key, Event)==true )
+		if (CurrentBrush && CurrentBrush->InputKey(ViewportClient, Viewport, Key, Event) == true)
 		{
 			return true;
 		}
 
-		if( CurrentToolSet && CurrentToolSet->GetTool() && CurrentToolSet->GetTool()->InputKey(ViewportClient, Viewport, Key, Event)==true )
+		if (CurrentTool && CurrentTool && CurrentTool->InputKey(ViewportClient, Viewport, Key, Event) == true)
 		{
 			return true;
 		}
@@ -1205,38 +1171,35 @@ bool FEdModeLandscape::InputKey( FEditorViewportClient* ViewportClient, FViewpor
 		{
 			// Cheat, but works... :)
 			bool bNeedSelectGizmo = CurrentGizmoActor.IsValid() && (GLandscapeEditRenderMode & ELandscapeEditRenderMode::Gizmo) && CurrentGizmoActor->IsSelected();
-			GEditor->SelectNone( false, true );
+			GEditor->SelectNone(false, true);
 			if (bNeedSelectGizmo)
 			{
-				GEditor->SelectActor( CurrentGizmoActor.Get(), true,  false );
+				GEditor->SelectActor(CurrentGizmoActor.Get(), true, false);
 			}
 		}
 
-		if( Key == EKeys::LeftMouseButton )
+		if (Key == EKeys::LeftMouseButton)
 		{
-			if( Event == IE_Pressed && (IsCtrlDown(Viewport) || (Viewport->IsPenActive() && Viewport->GetTabletPressure() > 0.f)) )
+			if (Event == IE_Pressed && (IsCtrlDown(Viewport) || (Viewport->IsPenActive() && Viewport->GetTabletPressure() > 0.f)))
 			{
-				if( CurrentToolSet )
+				if (CurrentTool)
 				{
-					if( CurrentToolSet->SetToolForTarget(CurrentToolTarget) )
+					FVector HitLocation;
+					if (LandscapeMouseTrace(ViewportClient, HitLocation))
 					{
-						FVector HitLocation;
-						if( LandscapeMouseTrace(ViewportClient, HitLocation) )
+						if (CurrentToolTarget.TargetType == ELandscapeToolTargetType::Weightmap && CurrentToolTarget.LayerInfo == NULL)
 						{
-							if (CurrentToolTarget.TargetType == ELandscapeToolTargetType::Weightmap && CurrentToolTarget.LayerInfo == NULL)
-							{
-								FMessageDialog::Open(EAppMsgType::Ok,
-									NSLOCTEXT("UnrealEd", "LandscapeNeedToCreateLayerInfo", "This layer has no layer info assigned yet. You must create or assign a layer info before you can paint this layer."));
-								// TODO: FName to LayerInfo: do we want to create the layer info here?
-								//if (FMessageDialog::Open(EAppMsgType::YesNo, NSLOCTEXT("UnrealEd", "LandscapeNeedToCreateLayerInfo", "This layer has no layer info assigned yet. You must create or assign a layer info before you can paint this layer.")) == EAppReturnType::Yes)
-								//{
-								//	CurrentToolTarget.LandscapeInfo->LandscapeProxy->CreateLayerInfo(*CurrentToolTarget.PlaceholderLayerName.ToString());
-								//}
-							}
-							else
-							{
-								bToolActive = CurrentToolSet->GetTool()->BeginTool(ViewportClient, CurrentToolTarget, HitLocation);
-							}
+							FMessageDialog::Open(EAppMsgType::Ok,
+								NSLOCTEXT("UnrealEd", "LandscapeNeedToCreateLayerInfo", "This layer has no layer info assigned yet. You must create or assign a layer info before you can paint this layer."));
+							// TODO: FName to LayerInfo: do we want to create the layer info here?
+							//if (FMessageDialog::Open(EAppMsgType::YesNo, NSLOCTEXT("UnrealEd", "LandscapeNeedToCreateLayerInfo", "This layer has no layer info assigned yet. You must create or assign a layer info before you can paint this layer.")) == EAppReturnType::Yes)
+							//{
+							//	CurrentToolTarget.LandscapeInfo->LandscapeProxy->CreateLayerInfo(*CurrentToolTarget.PlaceholderLayerName.ToString());
+							//}
+						}
+						else
+						{
+							bToolActive = CurrentTool->BeginTool(ViewportClient, CurrentToolTarget, HitLocation);
 						}
 					}
 				}
@@ -1244,20 +1207,20 @@ bool FEdModeLandscape::InputKey( FEditorViewportClient* ViewportClient, FViewpor
 			}
 		}
 
-		if( Key == EKeys::LeftMouseButton || Key==EKeys::LeftControl || Key==EKeys::RightControl )
+		if (Key == EKeys::LeftMouseButton || Key == EKeys::LeftControl || Key == EKeys::RightControl)
 		{
-			if( Event == IE_Released && CurrentToolSet && CurrentToolSet->GetTool() && bToolActive )
+			if (Event == IE_Released && CurrentTool && CurrentTool && bToolActive)
 			{
 				//Set the cursor position to that of the slate cursor so it wont snap back
 				Viewport->SetPreCaptureMousePosFromSlateCursor();
-				CurrentToolSet->GetTool()->EndTool(ViewportClient);
+				CurrentTool->EndTool(ViewportClient);
 				bToolActive = false;
 				return (Key == EKeys::LeftMouseButton);
 			}
 		}
 
 		// Change Brush Size
-		if ((Event == IE_Pressed || Event == IE_Repeat) && (Key == EKeys::LeftBracket || Key == EKeys::RightBracket) )
+		if ((Event == IE_Pressed || Event == IE_Repeat) && (Key == EKeys::LeftBracket || Key == EKeys::RightBracket))
 		{
 			if (CurrentBrush->GetBrushType() == FLandscapeBrush::BT_Component)
 			{
@@ -1285,7 +1248,7 @@ bool FEdModeLandscape::InputKey( FEditorViewportClient* ViewportClient, FViewpor
 					Diff = -Diff;
 				}
 
-				float NewValue = Radius*(1.f+Diff);
+				float NewValue = Radius*(1.f + Diff);
 
 				if (Key == EKeys::LeftBracket)
 				{
@@ -1309,33 +1272,33 @@ bool FEdModeLandscape::InputKey( FEditorViewportClient* ViewportClient, FViewpor
 		}
 
 		// Prev tool 
-		if( Event == IE_Pressed && Key == EKeys::Comma )
+		if (Event == IE_Pressed && Key == EKeys::Comma)
 		{
-			if( CurrentToolSet && CurrentToolSet->GetTool() && bToolActive )
+			if (CurrentTool && CurrentTool && bToolActive)
 			{
-				CurrentToolSet->GetTool()->EndTool(ViewportClient);
+				CurrentTool->EndTool(ViewportClient);
 				bToolActive = false;
 			}
 
-			int32 OldToolIndex = CurrentToolMode->ValidToolSets.Find(CurrentToolSet->GetToolSetName());
+			int32 OldToolIndex = CurrentToolMode->ValidTools.Find(CurrentTool->GetToolName());
 			int32 NewToolIndex = FMath::Max(OldToolIndex - 1, 0);
-			SetCurrentTool(CurrentToolMode->ValidToolSets[NewToolIndex]);
+			SetCurrentTool(CurrentToolMode->ValidTools[NewToolIndex]);
 
 			return true;
 		}
 
 		// Next tool 
-		if( Event == IE_Pressed && Key == EKeys::Period )
+		if (Event == IE_Pressed && Key == EKeys::Period)
 		{
-			if( CurrentToolSet && CurrentToolSet->GetTool() && bToolActive )
+			if (CurrentTool && CurrentTool && bToolActive)
 			{
-				CurrentToolSet->GetTool()->EndTool(ViewportClient);
+				CurrentTool->EndTool(ViewportClient);
 				bToolActive = false;
 			}
 
-			int32 OldToolIndex = CurrentToolMode->ValidToolSets.Find(CurrentToolSet->GetToolSetName());
-			int32 NewToolIndex = FMath::Min(OldToolIndex + 1, CurrentToolMode->ValidToolSets.Num() - 1);
-			SetCurrentTool(CurrentToolMode->ValidToolSets[NewToolIndex]);
+			int32 OldToolIndex = CurrentToolMode->ValidTools.Find(CurrentTool->GetToolName());
+			int32 NewToolIndex = FMath::Min(OldToolIndex + 1, CurrentToolMode->ValidTools.Num() - 1);
+			SetCurrentTool(CurrentToolMode->ValidTools[NewToolIndex]);
 
 			return true;
 		}
@@ -1345,7 +1308,7 @@ bool FEdModeLandscape::InputKey( FEditorViewportClient* ViewportClient, FViewpor
 }
 
 /** FEdMode: Called when mouse drag input is applied */
-bool FEdModeLandscape::InputDelta( FEditorViewportClient* InViewportClient, FViewport* InViewport, FVector& InDrag, FRotator& InRot, FVector& InScale )
+bool FEdModeLandscape::InputDelta(FEditorViewportClient* InViewportClient, FViewport* InViewport, FVector& InDrag, FRotator& InRot, FVector& InScale)
 {
 	if (GEditor->PlayWorld != NULL)
 	{
@@ -1378,26 +1341,26 @@ bool FEdModeLandscape::InputDelta( FEditorViewportClient* InViewportClient, FVie
 			case ELandscapeEdge::X_Negative:
 			case ELandscapeEdge::X_Negative_Y_Negative:
 			case ELandscapeEdge::X_Negative_Y_Positive:
-				{
-					const int32 InitialComponentCountX = UISettings->NewLandscape_ComponentCount.X;
-					const int32 Delta = FMath::RoundToInt(HitLocation.X + (float)InitialComponentCountX / 2);
-					UISettings->NewLandscape_ComponentCount.X = InitialComponentCountX - Delta;
-					UISettings->NewLandscape_ClampSize();
-					const int32 ActualDelta = UISettings->NewLandscape_ComponentCount.X - InitialComponentCountX;
-					UISettings->NewLandscape_Location -= Transform.TransformVector(FVector(((float)ActualDelta / 2), 0, 0));
-				}
+			{
+				const int32 InitialComponentCountX = UISettings->NewLandscape_ComponentCount.X;
+				const int32 Delta = FMath::RoundToInt(HitLocation.X + (float)InitialComponentCountX / 2);
+				UISettings->NewLandscape_ComponentCount.X = InitialComponentCountX - Delta;
+				UISettings->NewLandscape_ClampSize();
+				const int32 ActualDelta = UISettings->NewLandscape_ComponentCount.X - InitialComponentCountX;
+				UISettings->NewLandscape_Location -= Transform.TransformVector(FVector(((float)ActualDelta / 2), 0, 0));
+			}
 				break;
 			case ELandscapeEdge::X_Positive:
 			case ELandscapeEdge::X_Positive_Y_Negative:
 			case ELandscapeEdge::X_Positive_Y_Positive:
-				{
-					const int32 InitialComponentCountX = UISettings->NewLandscape_ComponentCount.X;
-					int32 Delta = FMath::RoundToInt(HitLocation.X - (float)InitialComponentCountX / 2);
-					UISettings->NewLandscape_ComponentCount.X = InitialComponentCountX + Delta;
-					UISettings->NewLandscape_ClampSize();
-					const int32 ActualDelta = UISettings->NewLandscape_ComponentCount.X - InitialComponentCountX;
-					UISettings->NewLandscape_Location += Transform.TransformVector(FVector(((float)ActualDelta / 2), 0, 0));
-				}
+			{
+				const int32 InitialComponentCountX = UISettings->NewLandscape_ComponentCount.X;
+				int32 Delta = FMath::RoundToInt(HitLocation.X - (float)InitialComponentCountX / 2);
+				UISettings->NewLandscape_ComponentCount.X = InitialComponentCountX + Delta;
+				UISettings->NewLandscape_ClampSize();
+				const int32 ActualDelta = UISettings->NewLandscape_ComponentCount.X - InitialComponentCountX;
+				UISettings->NewLandscape_Location += Transform.TransformVector(FVector(((float)ActualDelta / 2), 0, 0));
+			}
 				break;
 			}
 
@@ -1406,26 +1369,26 @@ bool FEdModeLandscape::InputDelta( FEditorViewportClient* InViewportClient, FVie
 			case ELandscapeEdge::Y_Negative:
 			case ELandscapeEdge::X_Negative_Y_Negative:
 			case ELandscapeEdge::X_Positive_Y_Negative:
-				{
-					const int32 InitialComponentCountY = UISettings->NewLandscape_ComponentCount.Y;
-					int32 Delta = FMath::RoundToInt(HitLocation.Y + (float)InitialComponentCountY / 2);
-					UISettings->NewLandscape_ComponentCount.Y = InitialComponentCountY - Delta;
-					UISettings->NewLandscape_ClampSize();
-					const int32 ActualDelta = UISettings->NewLandscape_ComponentCount.Y - InitialComponentCountY;
-					UISettings->NewLandscape_Location -= Transform.TransformVector(FVector(0, (float)ActualDelta / 2, 0));
-				}
+			{
+				const int32 InitialComponentCountY = UISettings->NewLandscape_ComponentCount.Y;
+				int32 Delta = FMath::RoundToInt(HitLocation.Y + (float)InitialComponentCountY / 2);
+				UISettings->NewLandscape_ComponentCount.Y = InitialComponentCountY - Delta;
+				UISettings->NewLandscape_ClampSize();
+				const int32 ActualDelta = UISettings->NewLandscape_ComponentCount.Y - InitialComponentCountY;
+				UISettings->NewLandscape_Location -= Transform.TransformVector(FVector(0, (float)ActualDelta / 2, 0));
+			}
 				break;
 			case ELandscapeEdge::Y_Positive:
 			case ELandscapeEdge::X_Negative_Y_Positive:
 			case ELandscapeEdge::X_Positive_Y_Positive:
-				{
-					const int32 InitialComponentCountY = UISettings->NewLandscape_ComponentCount.Y;
-					int32 Delta = FMath::RoundToInt(HitLocation.Y - (float)InitialComponentCountY / 2);
-					UISettings->NewLandscape_ComponentCount.Y = InitialComponentCountY + Delta;
-					UISettings->NewLandscape_ClampSize();
-					const int32 ActualDelta = UISettings->NewLandscape_ComponentCount.Y - InitialComponentCountY;
-					UISettings->NewLandscape_Location += Transform.TransformVector(FVector(0, (float)ActualDelta / 2, 0));
-				}
+			{
+				const int32 InitialComponentCountY = UISettings->NewLandscape_ComponentCount.Y;
+				int32 Delta = FMath::RoundToInt(HitLocation.Y - (float)InitialComponentCountY / 2);
+				UISettings->NewLandscape_ComponentCount.Y = InitialComponentCountY + Delta;
+				UISettings->NewLandscape_ClampSize();
+				const int32 ActualDelta = UISettings->NewLandscape_ComponentCount.Y - InitialComponentCountY;
+				UISettings->NewLandscape_Location += Transform.TransformVector(FVector(0, (float)ActualDelta / 2, 0));
+			}
 				break;
 			}
 
@@ -1433,7 +1396,7 @@ bool FEdModeLandscape::InputDelta( FEditorViewportClient* InViewportClient, FVie
 		}
 	}
 
-	if (CurrentToolSet && CurrentToolSet->GetTool() && CurrentToolSet->GetTool()->InputDelta(InViewportClient, InViewport, InDrag, InRot, InScale))
+	if (CurrentTool && CurrentTool && CurrentTool->InputDelta(InViewportClient, InViewport, InDrag, InRot, InScale))
 	{
 		return true;
 	}
@@ -1450,41 +1413,13 @@ void FEdModeLandscape::SetCurrentToolMode(FName ToolModeName, bool bRestoreCurre
 			if (LandscapeToolModes[i].ToolModeName == ToolModeName)
 			{
 				CurrentToolMode = &LandscapeToolModes[i];
-				if (CurrentToolMode->SupportedTargetTypes == ELandscapeToolTargetTypeMask::NA)
-				{
-					CurrentToolTarget.TargetType = (ELandscapeToolTargetType::Type)INDEX_NONE;
-					CurrentToolTarget.LayerInfo = NULL;
-					CurrentToolTarget.LayerName = NAME_None;
-				}
-				else if ( !(CurrentToolMode->SupportedTargetTypes & (1 << CurrentToolTarget.TargetType)) )
-				{
-					bool bFoundValidTarget = false;
-					for (int32 TargetIndex = 0; TargetIndex < LandscapeTargetList.Num(); ++TargetIndex)
-					{
-						if ( (CurrentToolMode->SupportedTargetTypes & (1 << LandscapeTargetList[TargetIndex]->TargetType)) )
-						{
-							check(CurrentToolTarget.LandscapeInfo == LandscapeTargetList[TargetIndex]->LandscapeInfo);
-							CurrentToolTarget.TargetType = LandscapeTargetList[TargetIndex]->TargetType;
-							CurrentToolTarget.LayerInfo = LandscapeTargetList[TargetIndex]->LayerInfoObj;
-							CurrentToolTarget.LayerName = LandscapeTargetList[TargetIndex]->LayerName;
-							bFoundValidTarget = true;
-							break;
-						}
-					}
-					if (!bFoundValidTarget)
-					{
-						CurrentToolTarget.TargetType = (ELandscapeToolTargetType::Type)INDEX_NONE;
-						CurrentToolTarget.LayerInfo = NULL;
-						CurrentToolTarget.LayerName = NAME_None;
-					}
-				}
 				if (bRestoreCurrentTool)
 				{
-					if (CurrentToolMode->CurrentToolSetName == NAME_None)
+					if (CurrentToolMode->CurrentToolName == NAME_None)
 					{
-						CurrentToolMode->CurrentToolSetName = CurrentToolMode->ValidToolSets[0];
+						CurrentToolMode->CurrentToolName = CurrentToolMode->ValidTools[0];
 					}
-					SetCurrentTool(CurrentToolMode->CurrentToolSetName);
+					SetCurrentTool(CurrentToolMode->CurrentToolName);
 				}
 				break;
 			}
@@ -1492,40 +1427,60 @@ void FEdModeLandscape::SetCurrentToolMode(FName ToolModeName, bool bRestoreCurre
 	}
 }
 
-void FEdModeLandscape::SetCurrentTool( FName ToolSetName )
+void FEdModeLandscape::SetCurrentTool(FName ToolName)
 {
-	int32 ToolIndex = 0;
-	for (; ToolIndex < LandscapeToolSets.Num(); ++ToolIndex )
+	// Several tools have identically named versions for sculpting and painting
+	// Prefer the one with the same target type as the current mode
+
+	int32 BackupToolIndex = INDEX_NONE;
+	int32 ToolIndex = INDEX_NONE;
+	for (int32 i = 0; i < LandscapeTools.Num(); ++i)
 	{
-		if (ToolSetName == LandscapeToolSets[ToolIndex].GetToolSetName())
+		FLandscapeTool* Tool = LandscapeTools[i].Get();
+		if (ToolName == Tool->GetToolName())
 		{
-			break;
+			if ((Tool->GetSupportedTargetTypes() & CurrentToolMode->SupportedTargetTypes) != 0)
+			{
+				ToolIndex = i;
+				break;
+			}
+			else if (BackupToolIndex == INDEX_NONE)
+			{
+				BackupToolIndex = i;
+			}
 		}
 	}
 
+	if (ToolIndex == INDEX_NONE)
+	{
+		checkf(BackupToolIndex != INDEX_NONE, TEXT("Tool '%s' not found, please check name is correct!"), *ToolName.ToString());
+		ToolIndex = BackupToolIndex;
+	}
+	check(ToolIndex != INDEX_NONE);
+
 	SetCurrentTool(ToolIndex);
 }
 
-void FEdModeLandscape::SetCurrentTool( int32 ToolIndex )
+void FEdModeLandscape::SetCurrentTool(int32 ToolIndex)
 {
-	if (CurrentToolSet)
+	if (CurrentTool)
 	{
-		CurrentToolSet->PreviousBrushIndex = CurrentBrushSetIndex;
-		if (CurrentToolSet->GetTool())
+		CurrentTool->PreviousBrushIndex = CurrentBrushSetIndex;
+		if (CurrentTool)
 		{
-			CurrentToolSet->GetTool()->ExitTool();
+			CurrentTool->ExitTool();
 		}
 	}
-	CurrentToolIndex = LandscapeToolSets.IsValidIndex(ToolIndex) ? ToolIndex : 0;
-	CurrentToolSet = &LandscapeToolSets[ CurrentToolIndex ];
-	if (!CurrentToolMode->ValidToolSets.Contains(CurrentToolSet->GetToolSetName()))
+	CurrentToolIndex = LandscapeTools.IsValidIndex(ToolIndex) ? ToolIndex : 0;
+	CurrentTool = LandscapeTools[CurrentToolIndex].Get();
+	if (!CurrentToolMode->ValidTools.Contains(CurrentTool->GetToolName()))
 	{
 		// if tool isn't valid for this mode then automatically switch modes
 		// this mostly happens with shortcut keys
 		bool bFoundValidMode = false;
 		for (int32 i = 0; i < LandscapeToolModes.Num(); ++i)
 		{
-			if (LandscapeToolModes[i].ValidToolSets.Contains(CurrentToolSet->GetToolSetName()))
+			if (LandscapeToolModes[i].ValidTools.Contains(CurrentTool->GetToolName()))
 			{
 				SetCurrentToolMode(LandscapeToolModes[i].ToolModeName, false);
 				bFoundValidMode = true;
@@ -1534,28 +1489,56 @@ void FEdModeLandscape::SetCurrentTool( int32 ToolIndex )
 		}
 		check(bFoundValidMode);
 	}
-	if (CurrentToolSet->GetToolSetName() != FName(TEXT("ToolSet_AddComponent")))
+
+	// Set target type appropriate for tool
+	if (CurrentTool->GetSupportedTargetTypes() == ELandscapeToolTargetTypeMask::NA)
 	{
-		LandscapeRenderAddCollision = NULL;
+		CurrentToolTarget.TargetType = ELandscapeToolTargetType::Invalid;
+		CurrentToolTarget.LayerInfo = nullptr;
+		CurrentToolTarget.LayerName = NAME_None;
+	}
+	else
+	{
+		const uint8 TargetTypeMask = CurrentToolMode->SupportedTargetTypes & CurrentTool->GetSupportedTargetTypes();
+		checkSlow(TargetTypeMask != 0);
+
+		if ((TargetTypeMask & ELandscapeToolTargetTypeMask::FromType(CurrentToolTarget.TargetType)) == 0)
+		{
+			auto filter = [TargetTypeMask](const TSharedRef<FLandscapeTargetListInfo>& Target){ return (TargetTypeMask & ELandscapeToolTargetTypeMask::FromType(Target->TargetType)) != 0; };
+			const TSharedRef<FLandscapeTargetListInfo>* Target = LandscapeTargetList.FindByPredicate(filter);
+			if (Target != nullptr)
+			{
+				check(CurrentToolTarget.LandscapeInfo == (*Target)->LandscapeInfo);
+				CurrentToolTarget.TargetType = (*Target)->TargetType;
+				CurrentToolTarget.LayerInfo = (*Target)->LayerInfoObj;
+				CurrentToolTarget.LayerName = (*Target)->LayerName;
+			}
+			else // can happen with for example paint tools if there are no paint layers defined
+			{
+				CurrentToolTarget.TargetType = ELandscapeToolTargetType::Invalid;
+				CurrentToolTarget.LayerInfo = nullptr;
+				CurrentToolTarget.LayerName = NAME_None;
+			}
+		}
 	}
-	CurrentToolSet->SetToolForTarget( CurrentToolTarget );
-	if (CurrentToolSet->GetTool())
+
+	if (CurrentTool)
 	{
-		CurrentToolSet->GetTool()->EnterTool();
+		CurrentTool->EnterTool();
 
-		CurrentToolSet->GetTool()->SetEditRenderType();
-		bool MaskEnabled = CurrentToolSet->GetTool()->SupportsMask() && CurrentToolTarget.LandscapeInfo.IsValid() && CurrentToolTarget.LandscapeInfo->SelectedRegion.Num();
+		CurrentTool->SetEditRenderType();
+		//bool MaskEnabled = CurrentTool->SupportsMask() && CurrentToolTarget.LandscapeInfo.IsValid() && CurrentToolTarget.LandscapeInfo->SelectedRegion.Num();
 	}
-	CurrentToolMode->CurrentToolSetName = CurrentToolSet->GetToolSetName();
+	CurrentToolMode->CurrentToolName = CurrentTool->GetToolName();
 
 	// Set Brush
-	if (!LandscapeBrushSets.IsValidIndex(CurrentToolSet->PreviousBrushIndex))
+	if (!LandscapeBrushSets.IsValidIndex(CurrentTool->PreviousBrushIndex))
 	{
-		SetCurrentBrushSet(CurrentToolSet->ValidBrushes[0]);
+		SetCurrentBrushSet(CurrentTool->ValidBrushes[0]);
 	}
 	else
 	{
-		SetCurrentBrushSet(CurrentToolSet->PreviousBrushIndex);
+		SetCurrentBrushSet(CurrentTool->PreviousBrushIndex);
 	}
 
 	// Update GizmoActor Landscape Target (is this necessary?)
@@ -1589,9 +1572,9 @@ void FEdModeLandscape::SetCurrentBrushSet(int32 BrushSetIndex)
 		LandscapeBrushSets[CurrentBrushSetIndex].PreviousBrushIndex = LandscapeBrushSets[CurrentBrushSetIndex].Brushes.IndexOfByKey(CurrentBrush);
 
 		CurrentBrushSetIndex = BrushSetIndex;
-		if (CurrentToolSet)
+		if (CurrentTool)
 		{
-			CurrentToolSet->PreviousBrushIndex = BrushSetIndex;
+			CurrentTool->PreviousBrushIndex = BrushSetIndex;
 		}
 
 		SetCurrentBrush(LandscapeBrushSets[CurrentBrushSetIndex].PreviousBrushIndex);
@@ -1619,7 +1602,7 @@ void FEdModeLandscape::SetCurrentBrush(int32 BrushIndex)
 		CurrentBrush->EnterBrush();
 
 		if (Toolkit.IsValid())
-{
+		{
 			StaticCastSharedPtr<FLandscapeToolKit>(Toolkit)->NotifyBrushChanged();
 		}
 	}
@@ -1649,7 +1632,7 @@ int32 FEdModeLandscape::UpdateLandscapeList()
 {
 	LandscapeList.Empty();
 
-	if( !CurrentGizmoActor.IsValid() )
+	if (!CurrentGizmoActor.IsValid())
 	{
 		ALandscapeGizmoActiveActor* GizmoActor = NULL;
 		for (FActorIterator It(GetWorld()); It; ++It)
@@ -1670,7 +1653,7 @@ int32 FEdModeLandscape::UpdateLandscapeList()
 		for (auto It = GetWorld()->LandscapeInfoMap.CreateIterator(); It; ++It)
 		{
 			ULandscapeInfo* LandscapeInfo = It.Value();
-			if( LandscapeInfo )
+			if (LandscapeInfo)
 			{
 				if (CurrentToolTarget.LandscapeInfo == LandscapeInfo)
 				{
@@ -1692,7 +1675,7 @@ int32 FEdModeLandscape::UpdateLandscapeList()
 				}
 
 				LandscapeList.Add(FLandscapeListInfo(*LandscapeInfo->GetLandscapeProxy()->GetName(), LandscapeInfo,
-									LandscapeInfo->ComponentSizeQuads, LandscapeInfo->ComponentNumSubsections, Width, Height));
+					LandscapeInfo->ComponentSizeQuads, LandscapeInfo->ComponentNumSubsections, Width, Height));
 				Index++;
 			}
 		}
@@ -1702,16 +1685,16 @@ int32 FEdModeLandscape::UpdateLandscapeList()
 	{
 		if (LandscapeList.Num() > 0)
 		{
-			if (CurrentToolSet != NULL)
+			if (CurrentTool != NULL)
 			{
 				CurrentBrush->LeaveBrush();
-				CurrentToolSet->GetTool()->ExitTool();
+				CurrentTool->ExitTool();
 			}
 			CurrentToolTarget.LandscapeInfo = LandscapeList[0].Info;
 			CurrentIndex = 0;
-			if (CurrentToolSet != NULL)
+			if (CurrentTool != NULL)
 			{
-				CurrentToolSet->GetTool()->EnterTool();
+				CurrentTool->EnterTool();
 				CurrentBrush->EnterBrush();
 			}
 		}
@@ -1727,16 +1710,31 @@ int32 FEdModeLandscape::UpdateLandscapeList()
 void FEdModeLandscape::UpdateTargetList()
 {
 	LandscapeTargetList.Empty();
-	
-	if( CurrentToolTarget.LandscapeInfo.IsValid() &&
-		CurrentToolTarget.LandscapeInfo->GetLandscapeProxy() )
+
+	if (CurrentToolTarget.LandscapeInfo.IsValid() &&
+		CurrentToolTarget.LandscapeInfo->GetLandscapeProxy())
 	{
 		CachedLandscapeMaterial = CurrentToolTarget.LandscapeInfo->GetLandscapeProxy()->GetLandscapeMaterial();
 
 		bool bFoundSelected = false;
+
 		// Add heightmap
 		LandscapeTargetList.Add(MakeShareable(new FLandscapeTargetListInfo(LOCTEXT("Heightmap", "Heightmap"), ELandscapeToolTargetType::Heightmap, CurrentToolTarget.LandscapeInfo.Get())));
 
+		if (CurrentToolTarget.TargetType == ELandscapeToolTargetType::Heightmap)
+		{
+			bFoundSelected = true;
+		}
+
+		// Add visibility
+		FLandscapeInfoLayerSettings VisibilitySettings(ALandscapeProxy::VisibilityLayer, CurrentToolTarget.LandscapeInfo->GetLandscapeProxy());
+		LandscapeTargetList.Add(MakeShareable(new FLandscapeTargetListInfo(LOCTEXT("Visibility", "Visibility"), ELandscapeToolTargetType::Visibility, VisibilitySettings)));
+
+		if (CurrentToolTarget.TargetType == ELandscapeToolTargetType::Visibility)
+		{
+			bFoundSelected = true;
+		}
+
 		// Add layers
 		UTexture2D* ThumbnailWeightmap = NULL;
 		UTexture2D* ThumbnailHeightmap = NULL;
@@ -1745,15 +1743,16 @@ void FEdModeLandscape::UpdateTargetList()
 		{
 			FLandscapeInfoLayerSettings& LayerSettings = *It;
 
-			if( LayerSettings.LayerInfoObj == ALandscapeProxy::DataLayer )
+			FName LayerName = LayerSettings.GetLayerName();
+
+			if (LayerSettings.LayerInfoObj == ALandscapeProxy::VisibilityLayer)
 			{
-				// don't list the visibility layer in the target layer list.
+				// Already handled above
 				continue;
 			}
 
-			FName LayerName = LayerSettings.GetLayerName();
-
-			if (CurrentToolTarget.TargetType == ELandscapeToolTargetType::Weightmap &&
+			if (!bFoundSelected &&
+				CurrentToolTarget.TargetType == ELandscapeToolTargetType::Weightmap &&
 				CurrentToolTarget.LayerInfo == LayerSettings.LayerInfoObj &&
 				CurrentToolTarget.LayerName == LayerSettings.LayerName)
 			{
@@ -1761,13 +1760,13 @@ void FEdModeLandscape::UpdateTargetList()
 			}
 
 			// Ensure thumbnails are up valid
-			if( LayerSettings.ThumbnailMIC == NULL )
+			if (LayerSettings.ThumbnailMIC == NULL)
 			{
-				if( ThumbnailWeightmap == NULL )
+				if (ThumbnailWeightmap == NULL)
 				{
 					ThumbnailWeightmap = LoadObject<UTexture2D>(NULL, TEXT("/Engine/EditorLandscapeResources/LandscapeThumbnailWeightmap.LandscapeThumbnailWeightmap"), NULL, LOAD_None, NULL);
 				}
-				if( ThumbnailHeightmap == NULL )
+				if (ThumbnailHeightmap == NULL)
 				{
 					ThumbnailHeightmap = LoadObject<UTexture2D>(NULL, TEXT("/Engine/EditorLandscapeResources/LandscapeThumbnailHeightmap.LandscapeThumbnailHeightmap"), NULL, LOAD_None, NULL);
 				}
@@ -1781,11 +1780,11 @@ void FEdModeLandscape::UpdateTargetList()
 			LandscapeTargetList.Add(MakeShareable(new FLandscapeTargetListInfo(FText::FromName(LayerName), ELandscapeToolTargetType::Weightmap, LayerSettings)));
 		}
 
-		if( !bFoundSelected )
+		if (!bFoundSelected)
 		{
 			CurrentToolTarget.TargetType = ELandscapeToolTargetType::Heightmap;
-			CurrentToolTarget.LayerInfo  = NULL;
-			CurrentToolTarget.LayerName  = NAME_None;
+			CurrentToolTarget.LayerInfo = NULL;
+			CurrentToolTarget.LayerName = NAME_None;
 		}
 	}
 
@@ -1819,10 +1818,10 @@ void FEdModeLandscape::OnMaterialCompilationFinished(UMaterialInterface* Materia
 }
 
 /** FEdMode: Render the mesh paint tool */
-void FEdModeLandscape::Render( const FSceneView* View, FViewport* Viewport, FPrimitiveDrawInterface* PDI )
+void FEdModeLandscape::Render(const FSceneView* View, FViewport* Viewport, FPrimitiveDrawInterface* PDI)
 {
 	/** Call parent implementation */
-	FEdMode::Render( View, Viewport, PDI );
+	FEdMode::Render(View, Viewport, PDI);
 
 	if (GEditor->PlayWorld != NULL)
 	{
@@ -1832,11 +1831,11 @@ void FEdModeLandscape::Render( const FSceneView* View, FViewport* Viewport, FPri
 	if (NewLandscapePreviewMode != ENewLandscapePreviewMode::None)
 	{
 		static const float        CornerSize = 0.33f;
-		static const FLinearColor CornerColour          (1.0f,  1.0f,  0.5f);
-		static const FLinearColor EdgeColour            (1.0f,  1.0f,  0.0f);
-		static const FLinearColor ComponentBorderColour (0.0f,  0.85f, 0.0f);
-		static const FLinearColor SectionBorderColour   (0.0f,  0.4f,  0.0f);
-		static const FLinearColor InnerColour           (0.0f,  0.25f, 0.0f);
+		static const FLinearColor CornerColour(1.0f, 1.0f, 0.5f);
+		static const FLinearColor EdgeColour(1.0f, 1.0f, 0.0f);
+		static const FLinearColor ComponentBorderColour(0.0f, 0.85f, 0.0f);
+		static const FLinearColor SectionBorderColour(0.0f, 0.4f, 0.0f);
+		static const FLinearColor InnerColour(0.0f, 0.25f, 0.0f);
 
 		const ELevelViewportType ViewportType = ((FEditorViewportClient*)Viewport->GetClient())->ViewportType;
 
@@ -1847,10 +1846,10 @@ void FEdModeLandscape::Render( const FSceneView* View, FViewport* Viewport, FPri
 		const FVector Offset = UISettings->NewLandscape_Location + FTransform(UISettings->NewLandscape_Rotation, FVector::ZeroVector, UISettings->NewLandscape_Scale).TransformVector(FVector(-ComponentCountX * ComponentSize / 2, -ComponentCountY * ComponentSize / 2, 0));
 		const FTransform Transform = FTransform(UISettings->NewLandscape_Rotation, Offset, UISettings->NewLandscape_Scale);
 
-		if( NewLandscapePreviewMode == ENewLandscapePreviewMode::ImportLandscape )
+		if (NewLandscapePreviewMode == ENewLandscapePreviewMode::ImportLandscape)
 		{
 			const TArray<uint16>& ImportHeights = UISettings->GetImportLandscapeData();
-			if( ImportHeights.Num() != 0)
+			if (ImportHeights.Num() != 0)
 			{
 				const float InvQuadsPerComponent = 1.f / (float)QuadsPerComponent;
 				const int32 SizeX = ComponentCountX * QuadsPerComponent + 1;
@@ -1860,20 +1859,20 @@ void FEdModeLandscape::Render( const FSceneView* View, FViewport* Viewport, FPri
 				const int32 OffsetX = (SizeX - ImportSizeX) / 2;
 				const int32 OffsetY = (SizeY - ImportSizeY) / 2;
 
-				for (int32 ComponentY = 0; ComponentY < ComponentCountY ; ComponentY++)
+				for (int32 ComponentY = 0; ComponentY < ComponentCountY; ComponentY++)
 				{
 					const int32 Y0 = ComponentY * QuadsPerComponent;
 					const int32 Y1 = (ComponentY + 1) * QuadsPerComponent;
 
-					const int32 ImportY0 = FMath::Clamp<int32>(Y0 - OffsetY, 0, ImportSizeY-1);
-					const int32 ImportY1 = FMath::Clamp<int32>(Y1 - OffsetY, 0, ImportSizeY-1);
+					const int32 ImportY0 = FMath::Clamp<int32>(Y0 - OffsetY, 0, ImportSizeY - 1);
+					const int32 ImportY1 = FMath::Clamp<int32>(Y1 - OffsetY, 0, ImportSizeY - 1);
 
-					for (int32 ComponentX = 0; ComponentX < ComponentCountX ; ComponentX++)
+					for (int32 ComponentX = 0; ComponentX < ComponentCountX; ComponentX++)
 					{
 						const int32 X0 = ComponentX * QuadsPerComponent;
 						const int32 X1 = (ComponentX + 1) * QuadsPerComponent;
-						const int32 ImportX0 = FMath::Clamp<int32>(X0 - OffsetX, 0, ImportSizeX-1);
-						const int32 ImportX1 = FMath::Clamp<int32>(X1 - OffsetX, 0, ImportSizeX-1);
+						const int32 ImportX0 = FMath::Clamp<int32>(X0 - OffsetX, 0, ImportSizeX - 1);
+						const int32 ImportX1 = FMath::Clamp<int32>(X1 - OffsetX, 0, ImportSizeX - 1);
 						const float Z00 = ((float)ImportHeights[ImportX0 + ImportY0 * ImportSizeX] - 32768.f) * LANDSCAPE_ZSCALE;
 						const float Z01 = ((float)ImportHeights[ImportX0 + ImportY1 * ImportSizeX] - 32768.f) * LANDSCAPE_ZSCALE;
 						const float Z10 = ((float)ImportHeights[ImportX1 + ImportY0 * ImportSizeX] - 32768.f) * LANDSCAPE_ZSCALE;
@@ -1886,7 +1885,7 @@ void FEdModeLandscape::Render( const FSceneView* View, FViewport* Viewport, FPri
 							PDI->SetHitProxy(NULL);
 						}
 
-						if (ComponentX == ComponentCountX-1)
+						if (ComponentX == ComponentCountX - 1)
 						{
 							PDI->SetHitProxy(new HNewLandscapeGrabHandleProxy(ELandscapeEdge::X_Positive));
 							PDI->DrawLine(Transform.TransformPosition(FVector(X1, Y0, Z10)), Transform.TransformPosition(FVector(X1, Y1, Z11)), ComponentBorderColour, SDPG_Foreground);
@@ -1904,7 +1903,7 @@ void FEdModeLandscape::Render( const FSceneView* View, FViewport* Viewport, FPri
 							PDI->SetHitProxy(NULL);
 						}
 
-						if (ComponentY == ComponentCountY-1)
+						if (ComponentY == ComponentCountY - 1)
 						{
 							PDI->SetHitProxy(new HNewLandscapeGrabHandleProxy(ELandscapeEdge::Y_Positive));
 							PDI->DrawLine(Transform.TransformPosition(FVector(X0, Y1, Z01)), Transform.TransformPosition(FVector(X1, Y1, Z11)), ComponentBorderColour, SDPG_Foreground);
@@ -1914,16 +1913,16 @@ void FEdModeLandscape::Render( const FSceneView* View, FViewport* Viewport, FPri
 						{
 							PDI->DrawLine(Transform.TransformPosition(FVector(X0, Y1, Z01)), Transform.TransformPosition(FVector(X1, Y1, Z11)), ComponentBorderColour, SDPG_Foreground);
 						}
-				
+
 						// intra-component lines - too slow for big landscapes
-						/* 
+						/*
 						for (int32 x=1;x<QuadsPerComponent;x++)
 						{
-							PDI->DrawLine(Transform.TransformPosition(FVector(X0+x, Y0, FMath::Lerp(Z00,Z10,(float)x*InvQuadsPerComponent))), Transform.TransformPosition(FVector(X0+x, Y1, FMath::Lerp(Z01,Z11,(float)x*InvQuadsPerComponent))), ComponentBorderColour, SDPG_World);
-						}		
+						PDI->DrawLine(Transform.TransformPosition(FVector(X0+x, Y0, FMath::Lerp(Z00,Z10,(float)x*InvQuadsPerComponent))), Transform.TransformPosition(FVector(X0+x, Y1, FMath::Lerp(Z01,Z11,(float)x*InvQuadsPerComponent))), ComponentBorderColour, SDPG_World);
+						}
 						for (int32 y=1;y<QuadsPerComponent;y++)
 						{
-							PDI->DrawLine(Transform.TransformPosition(FVector(X0, Y0+y, FMath::Lerp(Z00,Z01,(float)y*InvQuadsPerComponent))), Transform.TransformPosition(FVector(X1, Y0+y, FMath::Lerp(Z10,Z11,(float)y*InvQuadsPerComponent))), ComponentBorderColour, SDPG_World);
+						PDI->DrawLine(Transform.TransformPosition(FVector(X0, Y0+y, FMath::Lerp(Z00,Z01,(float)y*InvQuadsPerComponent))), Transform.TransformPosition(FVector(X1, Y0+y, FMath::Lerp(Z10,Z11,(float)y*InvQuadsPerComponent))), ComponentBorderColour, SDPG_World);
 						}
 						*/
 					}
@@ -2045,27 +2044,27 @@ void FEdModeLandscape::Render( const FSceneView* View, FViewport* Viewport, FPri
 
 		for (int32 i = 0; i < 8; ++i)
 		{
-			MeshBuilder.AddVertex(CurrentGizmoActor->FrustumVerts[i], FVector2D(0, 0), FVector(1,0,0), FVector(0,1,0), FVector(0,0,1), FColor(255,255,255));
+			MeshBuilder.AddVertex(CurrentGizmoActor->FrustumVerts[i], FVector2D(0, 0), FVector(1, 0, 0), FVector(0, 1, 0), FVector(0, 0, 1), FColor(255, 255, 255));
 		}
 
 		// Upper box.
-		MeshBuilder.AddTriangle( 0, 2, 1 );
-		MeshBuilder.AddTriangle( 0, 3, 2 );
+		MeshBuilder.AddTriangle(0, 2, 1);
+		MeshBuilder.AddTriangle(0, 3, 2);
 		// Lower box.
-		MeshBuilder.AddTriangle( 4, 6, 5 );
-		MeshBuilder.AddTriangle( 4, 7, 6 );
+		MeshBuilder.AddTriangle(4, 6, 5);
+		MeshBuilder.AddTriangle(4, 7, 6);
 		// Others
-		MeshBuilder.AddTriangle( 1, 4, 0 );
-		MeshBuilder.AddTriangle( 1, 5, 4 );
+		MeshBuilder.AddTriangle(1, 4, 0);
+		MeshBuilder.AddTriangle(1, 5, 4);
 
-		MeshBuilder.AddTriangle( 3, 6, 2 );
-		MeshBuilder.AddTriangle( 3, 7, 6 );
+		MeshBuilder.AddTriangle(3, 6, 2);
+		MeshBuilder.AddTriangle(3, 7, 6);
 
-		MeshBuilder.AddTriangle( 2, 5, 1 );
-		MeshBuilder.AddTriangle( 2, 6, 5 );
+		MeshBuilder.AddTriangle(2, 5, 1);
+		MeshBuilder.AddTriangle(2, 6, 5);
 
-		MeshBuilder.AddTriangle( 0, 7, 3 );
-		MeshBuilder.AddTriangle( 0, 4, 7 );
+		MeshBuilder.AddTriangle(0, 7, 3);
+		MeshBuilder.AddTriangle(0, 4, 7);
 
 		PDI->SetHitProxy(new HTranslucentActor(CurrentGizmoActor.Get(), NULL));
 		MeshBuilder.Draw(PDI, FMatrix::Identity, GizmoMaterial->GetRenderProxy(false), SDPG_World, true);
@@ -2073,14 +2072,14 @@ void FEdModeLandscape::Render( const FSceneView* View, FViewport* Viewport, FPri
 	}
 
 	// Override Rendering for Splines Tool
-	if (CurrentToolSet && CurrentToolSet->GetTool())
+	if (CurrentTool && CurrentTool)
 	{
-		CurrentToolSet->GetTool()->Render(View, Viewport, PDI);
+		CurrentTool->Render(View, Viewport, PDI);
 	}
 }
 
 /** FEdMode: Render HUD elements for this tool */
-void FEdModeLandscape::DrawHUD( FEditorViewportClient* ViewportClient, FViewport* Viewport, const FSceneView* View, FCanvas* Canvas )
+void FEdModeLandscape::DrawHUD(FEditorViewportClient* ViewportClient, FViewport* Viewport, const FSceneView* View, FCanvas* Canvas)
 {
 
 }
@@ -2093,12 +2092,12 @@ bool FEdModeLandscape::UsesTransformWidget() const
 	}
 
 	// Override Widget for Splines Tool
-	if (CurrentToolSet && CurrentToolSet->GetTool() && CurrentToolSet->GetTool()->UsesTransformWidget())
+	if (CurrentTool && CurrentTool && CurrentTool->UsesTransformWidget())
 	{
 		return true;
 	}
 
-	return (CurrentGizmoActor.IsValid() && CurrentGizmoActor->IsSelected() && (GLandscapeEditRenderMode & ELandscapeEditRenderMode::Gizmo) );
+	return (CurrentGizmoActor.IsValid() && CurrentGizmoActor->IsSelected() && (GLandscapeEditRenderMode & ELandscapeEditRenderMode::Gizmo));
 }
 
 bool FEdModeLandscape::ShouldDrawWidget() const
@@ -2106,18 +2105,18 @@ bool FEdModeLandscape::ShouldDrawWidget() const
 	return UsesTransformWidget();
 }
 
-EAxisList::Type FEdModeLandscape::GetWidgetAxisToDraw( FWidget::EWidgetMode InWidgetMode ) const
+EAxisList::Type FEdModeLandscape::GetWidgetAxisToDraw(FWidget::EWidgetMode InWidgetMode) const
 {
 	if (NewLandscapePreviewMode == ENewLandscapePreviewMode::None)
 	{
 		// Override Widget for Splines Tool
-		if (CurrentToolSet && CurrentToolSet->GetTool())
+		if (CurrentTool && CurrentTool)
 		{
-			return CurrentToolSet->GetTool()->GetWidgetAxisToDraw(InWidgetMode);
+			return CurrentTool->GetWidgetAxisToDraw(InWidgetMode);
 		}
 	}
 
-	switch(InWidgetMode)
+	switch (InWidgetMode)
 	{
 	case FWidget::WM_Translate:
 		return EAxisList::XYZ;
@@ -2143,22 +2142,22 @@ FVector FEdModeLandscape::GetWidgetLocation() const
 		{
 			// Apply Landscape transformation when it is available
 			ULandscapeInfo* Info = CurrentGizmoActor->TargetLandscapeInfo;
-			return CurrentGizmoActor->GetActorLocation() 
+			return CurrentGizmoActor->GetActorLocation()
 				+ FRotationMatrix(Info->GetLandscapeProxy()->GetActorRotation()).TransformPosition(FVector(0, 0, CurrentGizmoActor->GetLength()));
 		}
 		return CurrentGizmoActor->GetActorLocation();
 	}
 
 	// Override Widget for Splines Tool
-	if (CurrentToolSet && CurrentToolSet->GetTool())
+	if (CurrentTool && CurrentTool)
 	{
-		return CurrentToolSet->GetTool()->GetWidgetLocation();
+		return CurrentTool->GetWidgetLocation();
 	}
 
 	return FEdMode::GetWidgetLocation();
 }
 
-bool FEdModeLandscape::GetCustomDrawingCoordinateSystem( FMatrix& InMatrix, void* InData )
+bool FEdModeLandscape::GetCustomDrawingCoordinateSystem(FMatrix& InMatrix, void* InData)
 {
 	if (NewLandscapePreviewMode != ENewLandscapePreviewMode::None)
 	{
@@ -2167,22 +2166,22 @@ bool FEdModeLandscape::GetCustomDrawingCoordinateSystem( FMatrix& InMatrix, void
 	}
 
 	// Override Widget for Splines Tool
-	if (CurrentToolSet && CurrentToolSet->GetTool())
+	if (CurrentTool && CurrentTool)
 	{
-		InMatrix = CurrentToolSet->GetTool()->GetWidgetRotation();
+		InMatrix = CurrentTool->GetWidgetRotation();
 		return true;
 	}
 
 	return false;
 }
 
-bool FEdModeLandscape::GetCustomInputCoordinateSystem( FMatrix& InMatrix, void* InData )
+bool FEdModeLandscape::GetCustomInputCoordinateSystem(FMatrix& InMatrix, void* InData)
 {
 	return GetCustomDrawingCoordinateSystem(InMatrix, InData);
 }
 
 /** FEdMode: Handling SelectActor */
-bool FEdModeLandscape::Select( AActor* InActor, bool bInSelected )
+bool FEdModeLandscape::Select(AActor* InActor, bool bInSelected)
 {
 	if (GEditor->PlayWorld != NULL)
 	{
@@ -2217,7 +2216,7 @@ bool FEdModeLandscape::Select( AActor* InActor, bool bInSelected )
 }
 
 /** FEdMode: Check to see if an actor can be selected in this mode - no side effects */
-bool FEdModeLandscape::IsSelectionAllowed( AActor* InActor, bool bInSelection ) const
+bool FEdModeLandscape::IsSelectionAllowed(AActor* InActor, bool bInSelection) const
 {
 	if (GEditor->PlayWorld != NULL)
 	{
@@ -2230,12 +2229,12 @@ bool FEdModeLandscape::IsSelectionAllowed( AActor* InActor, bool bInSelection )
 	}
 
 	// Override Selection for Splines Tool
-	if (CurrentToolSet && CurrentToolSet->GetTool() && CurrentToolSet->GetTool()->OverrideSelection())
+	if (CurrentTool && CurrentTool && CurrentTool->OverrideSelection())
 	{
-		return CurrentToolSet->GetTool()->IsSelectionAllowed(InActor, bInSelection);
+		return CurrentTool->IsSelectionAllowed(InActor, bInSelection);
 	}
 
-	if (InActor->IsA(ALandscapeProxy::StaticClass()) )
+	if (InActor->IsA(ALandscapeProxy::StaticClass()))
 	{
 		return true;
 	}
@@ -2259,17 +2258,17 @@ void FEdModeLandscape::ActorSelectionChangeNotify()
 		GEditor->SelectNone(false, true);
 		GEditor->SelectActor(CurrentGizmoActor.Get(), true, false, true);
 	}
-/*
-	USelection* EditorSelection = GEditor->GetSelectedActors();
-	for ( FSelectionIterator Itor(EditorSelection) ; Itor ; ++Itor )
-	{
+	/*
+		USelection* EditorSelection = GEditor->GetSelectedActors();
+		for ( FSelectionIterator Itor(EditorSelection) ; Itor ; ++Itor )
+		{
 		if (((*Itor)->IsA(ALandscapeGizmoActor::StaticClass())) )
 		{
-			bIsGizmoSelected = true;
-			break;
+		bIsGizmoSelected = true;
+		break;
 		}
-	}
-*/
+		}
+		*/
 }
 
 void FEdModeLandscape::ActorMoveNotify()
@@ -2285,7 +2284,7 @@ void FEdModeLandscape::ForceRealTimeViewports(const bool bEnable, const bool bSt
 	if (LevelEditor.IsValid())
 	{
 		TArray<TSharedPtr<ILevelViewport>> Viewports = LevelEditor->GetViewports();
-		for (const TSharedPtr<ILevelViewport>& ViewportWindow :  Viewports)
+		for (const TSharedPtr<ILevelViewport>& ViewportWindow : Viewports)
 		{
 			if (ViewportWindow.IsValid())
 			{
@@ -2314,7 +2313,7 @@ void FEdModeLandscape::ReimportData(const FLandscapeTargetListInfo& TargetInfo)
 	}
 	else
 	{
-		FMessageDialog::Open( EAppMsgType::Ok, NSLOCTEXT("UnrealEd", "LandscapeReImport_BadFileName", "Reimport Source Filename is invalid") );
+		FMessageDialog::Open(EAppMsgType::Ok, NSLOCTEXT("UnrealEd", "LandscapeReImport_BadFileName", "Reimport Source Filename is invalid"));
 	}
 }
 
@@ -2335,10 +2334,10 @@ void FEdModeLandscape::ImportData(const FLandscapeTargetListInfo& TargetInfo, co
 
 				if (ImageWrapper->SetCompressed(Data.GetData(), Data.Num()))
 				{
-					if (ImageWrapper->GetWidth() != (1+MaxX-MinX) || ImageWrapper->GetHeight() != (1+MaxY-MinY))
+					if (ImageWrapper->GetWidth() != (1 + MaxX - MinX) || ImageWrapper->GetHeight() != (1 + MaxY - MinY))
 					{
-						FMessageDialog::Open( EAppMsgType::Ok,
-							FText::Format( NSLOCTEXT("UnrealEd", "LandscapeReImport_BadFileSize", "{0}'s filesize does not match with current Landscape extent"), FText::FromString(Filename)) );
+						FMessageDialog::Open(EAppMsgType::Ok,
+							FText::Format(NSLOCTEXT("UnrealEd", "LandscapeReImport_BadFileSize", "{0}'s filesize does not match with current Landscape extent"), FText::FromString(Filename)));
 						return;
 					}
 
@@ -2396,36 +2395,36 @@ void FEdModeLandscape::ImportData(const FLandscapeTargetListInfo& TargetInfo, co
 				}
 				else
 				{
-					FMessageDialog::Open( EAppMsgType::Ok,
-						FText::Format(LOCTEXT("ImportData_CorruptPngFile", "Import of {0} failed (corrupt png?)"), FText::FromString(Filename)) );
+					FMessageDialog::Open(EAppMsgType::Ok,
+						FText::Format(LOCTEXT("ImportData_CorruptPngFile", "Import of {0} failed (corrupt png?)"), FText::FromString(Filename)));
 					return;
 				}
 			}
 
 			if (TargetInfo.TargetType == ELandscapeToolTargetType::Heightmap)
 			{
-				if (Data.Num() == (1+MaxX-MinX)*(1+MaxY-MinY)*2)
+				if (Data.Num() == (1 + MaxX - MinX)*(1 + MaxY - MinY) * 2)
 				{
 					FHeightmapAccessor<false> HeightmapAccessor(LandscapeInfo);
 					HeightmapAccessor.SetData(MinX, MinY, MaxX, MaxY, (uint16*)Data.GetTypedData());
 				}
 				else
 				{
-					FMessageDialog::Open( EAppMsgType::Ok,
-						FText::Format( NSLOCTEXT("UnrealEd", "LandscapeReImport_BadFileSize", "{0}'s filesize does not match with current Landscape extent"), FText::FromString(Filename)) );
+					FMessageDialog::Open(EAppMsgType::Ok,
+						FText::Format(NSLOCTEXT("UnrealEd", "LandscapeReImport_BadFileSize", "{0}'s filesize does not match with current Landscape extent"), FText::FromString(Filename)));
 				}
 			}
 			else
 			{
-				if (Data.Num() == (1+MaxX-MinX)*(1+MaxY-MinY))
+				if (Data.Num() == (1 + MaxX - MinX)*(1 + MaxY - MinY))
 				{
 					FAlphamapAccessor<false, true> AlphamapAccessor(LandscapeInfo, TargetInfo.LayerInfoObj.Get());
 					AlphamapAccessor.SetData(MinX, MinY, MaxX, MaxY, Data.GetTypedData(), ELandscapeLayerPaintingRestriction::None);
 				}
 				else
 				{
-					FMessageDialog::Open( EAppMsgType::Ok,
-						FText::Format( NSLOCTEXT("UnrealEd", "LandscapeReImport_BadFileSize", "{0}'s filesize does not match with current Landscape extent"), FText::FromString(Filename)) );
+					FMessageDialog::Open(EAppMsgType::Ok,
+						FText::Format(NSLOCTEXT("UnrealEd", "LandscapeReImport_BadFileSize", "{0}'s filesize does not match with current Landscape extent"), FText::FromString(Filename)));
 				}
 			}
 		}
@@ -2475,9 +2474,9 @@ ALandscape* FEdModeLandscape::ChangeComponentSetting(int32 NumComponentsX, int32
 				{
 					if (LayerSettings.LayerInfoObj != NULL)
 					{
-						auto ImportLayerInfo = new(ImportLayerInfos) FLandscapeImportLayerInfo(LayerSettings);
+						auto ImportLayerInfo = new(ImportLayerInfos)FLandscapeImportLayerInfo(LayerSettings);
 						ImportLayerInfo->LayerData.AddZeroed(OldVertsX * OldVertsY * sizeof(uint8));
-						
+
 						TMinX = OldMinX; TMinY = OldMinY; TMaxX = OldMaxX; TMaxY = OldMaxY;
 						LandscapeEdit.GetWeightData(LayerSettings.LayerInfoObj, TMinX, TMinY, TMaxX, TMaxY, ImportLayerInfo->LayerData.GetData(), 0);
 
@@ -2517,7 +2516,7 @@ ALandscape* FEdModeLandscape::ChangeComponentSetting(int32 NumComponentsX, int32
 				{
 					if (LayerSettings.LayerInfoObj != NULL)
 					{
-						auto ImportLayerInfo = new(ImportLayerInfos) FLandscapeImportLayerInfo(LayerSettings);
+						auto ImportLayerInfo = new(ImportLayerInfos)FLandscapeImportLayerInfo(LayerSettings);
 						ImportLayerInfo->LayerData.AddZeroed(NewVertsX * NewVertsY * sizeof(uint8));
 
 						TMinX = RequestedMinX; TMinY = RequestedMinY; TMaxX = RequestedMaxX; TMaxY = RequestedMaxY;
@@ -2542,18 +2541,18 @@ ALandscape* FEdModeLandscape::ChangeComponentSetting(int32 NumComponentsX, int32
 			Landscape->CollisionMipLevel = OldLandscapeProxy->CollisionMipLevel;
 			Landscape->Import(FGuid::NewGuid(), NewVertsX, NewVertsY, NumSubsections*SubsectionSizeQuads, NumSubsections, SubsectionSizeQuads, HeightData.GetData(), *OldLandscapeProxy->ReimportHeightmapFilePath, ImportLayerInfos);
 
-			Landscape->MaxLODLevel                 = OldLandscapeProxy->MaxLODLevel;
-			Landscape->ExportLOD                   = OldLandscapeProxy->ExportLOD;
-			Landscape->StaticLightingLOD           = OldLandscapeProxy->StaticLightingLOD;
-			Landscape->DefaultPhysMaterial         = OldLandscapeProxy->DefaultPhysMaterial;
+			Landscape->MaxLODLevel = OldLandscapeProxy->MaxLODLevel;
+			Landscape->ExportLOD = OldLandscapeProxy->ExportLOD;
+			Landscape->StaticLightingLOD = OldLandscapeProxy->StaticLightingLOD;
+			Landscape->DefaultPhysMaterial = OldLandscapeProxy->DefaultPhysMaterial;
 			Landscape->StreamingDistanceMultiplier = OldLandscapeProxy->StreamingDistanceMultiplier;
-			Landscape->LandscapeHoleMaterial       = OldLandscapeProxy->LandscapeHoleMaterial;
-			Landscape->LODDistanceFactor           = OldLandscapeProxy->LODDistanceFactor;
-			Landscape->StaticLightingResolution    = OldLandscapeProxy->StaticLightingResolution;
-			Landscape->bCastStaticShadow           = OldLandscapeProxy->bCastStaticShadow;
-			Landscape->bCastShadowAsTwoSided	   = OldLandscapeProxy->bCastShadowAsTwoSided;
-			Landscape->LightmassSettings           = OldLandscapeProxy->LightmassSettings;
-			Landscape->CollisionThickness          = OldLandscapeProxy->CollisionThickness;
+			Landscape->LandscapeHoleMaterial = OldLandscapeProxy->LandscapeHoleMaterial;
+			Landscape->LODDistanceFactor = OldLandscapeProxy->LODDistanceFactor;
+			Landscape->StaticLightingResolution = OldLandscapeProxy->StaticLightingResolution;
+			Landscape->bCastStaticShadow = OldLandscapeProxy->bCastStaticShadow;
+			Landscape->bCastShadowAsTwoSided = OldLandscapeProxy->bCastShadowAsTwoSided;
+			Landscape->LightmassSettings = OldLandscapeProxy->LightmassSettings;
+			Landscape->CollisionThickness = OldLandscapeProxy->CollisionThickness;
 			Landscape->BodyInstance.SetCollisionProfileName(OldLandscapeProxy->BodyInstance.GetCollisionProfileName());
 			if (Landscape->BodyInstance.DoesUseCollisionProfile() == false)
 			{
@@ -2561,8 +2560,8 @@ ALandscape* FEdModeLandscape::ChangeComponentSetting(int32 NumComponentsX, int32
 				Landscape->BodyInstance.SetObjectType(OldLandscapeProxy->BodyInstance.GetObjectType());
 				Landscape->BodyInstance.SetResponseToChannels(OldLandscapeProxy->BodyInstance.GetResponseToChannels());
 			}
-			Landscape->EditorLayerSettings          = OldLandscapeProxy->EditorLayerSettings;
-			Landscape->bUsedForNavigation           = OldLandscapeProxy->bUsedForNavigation;
+			Landscape->EditorLayerSettings = OldLandscapeProxy->EditorLayerSettings;
+			Landscape->bUsedForNavigation = OldLandscapeProxy->bUsedForNavigation;
 			Landscape->MaxPaintedLayersPerComponent = OldLandscapeProxy->MaxPaintedLayersPerComponent;
 
 			// Clone landscape splines
@@ -2626,8 +2625,8 @@ ALandscape* FEdModeLandscape::ChangeComponentSetting(int32 NumComponentsX, int32
 bool LandscapeEditorUtils::SetHeightmapData(ALandscapeProxy* Landscape, const TArray<uint16>& Data)
 {
 	FIntRect ComponentsRect = Landscape->GetBoundingRect() + Landscape->LandscapeSectionOffset;
-		
-	if (Data.Num() == (1+ComponentsRect.Width())*(1+ComponentsRect.Height()))
+
+	if (Data.Num() == (1 + ComponentsRect.Width())*(1 + ComponentsRect.Height()))
 	{
 		FHeightmapAccessor<false> HeightmapAccessor(Landscape->GetLandscapeInfo());
 		HeightmapAccessor.SetData(ComponentsRect.Min.X, ComponentsRect.Min.Y, ComponentsRect.Max.X, ComponentsRect.Max.Y, Data.GetTypedData());
@@ -2639,9 +2638,9 @@ bool LandscapeEditorUtils::SetHeightmapData(ALandscapeProxy* Landscape, const TA
 
 bool LandscapeEditorUtils::SetWeightmapData(ALandscapeProxy* Landscape, ULandscapeLayerInfoObject* LayerObject, const TArray<uint8>& Data)
 {
-	FIntRect ComponentsRect = Landscape->GetBoundingRect() + Landscape->LandscapeSectionOffset;	
-	
-	if (Data.Num() == (1+ComponentsRect.Width())*(1+ComponentsRect.Height()))
+	FIntRect ComponentsRect = Landscape->GetBoundingRect() + Landscape->LandscapeSectionOffset;
+
+	if (Data.Num() == (1 + ComponentsRect.Width())*(1 + ComponentsRect.Height()))
 	{
 		FAlphamapAccessor<false, true> AlphamapAccessor(Landscape->GetLandscapeInfo(), LayerObject);
 		AlphamapAccessor.SetData(ComponentsRect.Min.X, ComponentsRect.Min.Y, ComponentsRect.Max.X, ComponentsRect.Max.Y, Data.GetTypedData(), ELandscapeLayerPaintingRestriction::None);
