Remove FLandscapeToolSet

diff --git a/Engine/Source/Editor/LandscapeEditor/Private/LandscapeEdModeBrushes.cpp b/Engine/Source/Editor/LandscapeEditor/Private/LandscapeEdModeBrushes.cpp
--- a/Engine/Source/Editor/LandscapeEditor/Private/LandscapeEdModeBrushes.cpp
+++ b/Engine/Source/Editor/LandscapeEditor/Private/LandscapeEdModeBrushes.cpp
@@ -17,11 +17,11 @@
 //
 bool GInLandscapeBrushTransaction = false;
 
-void FLandscapeBrush::BeginStroke(float LandscapeX, float LandscapeY, class FLandscapeTool* CurrentTool)
+void FLandscapeBrush::BeginStroke(float LandscapeX, float LandscapeY, FLandscapeTool* CurrentTool)
 {
 	if (!GInLandscapeBrushTransaction)
 	{
-		GEditor->BeginTransaction( FText::Format( NSLOCTEXT("UnrealEd", "LandscapeMode_EditTransaction", "Landscape Editing: {0}"), CurrentTool->GetDisplayName() ) );
+		GEditor->BeginTransaction(FText::Format(NSLOCTEXT("UnrealEd", "LandscapeMode_EditTransaction", "Landscape Editing: {0}"), CurrentTool->GetDisplayName()));
 		GInLandscapeBrushTransaction = true;
 	}
 }
@@ -49,19 +49,19 @@ protected:
 	UMaterialInterface* BrushMaterial;
 	TMap<ULandscapeComponent*, UMaterialInstanceDynamic*> BrushMaterialInstanceMap;
 
-	virtual float CalculateFalloff( float Distance, float Radius, float Falloff ) = 0;
+	virtual float CalculateFalloff(float Distance, float Radius, float Falloff) = 0;
 
 	/** Protected so that only subclasses can create instances of this class. */
-	FLandscapeBrushCircle(class FEdModeLandscape* InEdMode, UMaterialInterface* InBrushMaterial)
-		:	LastMousePosition(0,0)
-		,	BrushMaterial(NULL)
-		,	EdMode(InEdMode)
+	FLandscapeBrushCircle(FEdModeLandscape* InEdMode, UMaterialInterface* InBrushMaterial)
+		: LastMousePosition(0, 0)
+		, BrushMaterial(NULL)
+		, EdMode(InEdMode)
 	{
 		BrushMaterial = InBrushMaterial; //UMaterialInstanceDynamic::Create(InBrushMaterial, NULL);
 	}
 
 public:
-	class FEdModeLandscape* EdMode;
+	FEdModeLandscape* EdMode;
 
 	// FGCObject interface
 	virtual void AddReferencedObjects(FReferenceCollector& Collector) override
@@ -79,9 +79,9 @@ public:
 
 	virtual void LeaveBrush() override
 	{
-		for( TSet<ULandscapeComponent*>::TIterator It(BrushMaterialComponents); It; ++It )
+		for (TSet<ULandscapeComponent*>::TIterator It(BrushMaterialComponents); It; ++It)
 		{
-			if( (*It)->EditToolRenderData != NULL )
+			if ((*It)->EditToolRenderData != NULL)
 			{
 				(*It)->EditToolRenderData->Update(NULL);
 			}
@@ -93,13 +93,13 @@ public:
 		BrushMaterialComponents.Empty();
 	}
 
-	virtual void BeginStroke(float LandscapeX, float LandscapeY, class FLandscapeTool* CurrentTool) override
+	virtual void BeginStroke(float LandscapeX, float LandscapeY, FLandscapeTool* CurrentTool) override
 	{
-		FLandscapeBrush::BeginStroke(LandscapeX,LandscapeY,CurrentTool);
+		FLandscapeBrush::BeginStroke(LandscapeX, LandscapeY, CurrentTool);
 		LastMousePosition = FVector2D(LandscapeX, LandscapeY);
 	}
 
-	virtual void Tick(FEditorViewportClient* ViewportClient,float DeltaTime) override
+	virtual void Tick(FEditorViewportClient* ViewportClient, float DeltaTime) override
 	{
 		ULandscapeInfo* LandscapeInfo = EdMode->CurrentToolTarget.LandscapeInfo.Get();
 		ALandscapeProxy* Proxy = LandscapeInfo ? LandscapeInfo->GetLandscapeProxy() : NULL;
@@ -110,10 +110,10 @@ public:
 			float Falloff = EdMode->UISettings->BrushFalloff * EdMode->UISettings->BrushRadius / ScaleXY;
 
 			// Set brush material.
-			int32 X1 = FMath::FloorToInt(LastMousePosition.X - (Radius+Falloff));
-			int32 Y1 = FMath::FloorToInt(LastMousePosition.Y - (Radius+Falloff));
-			int32 X2 = FMath::CeilToInt(LastMousePosition.X + (Radius+Falloff));
-			int32 Y2 = FMath::CeilToInt(LastMousePosition.Y + (Radius+Falloff));
+			int32 X1 = FMath::FloorToInt(LastMousePosition.X - (Radius + Falloff));
+			int32 Y1 = FMath::FloorToInt(LastMousePosition.Y - (Radius + Falloff));
+			int32 X2 = FMath::CeilToInt(LastMousePosition.X + (Radius + Falloff));
+			int32 Y2 = FMath::CeilToInt(LastMousePosition.Y + (Radius + Falloff));
 
 			TSet<ULandscapeComponent*> NewComponents;
 
@@ -157,7 +157,7 @@ public:
 			BrushMaterialComponents = MoveTemp(NewComponents);
 
 			// Set params for brush material.
-			FVector WorldLocation = Proxy->LandscapeActorToWorld().TransformPosition(FVector(LastMousePosition.X,LastMousePosition.Y,0));
+			FVector WorldLocation = Proxy->LandscapeActorToWorld().TransformPosition(FVector(LastMousePosition.X, LastMousePosition.Y, 0));
 
 			for (auto It = BrushMaterialInstanceMap.CreateConstIterator(); It; ++It)
 			{
@@ -172,7 +172,7 @@ public:
 
 				It.Value()->SetScalarParameterValue(FName(TEXT("LocalRadius")), Radius);
 				It.Value()->SetScalarParameterValue(FName(TEXT("LocalFalloff")), Falloff);
-				It.Value()->SetVectorParameterValue(FName(TEXT("WorldPosition")), FLinearColor(WorldLocation.X,WorldLocation.Y,WorldLocation.Z,ScaleXY));
+				It.Value()->SetVectorParameterValue(FName(TEXT("WorldPosition")), FLinearColor(WorldLocation.X, WorldLocation.Y, WorldLocation.Z, ScaleXY));
 
 				bool bCanPaint = true;
 
@@ -205,13 +205,13 @@ public:
 		LastMousePosition = FVector2D(LandscapeX, LandscapeY);
 	}
 
-	virtual bool InputKey( FEditorViewportClient* InViewportClient, FViewport* InViewport, FKey InKey, EInputEvent InEvent ) override
+	virtual bool InputKey(FEditorViewportClient* InViewportClient, FViewport* InViewport, FKey InKey, EInputEvent InEvent) override
 	{
 		bool bUpdate = false;
 		return bUpdate;
 	}
 
-	virtual bool ApplyBrush( const TArray<FLandscapeToolMousePosition>& MousePositions, TMap<FIntPoint, float>& OutBrush, int32& X1, int32& Y1, int32& X2, int32& Y2 ) override
+	virtual bool ApplyBrush(const TArray<FLandscapeToolMousePosition>& MousePositions, TMap<FIntPoint, float>& OutBrush, int32& X1, int32& Y1, int32& X2, int32& Y2) override
 	{
 		ULandscapeInfo* LandscapeInfo = EdMode->CurrentToolTarget.LandscapeInfo.Get();
 		if (LandscapeInfo)
@@ -221,17 +221,17 @@ public:
 			float Falloff = EdMode->UISettings->BrushFalloff * EdMode->UISettings->BrushRadius / ScaleXY;
 
 			bool bHasOutput = false;
-			for( int32 MoveIdx=0;MoveIdx<MousePositions.Num();MoveIdx++ )
+			for (int32 MoveIdx = 0; MoveIdx < MousePositions.Num(); MoveIdx++)
 			{
 				float MousePositionX = MousePositions[MoveIdx].PositionX;
 				float MousePositionY = MousePositions[MoveIdx].PositionY;
 
-				int32 ThisMoveX1 = FMath::FloorToInt(MousePositionX - (Radius+Falloff));
-				int32 ThisMoveY1 = FMath::FloorToInt(MousePositionY - (Radius+Falloff));
-				int32 ThisMoveX2 = FMath::CeilToInt(MousePositionX + (Radius+Falloff));
-				int32 ThisMoveY2 = FMath::CeilToInt(MousePositionY + (Radius+Falloff));
-				
-				if( MoveIdx == 0 )
+				int32 ThisMoveX1 = FMath::FloorToInt(MousePositionX - (Radius + Falloff));
+				int32 ThisMoveY1 = FMath::FloorToInt(MousePositionY - (Radius + Falloff));
+				int32 ThisMoveX2 = FMath::CeilToInt(MousePositionX + (Radius + Falloff));
+				int32 ThisMoveY2 = FMath::CeilToInt(MousePositionY + (Radius + Falloff));
+
+				if (MoveIdx == 0)
 				{
 					X1 = ThisMoveX1;
 					Y1 = ThisMoveY1;
@@ -246,23 +246,23 @@ public:
 					Y2 = FMath::Max<int32>(ThisMoveY2, Y2);
 				}
 
-				for( int32 Y=ThisMoveY1;Y<=ThisMoveY2;Y++ )
+				for (int32 Y = ThisMoveY1; Y <= ThisMoveY2; Y++)
 				{
-					for( int32 X=ThisMoveX1;X<=ThisMoveX2;X++ )
+					for (int32 X = ThisMoveX1; X <= ThisMoveX2; X++)
 					{
-						FIntPoint VertexKey(X,Y);
+						FIntPoint VertexKey(X, Y);
 
 						float* PrevAmount = OutBrush.Find(VertexKey);
-						if( !PrevAmount || *PrevAmount < 1.f )
+						if (!PrevAmount || *PrevAmount < 1.f)
 						{
 							// Distance from mouse
-							float MouseDist = FMath::Sqrt(FMath::Square(MousePositionX-(float)X) + FMath::Square(MousePositionY-(float)Y));
+							float MouseDist = FMath::Sqrt(FMath::Square(MousePositionX - (float)X) + FMath::Square(MousePositionY - (float)Y));
 
 							float PaintAmount = CalculateFalloff(MouseDist, Radius, Falloff);
 
-							if( PaintAmount > 0.f )
+							if (PaintAmount > 0.f)
 							{
-								if (EdMode->CurrentToolSet && EdMode->CurrentToolSet->GetTool() && EdMode->CurrentToolSet->GetTool()->GetToolType() != FLandscapeTool::TT_Mask 
+								if (EdMode->CurrentTool && EdMode->CurrentTool && EdMode->CurrentTool->GetToolType() != FLandscapeTool::TT_Mask
 									&& EdMode->UISettings->bUseSelectedRegion)
 								{
 									float MaskValue = LandscapeInfo->SelectedRegion.FindRef(VertexKey);
@@ -273,7 +273,7 @@ public:
 									PaintAmount *= MaskValue;
 								}
 
-								if( !PrevAmount || *PrevAmount < PaintAmount )
+								if (!PrevAmount || *PrevAmount < PaintAmount)
 								{
 									// Set the brush value for this vertex
 									OutBrush.Add(VertexKey, PaintAmount);
@@ -306,11 +306,11 @@ protected:
 	FVector2D LastMousePosition;
 	UMaterial* BrushMaterial;
 public:
-	class FEdModeLandscape* EdMode;
+	FEdModeLandscape* EdMode;
 
-	FLandscapeBrushComponent(class FEdModeLandscape* InEdMode)
-		:	BrushMaterial(NULL)
-		,   EdMode(InEdMode)
+	FLandscapeBrushComponent(FEdModeLandscape* InEdMode)
+		: BrushMaterial(NULL)
+		, EdMode(InEdMode)
 	{
 		BrushMaterial = LoadObject<UMaterial>(NULL, TEXT("/Engine/EditorLandscapeResources/SelectBrushMaterial.SelectBrushMaterial"), NULL, LOAD_None, NULL);
 	}
@@ -325,9 +325,9 @@ public:
 
 	virtual void LeaveBrush() override
 	{
-		for( TSet<ULandscapeComponent*>::TIterator It(BrushMaterialComponents); It; ++It )
+		for (TSet<ULandscapeComponent*>::TIterator It(BrushMaterialComponents); It; ++It)
 		{
-			if( (*It)->EditToolRenderData != NULL )
+			if ((*It)->EditToolRenderData != NULL)
 			{
 				(*It)->EditToolRenderData->Update(NULL);
 			}
@@ -335,13 +335,13 @@ public:
 		BrushMaterialComponents.Empty();
 	}
 
-	virtual void BeginStroke(float LandscapeX, float LandscapeY, class FLandscapeTool* CurrentTool) override
+	virtual void BeginStroke(float LandscapeX, float LandscapeY, FLandscapeTool* CurrentTool) override
 	{
-		FLandscapeBrush::BeginStroke(LandscapeX,LandscapeY,CurrentTool);
+		FLandscapeBrush::BeginStroke(LandscapeX, LandscapeY, CurrentTool);
 		LastMousePosition = FVector2D(LandscapeX, LandscapeY);
 	}
 
-	virtual void Tick(FEditorViewportClient* ViewportClient,float DeltaTime) override
+	virtual void Tick(FEditorViewportClient* ViewportClient, float DeltaTime) override
 	{
 		TSet<ULandscapeComponent*> NewComponents;
 
@@ -364,7 +364,7 @@ public:
 						if (Component && FLevelUtils::IsLevelVisible(Component->GetLandscapeProxy()->GetLevel()))
 						{
 							// For MoveToLevel
-							if (EdMode->CurrentToolSet->GetToolSetName() == FName(TEXT("ToolSet_MoveToLevel")))
+							if (EdMode->CurrentTool->GetToolName() == FName("MoveToLevel"))
 							{
 								if (Component->GetLandscapeProxy() && !Component->GetLandscapeProxy()->GetLevel()->IsCurrentLevel())
 								{
@@ -408,7 +408,7 @@ public:
 		LastMousePosition = FVector2D(LandscapeX, LandscapeY);
 	}
 
-	virtual bool InputKey( FEditorViewportClient* InViewportClient, FViewport* InViewport, FKey InKey, EInputEvent InEvent ) override
+	virtual bool InputKey(FEditorViewportClient* InViewportClient, FViewport* InViewport, FKey InKey, EInputEvent InEvent) override
 	{
 		bool bUpdate = false;
 
@@ -416,36 +416,36 @@ public:
 		return bUpdate;
 	}
 
-	virtual bool ApplyBrush( const TArray<FLandscapeToolMousePosition>& MousePositions, TMap<FIntPoint, float>& OutBrush, int32& X1, int32& Y1, int32& X2, int32& Y2 ) override
+	virtual bool ApplyBrush(const TArray<FLandscapeToolMousePosition>& MousePositions, TMap<FIntPoint, float>& OutBrush, int32& X1, int32& Y1, int32& X2, int32& Y2) override
 	{
 		// Selection Brush only works for 
 		ULandscapeInfo* LandscapeInfo = EdMode->CurrentToolTarget.LandscapeInfo.Get();
-		X1= INT_MAX;
-		Y1= INT_MAX;
-		X2= INT_MIN;
-		Y2= INT_MIN;
+		X1 = INT_MAX;
+		Y1 = INT_MAX;
+		X2 = INT_MIN;
+		Y2 = INT_MIN;
 
 		// Check for add component...
-		if( EdMode->LandscapeRenderAddCollision )
+		if (EdMode->LandscapeRenderAddCollision)
 		{
 			// Apply Brush size..
 			int32 X = FMath::FloorToInt(LastMousePosition.X);
 			int32 Y = FMath::FloorToInt(LastMousePosition.Y);
 
-			int32 ComponentIndexX = (X >= 0.f) ? X / LandscapeInfo->ComponentSizeQuads : (X+1) / LandscapeInfo->ComponentSizeQuads - 1;
-			int32 ComponentIndexY = (Y >= 0.f) ? Y / LandscapeInfo->ComponentSizeQuads : (Y+1) / LandscapeInfo->ComponentSizeQuads - 1;
+			int32 ComponentIndexX = (X >= 0.f) ? X / LandscapeInfo->ComponentSizeQuads : (X + 1) / LandscapeInfo->ComponentSizeQuads - 1;
+			int32 ComponentIndexY = (Y >= 0.f) ? Y / LandscapeInfo->ComponentSizeQuads : (Y + 1) / LandscapeInfo->ComponentSizeQuads - 1;
 
-			int32 BrushSize = FMath::Max(EdMode->UISettings->BrushComponentSize-1, 0);
+			int32 BrushSize = FMath::Max(EdMode->UISettings->BrushComponentSize - 1, 0);
 
-			X1 = (ComponentIndexX -(BrushSize>>1)) * LandscapeInfo->ComponentSizeQuads;
-			X2 = (ComponentIndexX + (BrushSize>>1) + (BrushSize%2) + 1) * LandscapeInfo->ComponentSizeQuads;
-			Y1 = (ComponentIndexY -(BrushSize>>1)) * LandscapeInfo->ComponentSizeQuads;
-			Y2 = (ComponentIndexY + (BrushSize>>1) + (BrushSize%2) + 1) * LandscapeInfo->ComponentSizeQuads;
+			X1 = (ComponentIndexX - (BrushSize >> 1)) * LandscapeInfo->ComponentSizeQuads;
+			X2 = (ComponentIndexX + (BrushSize >> 1) + (BrushSize % 2) + 1) * LandscapeInfo->ComponentSizeQuads;
+			Y1 = (ComponentIndexY - (BrushSize >> 1)) * LandscapeInfo->ComponentSizeQuads;
+			Y2 = (ComponentIndexY + (BrushSize >> 1) + (BrushSize % 2) + 1) * LandscapeInfo->ComponentSizeQuads;
 		}
 		else
 		{
 			// Get extent for all components
-			for ( TSet<ULandscapeComponent*>::TIterator It(BrushMaterialComponents); It; ++It )
+			for (TSet<ULandscapeComponent*>::TIterator It(BrushMaterialComponents); It; ++It)
 			{
 				if (*It)
 				{
@@ -457,14 +457,14 @@ public:
 		// Should not be possible...
 		//check(X1 <= X2 && Y1 <= Y2);
 
-		for( int32 Y=Y1;Y<=Y2;Y++ )
+		for (int32 Y = Y1; Y <= Y2; Y++)
 		{
-			for( int32 X=X1;X<=X2;X++ )
+			for (int32 X = X1; X <= X2; X++)
 			{
-				FIntPoint VertexKey = ALandscape::MakeKey(X,Y);
+				FIntPoint VertexKey = ALandscape::MakeKey(X, Y);
 
 				float PaintAmount = 1.0f;
-				if (EdMode->CurrentToolSet && EdMode->CurrentToolSet->GetTool() && EdMode->CurrentToolSet->GetTool()->GetToolType() != FLandscapeTool::TT_Mask 
+				if (EdMode->CurrentTool && EdMode->CurrentTool && EdMode->CurrentTool->GetToolType() != FLandscapeTool::TT_Mask
 					&& EdMode->UISettings->bUseSelectedRegion)
 				{
 					float MaskValue = LandscapeInfo->SelectedRegion.FindRef(VertexKey);
@@ -479,11 +479,11 @@ public:
 				OutBrush.Add(VertexKey, PaintAmount);
 			}
 		}
-		
+
 		return (X1 <= X2 && Y1 <= Y2);
 	}
 };
-	
+
 // 
 // FLandscapeBrushGizmo
 //
@@ -499,10 +499,10 @@ protected:
 	FVector2D LastMousePosition;
 	UMaterialInstanceDynamic* BrushMaterial;
 public:
-	class FEdModeLandscape* EdMode;
+	FEdModeLandscape* EdMode;
 
-	FLandscapeBrushGizmo(class FEdModeLandscape* InEdMode)
-		:	EdMode(InEdMode),
+	FLandscapeBrushGizmo(FEdModeLandscape* InEdMode)
+		: EdMode(InEdMode),
 		BrushMaterial(NULL)
 	{
 		UMaterialInstanceConstant* GizmoMaterial = LoadObject<UMaterialInstanceConstant>(NULL, TEXT("/Engine/EditorLandscapeResources/MaskBrushMaterial_Gizmo.MaskBrushMaterial_Gizmo"), NULL, LOAD_None, NULL);
@@ -519,9 +519,9 @@ public:
 
 	virtual void LeaveBrush() override
 	{
-		for( TSet<ULandscapeComponent*>::TIterator It(BrushMaterialComponents); It; ++It )
+		for (TSet<ULandscapeComponent*>::TIterator It(BrushMaterialComponents); It; ++It)
 		{
-			if( (*It)->EditToolRenderData != NULL )
+			if ((*It)->EditToolRenderData != NULL)
 			{
 				(*It)->EditToolRenderData->Update(NULL);
 			}
@@ -529,25 +529,25 @@ public:
 		BrushMaterialComponents.Empty();
 	}
 
-	virtual void BeginStroke(float LandscapeX, float LandscapeY, class FLandscapeTool* CurrentTool) override
+	virtual void BeginStroke(float LandscapeX, float LandscapeY, FLandscapeTool* CurrentTool) override
 	{
-		FLandscapeBrush::BeginStroke(LandscapeX,LandscapeY,CurrentTool);
+		FLandscapeBrush::BeginStroke(LandscapeX, LandscapeY, CurrentTool);
 		LastMousePosition = FVector2D(LandscapeX, LandscapeY);
 	}
 
-	virtual void Tick(FEditorViewportClient* ViewportClient,float DeltaTime) override
+	virtual void Tick(FEditorViewportClient* ViewportClient, float DeltaTime) override
 	{
 		if (GLandscapeEditRenderMode & ELandscapeEditRenderMode::Gizmo || GLandscapeEditRenderMode & ELandscapeEditRenderMode::Select)
 		{
 			ALandscapeGizmoActiveActor* Gizmo = EdMode->CurrentGizmoActor.Get();
 
-			if( Gizmo && Gizmo->TargetLandscapeInfo && (Gizmo->TargetLandscapeInfo == EdMode->CurrentToolTarget.LandscapeInfo.Get()) && Gizmo->GizmoTexture && Gizmo->GetRootComponent() )
+			if (Gizmo && Gizmo->TargetLandscapeInfo && (Gizmo->TargetLandscapeInfo == EdMode->CurrentToolTarget.LandscapeInfo.Get()) && Gizmo->GizmoTexture && Gizmo->GetRootComponent())
 			{
 				ULandscapeInfo* LandscapeInfo = Gizmo->TargetLandscapeInfo;
 				if (LandscapeInfo && LandscapeInfo->GetLandscapeProxy())
 				{
 					float ScaleXY = LandscapeInfo->DrawScale.X;
-					FMatrix LToW =  LandscapeInfo->GetLandscapeProxy()->LandscapeActorToWorld().ToMatrixWithScale();
+					FMatrix LToW = LandscapeInfo->GetLandscapeProxy()->LandscapeActorToWorld().ToMatrixWithScale();
 					FMatrix WToL = LToW.Inverse();
 
 					UTexture2D* DataTexture = Gizmo->GizmoTexture;
@@ -582,25 +582,25 @@ public:
 					BrushMaterial->SetTextureParameterValue(FName(TEXT("AlphaTexture")), DataTexture);
 
 					// Set brush material for components in new region
-					for( TSet<ULandscapeComponent*>::TIterator It(NewComponents); It; ++It )
+					for (TSet<ULandscapeComponent*>::TIterator It(NewComponents); It; ++It)
 					{
-						if( (*It)->EditToolRenderData != NULL )
+						if ((*It)->EditToolRenderData != NULL)
 						{
-							(*It)->EditToolRenderData->UpdateGizmo((Gizmo->DataType != LGT_None) && (GLandscapeEditRenderMode & ELandscapeEditRenderMode::Gizmo)? BrushMaterial : NULL);
+							(*It)->EditToolRenderData->UpdateGizmo((Gizmo->DataType != LGT_None) && (GLandscapeEditRenderMode & ELandscapeEditRenderMode::Gizmo) ? BrushMaterial : NULL);
 						}
 					}
 
 					// Remove the material from any old components that are no longer in the region
 					TSet<ULandscapeComponent*> RemovedComponents = BrushMaterialComponents.Difference(NewComponents);
-					for ( TSet<ULandscapeComponent*>::TIterator It(RemovedComponents); It; ++It )
+					for (TSet<ULandscapeComponent*>::TIterator It(RemovedComponents); It; ++It)
 					{
-						if( (*It)->EditToolRenderData != NULL )
+						if ((*It)->EditToolRenderData != NULL)
 						{
 							(*It)->EditToolRenderData->UpdateGizmo(NULL);
 						}
 					}
 
-					BrushMaterialComponents = NewComponents;	
+					BrushMaterialComponents = NewComponents;
 				}
 			}
 		}
@@ -611,13 +611,13 @@ public:
 		LastMousePosition = FVector2D(LandscapeX, LandscapeY);
 	}
 
-	virtual bool InputKey( FEditorViewportClient* InViewportClient, FViewport* InViewport, FKey InKey, EInputEvent InEvent ) override
+	virtual bool InputKey(FEditorViewportClient* InViewportClient, FViewport* InViewport, FKey InKey, EInputEvent InEvent) override
 	{
 		bool bUpdate = false;
 		return bUpdate;
 	}
 
-	virtual bool ApplyBrush( const TArray<FLandscapeToolMousePosition>& MousePositions, TMap<FIntPoint, float>& OutBrush, int32& X1, int32& Y1, int32& X2, int32& Y2 ) override
+	virtual bool ApplyBrush(const TArray<FLandscapeToolMousePosition>& MousePositions, TMap<FIntPoint, float>& OutBrush, int32& X1, int32& Y1, int32& X2, int32& Y2) override
 	{
 		// Selection Brush only works for 
 		ALandscapeGizmoActiveActor* Gizmo = EdMode->CurrentGizmoActor.Get();
@@ -628,15 +628,15 @@ public:
 			Gizmo->TargetLandscapeInfo = LandscapeInfo;
 			float ScaleXY = LandscapeInfo->DrawScale.X;
 
-			X1= INT_MAX;
-			Y1= INT_MAX;
-			X2= INT_MIN;
-			Y2= INT_MIN;
+			X1 = INT_MAX;
+			Y1 = INT_MAX;
+			X2 = INT_MIN;
+			Y2 = INT_MIN;
 
 			// Get extent for all components
-			for ( TSet<ULandscapeComponent*>::TIterator It(BrushMaterialComponents); It; ++It )
+			for (TSet<ULandscapeComponent*>::TIterator It(BrushMaterialComponents); It; ++It)
 			{
-				if( *It )
+				if (*It)
 				{
 					(*It)->GetComponentExtent(X1, Y1, X2, Y2);
 				}
@@ -651,19 +651,19 @@ public:
 
 			FMatrix WToL = LandscapeInfo->GetLandscapeProxy()->LandscapeActorToWorld().ToMatrixWithScale().Inverse();
 			FVector BaseLocation = WToL.TransformPosition(Gizmo->GetActorLocation());
-			FMatrix LandscapeToGizmoLocal = 
-				(FTranslationMatrix(FVector(- LW + 0.5, - LH + 0.5, 0)) * FRotationTranslationMatrix(FRotator(0, Gizmo->GetActorRotation().Yaw, 0), FVector(BaseLocation.X, BaseLocation.Y, 0))).Inverse();
+			FMatrix LandscapeToGizmoLocal =
+				(FTranslationMatrix(FVector(-LW + 0.5, -LH + 0.5, 0)) * FRotationTranslationMatrix(FRotator(0, Gizmo->GetActorRotation().Yaw, 0), FVector(BaseLocation.X, BaseLocation.Y, 0))).Inverse();
 
 			float W = Gizmo->GetWidth() / ScaleXY; //Gizmo->GetWidth() / (Gizmo->DrawScale * Gizmo->DrawScale3D.X);
 			float H = Gizmo->GetHeight() / ScaleXY; //Gizmo->GetHeight() / (Gizmo->DrawScale * Gizmo->DrawScale3D.Y);
 
-			for( int32 Y=Y1;Y<=Y2;Y++ )
+			for (int32 Y = Y1; Y <= Y2; Y++)
 			{
-				for( int32 X=X1;X<=X2;X++ )
+				for (int32 X = X1; X <= X2; X++)
 				{
-					FIntPoint VertexKey = ALandscape::MakeKey(X,Y);
+					FIntPoint VertexKey = ALandscape::MakeKey(X, Y);
 
-					FVector GizmoLocal = LandscapeToGizmoLocal.TransformPosition(FVector(X,Y,0));
+					FVector GizmoLocal = LandscapeToGizmoLocal.TransformPosition(FVector(X, Y, 0));
 					if (GizmoLocal.X < W && GizmoLocal.X > 0 && GizmoLocal.Y < H && GizmoLocal.Y > 0)
 					{
 						float PaintAmount = 1.f;
@@ -678,12 +678,12 @@ public:
 							float RatioX = FalloffRadius > 0.f ? 1.f - FMath::Clamp((FMath::Abs(TransformedLocal.X) - Cos*SquareRadius) / FalloffRadius, 0.f, 1.f) : 1.f;
 							float RatioY = FalloffRadius > 0.f ? 1.f - FMath::Clamp((FMath::Abs(TransformedLocal.Y) - Sin*SquareRadius) / FalloffRadius, 0.f, 1.f) : 1.f;
 							float Ratio = TransformedLocal.Size2D() > SquareRadius ? RatioX * RatioY : 1.f; //TransformedLocal.X / LW * TransformedLocal.Y / LW;
-							PaintAmount = Ratio*Ratio*(3-2*Ratio); //FMath::Lerp(SquareFalloff, RectFalloff*RectFalloff, Ratio);
+							PaintAmount = Ratio*Ratio*(3 - 2 * Ratio); //FMath::Lerp(SquareFalloff, RectFalloff*RectFalloff, Ratio);
 						}
 
 						if (PaintAmount)
 						{
-							if (EdMode->CurrentToolSet && EdMode->CurrentToolSet->GetTool() && EdMode->CurrentToolSet->GetTool()->GetToolType() != FLandscapeTool::TT_Mask 
+							if (EdMode->CurrentTool && EdMode->CurrentTool && EdMode->CurrentTool->GetToolType() != FLandscapeTool::TT_Mask
 								&& EdMode->UISettings->bUseSelectedRegion)
 							{
 								float MaskValue = LandscapeInfo->SelectedRegion.FindRef(VertexKey);
@@ -714,9 +714,9 @@ public:
 	const TCHAR* GetBrushName() override { return TEXT("Splines"); }
 	virtual FText GetDisplayName() override { return NSLOCTEXT("UnrealEd", "LandscapeMode_Brush_Splines", "Splines"); };
 
-	class FEdModeLandscape* EdMode;
+	FEdModeLandscape* EdMode;
 
-	FLandscapeBrushSplines(class FEdModeLandscape* InEdMode)
+	FLandscapeBrushSplines(FEdModeLandscape* InEdMode)
 		: EdMode(InEdMode)
 	{
 	}
@@ -727,11 +727,11 @@ public:
 
 	virtual EBrushType GetBrushType() override { return BT_Splines; }
 
-	virtual void MouseMove(float LandscapeX, float LandscapeY) override 
+	virtual void MouseMove(float LandscapeX, float LandscapeY) override
 	{
 	}
 
-	virtual bool ApplyBrush( const TArray<FLandscapeToolMousePosition>& MousePositions, TMap<FIntPoint, float>& OutBrush, int32& X1, int32& Y1, int32& X2, int32& Y2 ) override
+	virtual bool ApplyBrush(const TArray<FLandscapeToolMousePosition>& MousePositions, TMap<FIntPoint, float>& OutBrush, int32& X1, int32& Y1, int32& X2, int32& Y2) override
 	{
 		return false;
 	}
@@ -746,9 +746,9 @@ public:
 	const TCHAR* GetBrushName() { return TEXT("None"); }
 	virtual FText GetDisplayName() override { return NSLOCTEXT("UnrealEd", "LandscapeMode_Brush_None", "None"); };
 
-	class FEdModeLandscape* EdMode;
+	FEdModeLandscape* EdMode;
 
-	FLandscapeBrushDummy(class FEdModeLandscape* InEdMode)
+	FLandscapeBrushDummy(FEdModeLandscape* InEdMode)
 		: EdMode(InEdMode)
 	{
 	}
@@ -759,11 +759,11 @@ public:
 
 	virtual EBrushType GetBrushType() override { return BT_Normal; }
 
-	virtual void MouseMove(float LandscapeX, float LandscapeY) override 
+	virtual void MouseMove(float LandscapeX, float LandscapeY) override
 	{
 	}
 
-	virtual bool ApplyBrush( const TArray<FLandscapeToolMousePosition>& MousePositions, TMap<FIntPoint, float>& OutBrush, int32& X1, int32& Y1, int32& X2, int32& Y2 ) override
+	virtual bool ApplyBrush(const TArray<FLandscapeToolMousePosition>& MousePositions, TMap<FIntPoint, float>& OutBrush, int32& X1, int32& Y1, int32& X2, int32& Y2) override
 	{
 		return false;
 	}
@@ -773,20 +773,20 @@ class FLandscapeBrushCircle_Linear : public FLandscapeBrushCircle
 {
 protected:
 	/** Protected so only subclasses can create instances. */
-	FLandscapeBrushCircle_Linear(class FEdModeLandscape* InEdMode, UMaterialInterface* InBrushMaterial)
-		:	FLandscapeBrushCircle(InEdMode, InBrushMaterial)
+	FLandscapeBrushCircle_Linear(FEdModeLandscape* InEdMode, UMaterialInterface* InBrushMaterial)
+		: FLandscapeBrushCircle(InEdMode, InBrushMaterial)
 	{
 	}
 
-	virtual float CalculateFalloff( float Distance, float Radius, float Falloff ) override
+	virtual float CalculateFalloff(float Distance, float Radius, float Falloff) override
 	{
-		return Distance < Radius ? 1.f : 
-			Falloff > 0.f ? FMath::Max<float>(0.f, 1.f - (Distance - Radius) / Falloff) : 
+		return Distance < Radius ? 1.f :
+			Falloff > 0.f ? FMath::Max<float>(0.f, 1.f - (Distance - Radius) / Falloff) :
 			0.f;
 	}
 
 public:
-	static FLandscapeBrushCircle_Linear* Create(class FEdModeLandscape* InEdMode)
+	static FLandscapeBrushCircle_Linear* Create(FEdModeLandscape* InEdMode)
 	{
 		UMaterialInstanceConstant* CircleBrushMaterial_Linear = LoadObject<UMaterialInstanceConstant>(NULL, TEXT("/Engine/EditorLandscapeResources/CircleBrushMaterial_Linear.CircleBrushMaterial_Linear"), NULL, LOAD_None, NULL);
 		return new FLandscapeBrushCircle_Linear(InEdMode, CircleBrushMaterial_Linear);
@@ -799,20 +799,20 @@ public:
 class FLandscapeBrushCircle_Smooth : public FLandscapeBrushCircle_Linear
 {
 protected:
-	FLandscapeBrushCircle_Smooth(class FEdModeLandscape* InEdMode, UMaterialInterface* InBrushMaterial)
-		:	FLandscapeBrushCircle_Linear(InEdMode, InBrushMaterial)
+	FLandscapeBrushCircle_Smooth(FEdModeLandscape* InEdMode, UMaterialInterface* InBrushMaterial)
+		: FLandscapeBrushCircle_Linear(InEdMode, InBrushMaterial)
 	{
 	}
 
-	virtual float CalculateFalloff( float Distance, float Radius, float Falloff ) override
+	virtual float CalculateFalloff(float Distance, float Radius, float Falloff) override
 	{
 		float y = FLandscapeBrushCircle_Linear::CalculateFalloff(Distance, Radius, Falloff);
 		// Smooth-step it
-		return y*y*(3-2*y);
+		return y*y*(3 - 2 * y);
 	}
 
 public:
-	static FLandscapeBrushCircle_Smooth* Create(class FEdModeLandscape* InEdMode)
+	static FLandscapeBrushCircle_Smooth* Create(FEdModeLandscape* InEdMode)
 	{
 		UMaterialInstanceConstant* CircleBrushMaterial_Smooth = LoadObject<UMaterialInstanceConstant>(NULL, TEXT("/Engine/EditorLandscapeResources/CircleBrushMaterial_Smooth.CircleBrushMaterial_Smooth"), NULL, LOAD_None, NULL);
 		return new FLandscapeBrushCircle_Smooth(InEdMode, CircleBrushMaterial_Smooth);
@@ -825,29 +825,29 @@ public:
 class FLandscapeBrushCircle_Spherical : public FLandscapeBrushCircle
 {
 protected:
-	FLandscapeBrushCircle_Spherical(class FEdModeLandscape* InEdMode, UMaterialInterface* InBrushMaterial)
-		:	FLandscapeBrushCircle(InEdMode, InBrushMaterial)
+	FLandscapeBrushCircle_Spherical(FEdModeLandscape* InEdMode, UMaterialInterface* InBrushMaterial)
+		: FLandscapeBrushCircle(InEdMode, InBrushMaterial)
 	{
 	}
 
-	virtual float CalculateFalloff( float Distance, float Radius, float Falloff ) override
+	virtual float CalculateFalloff(float Distance, float Radius, float Falloff) override
 	{
-		if( Distance <= Radius )
+		if (Distance <= Radius)
 		{
 			return 1.f;
 		}
 
-		if( Distance > Radius + Falloff )
+		if (Distance > Radius + Falloff)
 		{
 			return 0.f;
 		}
 
 		// Elliptical falloff
-		return FMath::Sqrt( 1.f - FMath::Square((Distance - Radius) / Falloff) );
+		return FMath::Sqrt(1.f - FMath::Square((Distance - Radius) / Falloff));
 	}
 
 public:
-	static FLandscapeBrushCircle_Spherical* Create(class FEdModeLandscape* InEdMode)
+	static FLandscapeBrushCircle_Spherical* Create(FEdModeLandscape* InEdMode)
 	{
 		UMaterialInstanceConstant* CircleBrushMaterial_Spherical = LoadObject<UMaterialInstanceConstant>(NULL, TEXT("/Engine/EditorLandscapeResources/CircleBrushMaterial_Spherical.CircleBrushMaterial_Spherical"), NULL, LOAD_None, NULL);
 		return new FLandscapeBrushCircle_Spherical(InEdMode, CircleBrushMaterial_Spherical);
@@ -859,25 +859,25 @@ public:
 class FLandscapeBrushCircle_Tip : public FLandscapeBrushCircle
 {
 protected:
-	FLandscapeBrushCircle_Tip(class FEdModeLandscape* InEdMode, UMaterialInterface* InBrushMaterial)
-		:	FLandscapeBrushCircle(InEdMode, InBrushMaterial)
+	FLandscapeBrushCircle_Tip(FEdModeLandscape* InEdMode, UMaterialInterface* InBrushMaterial)
+		: FLandscapeBrushCircle(InEdMode, InBrushMaterial)
 	{
 	}
 
-	virtual float CalculateFalloff( float Distance, float Radius, float Falloff ) override
+	virtual float CalculateFalloff(float Distance, float Radius, float Falloff) override
 	{
-		if( Distance <= Radius )
+		if (Distance <= Radius)
 		{
 			return 1.f;
 		}
 
-		if( Distance > Radius + Falloff )
+		if (Distance > Radius + Falloff)
 		{
 			return 0.f;
 		}
 
 		// inverse elliptical falloff
-		return 1.f - FMath::Sqrt( 1.f - FMath::Square((Falloff + Radius - Distance) / Falloff) );
+		return 1.f - FMath::Sqrt(1.f - FMath::Square((Falloff + Radius - Distance) / Falloff));
 	}
 
 public:
@@ -895,31 +895,31 @@ public:
 class FLandscapeBrushAlphaBase : public FLandscapeBrushCircle_Smooth
 {
 public:
-	FLandscapeBrushAlphaBase(class FEdModeLandscape* InEdMode, UMaterialInterface* InBrushMaterial)
-		:	FLandscapeBrushCircle_Smooth(InEdMode, InBrushMaterial)
+	FLandscapeBrushAlphaBase(FEdModeLandscape* InEdMode, UMaterialInterface* InBrushMaterial)
+		: FLandscapeBrushCircle_Smooth(InEdMode, InBrushMaterial)
 	{}
 
-	float GetAlphaSample( float SampleX, float SampleY )
+	float GetAlphaSample(float SampleX, float SampleY)
 	{
 		int32 SizeX = EdMode->UISettings->AlphaTextureSizeX;
 		int32 SizeY = EdMode->UISettings->AlphaTextureSizeY;
 
 		// Bilinear interpolate the values from the alpha texture
 		int32 SampleX0 = FMath::FloorToInt(SampleX);
-		int32 SampleX1 = (SampleX0+1) % SizeX;
+		int32 SampleX1 = (SampleX0 + 1) % SizeX;
 		int32 SampleY0 = FMath::FloorToInt(SampleY);
-		int32 SampleY1 = (SampleY0+1) % SizeY;
+		int32 SampleY1 = (SampleY0 + 1) % SizeY;
 
 		const uint8* AlphaData = EdMode->UISettings->AlphaTextureData.GetData();
 
-		float Alpha00 = (float)AlphaData[ SampleX0 + SampleY0 * SizeX ] / 255.f;
-		float Alpha01 = (float)AlphaData[ SampleX0 + SampleY1 * SizeX ] / 255.f;
-		float Alpha10 = (float)AlphaData[ SampleX1 + SampleY0 * SizeX ] / 255.f;
-		float Alpha11 = (float)AlphaData[ SampleX1 + SampleY1 * SizeX ] / 255.f;
+		float Alpha00 = (float)AlphaData[SampleX0 + SampleY0 * SizeX] / 255.f;
+		float Alpha01 = (float)AlphaData[SampleX0 + SampleY1 * SizeX] / 255.f;
+		float Alpha10 = (float)AlphaData[SampleX1 + SampleY0 * SizeX] / 255.f;
+		float Alpha11 = (float)AlphaData[SampleX1 + SampleY1 * SizeX] / 255.f;
 
 		return FMath::Lerp(
-			FMath::Lerp( Alpha00, Alpha01, FMath::Fractional(SampleX) ),
-			FMath::Lerp( Alpha10, Alpha11, FMath::Fractional(SampleX) ),
+			FMath::Lerp(Alpha00, Alpha01, FMath::Fractional(SampleX)),
+			FMath::Lerp(Alpha10, Alpha11, FMath::Fractional(SampleX)),
 			FMath::Fractional(SampleY)
 			);
 	}
@@ -932,13 +932,13 @@ public:
 class FLandscapeBrushAlphaPattern : public FLandscapeBrushAlphaBase
 {
 protected:
-	FLandscapeBrushAlphaPattern(class FEdModeLandscape* InEdMode, UMaterialInterface* InBrushMaterial)
-		:	FLandscapeBrushAlphaBase(InEdMode, InBrushMaterial)
+	FLandscapeBrushAlphaPattern(FEdModeLandscape* InEdMode, UMaterialInterface* InBrushMaterial)
+		: FLandscapeBrushAlphaBase(InEdMode, InBrushMaterial)
 	{
 	}
 
 public:
-	static FLandscapeBrushAlphaPattern* Create(class FEdModeLandscape* InEdMode)
+	static FLandscapeBrushAlphaPattern* Create(FEdModeLandscape* InEdMode)
 	{
 		UMaterialInstanceConstant* PatternBrushMaterial = LoadObject<UMaterialInstanceConstant>(NULL, TEXT("/Engine/EditorLandscapeResources/PatternBrushMaterial_Smooth.PatternBrushMaterial_Smooth"), NULL, LOAD_None, NULL);
 		return new FLandscapeBrushAlphaPattern(InEdMode, PatternBrushMaterial);
@@ -946,7 +946,7 @@ public:
 
 	virtual EBrushType GetBrushType() override { return BT_Alpha; }
 
-	virtual bool ApplyBrush( const TArray<FLandscapeToolMousePosition>& MousePositions, TMap<FIntPoint, float>& OutBrush, int32& X1, int32& Y1, int32& X2, int32& Y2 ) override
+	virtual bool ApplyBrush(const TArray<FLandscapeToolMousePosition>& MousePositions, TMap<FIntPoint, float>& OutBrush, int32& X1, int32& Y1, int32& X2, int32& Y2) override
 	{
 		ULandscapeInfo* LandscapeInfo = EdMode->CurrentToolTarget.LandscapeInfo.Get();
 		float ScaleXY = LandscapeInfo->DrawScale.X;
@@ -957,16 +957,16 @@ public:
 		int32 SizeX = EdMode->UISettings->AlphaTextureSizeX;
 		int32 SizeY = EdMode->UISettings->AlphaTextureSizeY;
 
-		X1 = FMath::FloorToInt(LastMousePosition.X - (Radius+Falloff));
-		Y1 = FMath::FloorToInt(LastMousePosition.Y - (Radius+Falloff));
-		X2 = FMath::CeilToInt(LastMousePosition.X + (Radius+Falloff));
-		Y2 = FMath::CeilToInt(LastMousePosition.Y + (Radius+Falloff));
+		X1 = FMath::FloorToInt(LastMousePosition.X - (Radius + Falloff));
+		Y1 = FMath::FloorToInt(LastMousePosition.Y - (Radius + Falloff));
+		X2 = FMath::CeilToInt(LastMousePosition.X + (Radius + Falloff));
+		Y2 = FMath::CeilToInt(LastMousePosition.Y + (Radius + Falloff));
 
-		for( int32 Y=Y1;Y<=Y2;Y++ )
+		for (int32 Y = Y1; Y <= Y2; Y++)
 		{
-			for( int32 X=X1;X<=X2;X++ )
+			for (int32 X = X1; X <= X2; X++)
 			{
-				FIntPoint VertexKey = ALandscape::MakeKey(X,Y);
+				FIntPoint VertexKey = ALandscape::MakeKey(X, Y);
 
 				// Find alphamap sample location
 				float SampleX = (float)X / EdMode->UISettings->AlphaBrushScale + (float)SizeX * EdMode->UISettings->AlphaBrushPanU;
@@ -974,14 +974,14 @@ public:
 
 				float Angle = PI * EdMode->UISettings->AlphaBrushRotation / 180.f;
 
-				float ModSampleX = FMath::Fmod( SampleX * FMath::Cos(Angle) - SampleY * FMath::Sin(Angle), (float)SizeX );
-				float ModSampleY = FMath::Fmod( SampleY * FMath::Cos(Angle) + SampleX * FMath::Sin(Angle), (float)SizeY );
+				float ModSampleX = FMath::Fmod(SampleX * FMath::Cos(Angle) - SampleY * FMath::Sin(Angle), (float)SizeX);
+				float ModSampleY = FMath::Fmod(SampleY * FMath::Cos(Angle) + SampleX * FMath::Sin(Angle), (float)SizeY);
 
-				if( ModSampleX < 0.f )
+				if (ModSampleX < 0.f)
 				{
 					ModSampleX += (float)SizeX;
 				}
-				if( ModSampleY < 0.f )
+				if (ModSampleY < 0.f)
 				{
 					ModSampleY += (float)SizeY;
 				}
@@ -990,13 +990,13 @@ public:
 				float Alpha = GetAlphaSample(ModSampleX, ModSampleY);
 
 				// Distance from mouse
-				float MouseDist = FMath::Sqrt(FMath::Square(LastMousePosition.X-(float)X) + FMath::Square(LastMousePosition.Y-(float)Y));
+				float MouseDist = FMath::Sqrt(FMath::Square(LastMousePosition.X - (float)X) + FMath::Square(LastMousePosition.Y - (float)Y));
 
 				float PaintAmount = CalculateFalloff(MouseDist, Radius, Falloff) * Alpha;
 
-				if( PaintAmount > 0.f )
+				if (PaintAmount > 0.f)
 				{
-					if (EdMode->CurrentToolSet && EdMode->CurrentToolSet->GetTool() && EdMode->CurrentToolSet->GetTool()->GetToolType() != FLandscapeTool::TT_Mask 
+					if (EdMode->CurrentTool && EdMode->CurrentTool && EdMode->CurrentTool->GetToolType() != FLandscapeTool::TT_Mask
 						&& EdMode->UISettings->bUseSelectedRegion)
 					{
 						float MaskValue = LandscapeInfo->SelectedRegion.FindRef(VertexKey);
@@ -1014,12 +1014,12 @@ public:
 		return (X1 <= X2 && Y1 <= Y2);
 	}
 
-	virtual void Tick(FEditorViewportClient* ViewportClient,float DeltaTime) override
+	virtual void Tick(FEditorViewportClient* ViewportClient, float DeltaTime) override
 	{
-		FLandscapeBrushCircle::Tick(ViewportClient,DeltaTime);
+		FLandscapeBrushCircle::Tick(ViewportClient, DeltaTime);
 
 		ALandscapeProxy* Proxy = EdMode->CurrentToolTarget.LandscapeInfo.IsValid() ? EdMode->CurrentToolTarget.LandscapeInfo->GetLandscapeProxy() : NULL;
-		if( Proxy )
+		if (Proxy)
 		{
 			const float ScaleXY = EdMode->CurrentToolTarget.LandscapeInfo->DrawScale.X;
 			int32 SizeX = EdMode->UISettings->AlphaTextureSizeX;
@@ -1036,7 +1036,7 @@ public:
 			FLinearColor LandscapeLocation(Proxy->GetActorLocation().X, Proxy->GetActorLocation().Y, Proxy->GetActorLocation().Z, Angle);
 
 			int32 Channel = EdMode->UISettings->AlphaTextureChannel;
-			FLinearColor AlphaTextureMask(Channel==0?1:0, Channel==1?1:0, Channel==2?1:0, Channel==3?1:0);
+			FLinearColor AlphaTextureMask(Channel == 0 ? 1 : 0, Channel == 1 ? 1 : 0, Channel == 2 ? 1 : 0, Channel == 3 ? 1 : 0);
 
 			for (auto It = BrushMaterialInstanceMap.CreateConstIterator(); It; ++It)
 			{
@@ -1063,25 +1063,25 @@ class FLandscapeBrushAlpha : public FLandscapeBrushAlphaBase
 	double LastMouseSampleTime;
 
 protected:
-	FLandscapeBrushAlpha(class FEdModeLandscape* InEdMode, UMaterialInterface* InBrushMaterial)
-		:	FLandscapeBrushAlphaBase(InEdMode, InBrushMaterial)
-		,	LastMouseAngle(0.f)
-		,	OldMousePosition(0.f,0.f)
-		,	LastMouseSampleTime(FPlatformTime::Seconds())
+	FLandscapeBrushAlpha(FEdModeLandscape* InEdMode, UMaterialInterface* InBrushMaterial)
+		: FLandscapeBrushAlphaBase(InEdMode, InBrushMaterial)
+		, LastMouseAngle(0.f)
+		, OldMousePosition(0.f, 0.f)
+		, LastMouseSampleTime(FPlatformTime::Seconds())
 	{
 	}
 
 public:
-	static FLandscapeBrushAlpha* Create(class FEdModeLandscape* InEdMode)
+	static FLandscapeBrushAlpha* Create(FEdModeLandscape* InEdMode)
 	{
 		UMaterialInstanceConstant* AlphaBrushMaterial = LoadObject<UMaterialInstanceConstant>(NULL, TEXT("/Engine/EditorLandscapeResources/AlphaBrushMaterial_Smooth.AlphaBrushMaterial_Smooth"), NULL, LOAD_None, NULL);
 		return new FLandscapeBrushAlpha(InEdMode, AlphaBrushMaterial);
 	}
 
-	virtual bool ApplyBrush( const TArray<FLandscapeToolMousePosition>& MousePositions, TMap<FIntPoint, float>& OutBrush, int32& X1, int32& Y1, int32& X2, int32& Y2 ) override
+	virtual bool ApplyBrush(const TArray<FLandscapeToolMousePosition>& MousePositions, TMap<FIntPoint, float>& OutBrush, int32& X1, int32& Y1, int32& X2, int32& Y2) override
 	{
 		ULandscapeInfo* LandscapeInfo = EdMode->CurrentToolTarget.LandscapeInfo.Get();
-		if( OldMousePosition.IsZero() )
+		if (OldMousePosition.IsZero())
 		{
 			X1 = FMath::FloorToInt(LastMousePosition.X);
 			Y1 = FMath::FloorToInt(LastMousePosition.Y);
@@ -1097,17 +1097,17 @@ public:
 			float Radius = EdMode->UISettings->BrushRadius / ScaleXY;
 			int32 SizeX = EdMode->UISettings->AlphaTextureSizeX;
 			int32 SizeY = EdMode->UISettings->AlphaTextureSizeY;
-			float MaxSize = 2.f * FMath::Sqrt( FMath::Square(Radius) / 2.f );
-			float AlphaBrushScale = MaxSize / (float)FMath::Max<int32>(SizeX,SizeY);
+			float MaxSize = 2.f * FMath::Sqrt(FMath::Square(Radius) / 2.f);
+			float AlphaBrushScale = MaxSize / (float)FMath::Max<int32>(SizeX, SizeY);
 
 			X1 = FMath::FloorToInt(LastMousePosition.X - Radius);
 			Y1 = FMath::FloorToInt(LastMousePosition.Y - Radius);
 			X2 = FMath::CeilToInt(LastMousePosition.X + Radius);
 			Y2 = FMath::CeilToInt(LastMousePosition.Y + Radius);
 
-			for( int32 Y=Y1;Y<=Y2;Y++ )
+			for (int32 Y = Y1; Y <= Y2; Y++)
 			{
-				for( int32 X=X1;X<=X2;X++ )
+				for (int32 X = X1; X <= X2; X++)
 				{
 					// Find alphamap sample location
 					float ScaleSampleX = ((float)X - LastMousePosition.X) / AlphaBrushScale;
@@ -1120,18 +1120,18 @@ public:
 					SampleX += (float)SizeX * 0.5f;
 					SampleY += (float)SizeY * 0.5f;
 
-					if( SampleX >= 0.f && SampleX < (float)SizeX &&
-						SampleY >= 0.f && SampleY < (float)SizeY )
+					if (SampleX >= 0.f && SampleX < (float)SizeX &&
+						SampleY >= 0.f && SampleY < (float)SizeY)
 					{
 						// Sample the alpha texture
 						float Alpha = GetAlphaSample(SampleX, SampleY);
 
-						if( Alpha > 0.f )
+						if (Alpha > 0.f)
 						{
 							// Set the brush value for this vertex
-							FIntPoint VertexKey = ALandscape::MakeKey(X,Y);
+							FIntPoint VertexKey = ALandscape::MakeKey(X, Y);
 
-							if (EdMode->CurrentToolSet && EdMode->CurrentToolSet->GetTool() && EdMode->CurrentToolSet->GetTool()->GetToolType() != FLandscapeTool::TT_Mask 
+							if (EdMode->CurrentTool && EdMode->CurrentTool && EdMode->CurrentTool->GetToolType() != FLandscapeTool::TT_Mask
 								&& EdMode->UISettings->bUseSelectedRegion)
 							{
 								float MaskValue = LandscapeInfo->SelectedRegion.FindRef(VertexKey);
@@ -1158,12 +1158,12 @@ public:
 
 		// don't do anything with the angle unless we move at least 0.1 units.
 		FVector2D MouseDelta = LastMousePosition - OldMousePosition;
-		if( MouseDelta.SizeSquared() >= FMath::Square(0.5f) )
+		if (MouseDelta.SizeSquared() >= FMath::Square(0.5f))
 		{
 			double SampleTime = FPlatformTime::Seconds();
 			float DeltaTime = (float)(SampleTime - LastMouseSampleTime);
 			FVector2D MouseDirection = MouseDelta.SafeNormal();
-			float MouseAngle = FMath::Lerp( LastMouseAngle, FMath::Atan2( -MouseDirection.Y, MouseDirection.X ), FMath::Min<float>(10.f * DeltaTime, 1.f) );		// lerp over 100ms
+			float MouseAngle = FMath::Lerp(LastMouseAngle, FMath::Atan2(-MouseDirection.Y, MouseDirection.X), FMath::Min<float>(10.f * DeltaTime, 1.f));		// lerp over 100ms
 			LastMouseAngle = MouseAngle;
 			LastMouseSampleTime = SampleTime;
 			OldMousePosition = LastMousePosition;
@@ -1171,9 +1171,9 @@ public:
 		}
 	}
 
-	virtual void Tick(FEditorViewportClient* ViewportClient,float DeltaTime) override
+	virtual void Tick(FEditorViewportClient* ViewportClient, float DeltaTime) override
 	{
-		FLandscapeBrushCircle::Tick(ViewportClient,DeltaTime);
+		FLandscapeBrushCircle::Tick(ViewportClient, DeltaTime);
 
 		ULandscapeInfo* LandscapeInfo = EdMode->CurrentToolTarget.LandscapeInfo.Get();
 		if (LandscapeInfo)
@@ -1182,8 +1182,8 @@ public:
 			int32 SizeX = EdMode->UISettings->AlphaTextureSizeX;
 			int32 SizeY = EdMode->UISettings->AlphaTextureSizeY;
 			float Radius = EdMode->UISettings->BrushRadius / ScaleXY;
-			float MaxSize = 2.f * FMath::Sqrt( FMath::Square(Radius) / 2.f );
-			float AlphaBrushScale = MaxSize / (float)FMath::Max<int32>(SizeX,SizeY);
+			float MaxSize = 2.f * FMath::Sqrt(FMath::Square(Radius) / 2.f);
+			float AlphaBrushScale = MaxSize / (float)FMath::Max<int32>(SizeX, SizeY);
 
 			FLinearColor BrushScaleRot(
 				1.0f / (AlphaBrushScale * SizeX),
@@ -1193,7 +1193,7 @@ public:
 				);
 
 			int32 Channel = EdMode->UISettings->AlphaTextureChannel;
-			FLinearColor AlphaTextureMask(Channel==0?1:0, Channel==1?1:0, Channel==2?1:0, Channel==3?1:0);
+			FLinearColor AlphaTextureMask(Channel == 0 ? 1 : 0, Channel == 1 ? 1 : 0, Channel == 2 ? 1 : 0, Channel == 3 ? 1 : 0);
 
 			for (auto It = BrushMaterialInstanceMap.CreateConstIterator(); It; ++It)
 			{
@@ -1212,29 +1212,29 @@ public:
 
 void FEdModeLandscape::InitializeBrushes()
 {
-	FLandscapeBrushSet* BrushSet; 
-	BrushSet = new(LandscapeBrushSets) FLandscapeBrushSet(TEXT("BrushSet_Circle"));
+	FLandscapeBrushSet* BrushSet;
+	BrushSet = new(LandscapeBrushSets)FLandscapeBrushSet(TEXT("BrushSet_Circle"));
 	BrushSet->Brushes.Add(FLandscapeBrushCircle_Smooth::Create(this));
 	BrushSet->Brushes.Add(FLandscapeBrushCircle_Linear::Create(this));
 	BrushSet->Brushes.Add(FLandscapeBrushCircle_Spherical::Create(this));
 	BrushSet->Brushes.Add(FLandscapeBrushCircle_Tip::Create(this));
 
-	BrushSet = new(LandscapeBrushSets) FLandscapeBrushSet(TEXT("BrushSet_Alpha"));
+	BrushSet = new(LandscapeBrushSets)FLandscapeBrushSet(TEXT("BrushSet_Alpha"));
 	BrushSet->Brushes.Add(FLandscapeBrushAlpha::Create(this));
 
-	BrushSet = new(LandscapeBrushSets) FLandscapeBrushSet(TEXT("BrushSet_Pattern"));
+	BrushSet = new(LandscapeBrushSets)FLandscapeBrushSet(TEXT("BrushSet_Pattern"));
 	BrushSet->Brushes.Add(FLandscapeBrushAlphaPattern::Create(this));
 
-	BrushSet = new(LandscapeBrushSets) FLandscapeBrushSet(TEXT("BrushSet_Component"));
+	BrushSet = new(LandscapeBrushSets)FLandscapeBrushSet(TEXT("BrushSet_Component"));
 	BrushSet->Brushes.Add(new FLandscapeBrushComponent(this));
 
-	BrushSet = new(LandscapeBrushSets) FLandscapeBrushSet(TEXT("BrushSet_Gizmo"));
+	BrushSet = new(LandscapeBrushSets)FLandscapeBrushSet(TEXT("BrushSet_Gizmo"));
 	GizmoBrush = new FLandscapeBrushGizmo(this);
 	BrushSet->Brushes.Add(GizmoBrush);
 
-	BrushSet = new(LandscapeBrushSets) FLandscapeBrushSet(TEXT("BrushSet_Splines"));
+	BrushSet = new(LandscapeBrushSets)FLandscapeBrushSet(TEXT("BrushSet_Splines"));
 	BrushSet->Brushes.Add(new FLandscapeBrushSplines(this));
 
-	BrushSet = new(LandscapeBrushSets) FLandscapeBrushSet(TEXT("BrushSet_Dummy"));
+	BrushSet = new(LandscapeBrushSets)FLandscapeBrushSet(TEXT("BrushSet_Dummy"));
 	BrushSet->Brushes.Add(new FLandscapeBrushDummy(this));
 }
