Remove FLandscapeToolSet

diff --git a/Engine/Source/Editor/LandscapeEditor/Private/LandscapeEdModeComponentTools.cpp b/Engine/Source/Editor/LandscapeEditor/Private/LandscapeEdModeComponentTools.cpp
--- a/Engine/Source/Editor/LandscapeEditor/Private/LandscapeEdModeComponentTools.cpp
+++ b/Engine/Source/Editor/LandscapeEditor/Private/LandscapeEdModeComponentTools.cpp
@@ -24,42 +24,42 @@ class FLandscapeToolStrokeSelect : FLandscapeStrokeBase<const FLandscapeToolTarg
 
 public:
 	FLandscapeToolStrokeSelect(FEdModeLandscape* InEdMode, const FLandscapeToolTarget& InTarget)
-		:	FLandscapeStrokeBase<const FLandscapeToolTarget>(InTarget)
-		,	LandscapeInfo(InTarget.LandscapeInfo.Get())
-		,	Cache(InTarget)
-		,	bInitializedComponentInvert(false)
+		: FLandscapeStrokeBase<const FLandscapeToolTarget>(InTarget)
+		, LandscapeInfo(InTarget.LandscapeInfo.Get())
+		, Cache(InTarget)
+		, bInitializedComponentInvert(false)
 	{
 	}
 
 	virtual void Apply(FEditorViewportClient* ViewportClient, FLandscapeBrush* Brush, const ULandscapeEditorObject* UISettings, const TArray<FLandscapeToolMousePosition>& MousePositions)
 	{
-		if( LandscapeInfo )
+		if (LandscapeInfo)
 		{
 			LandscapeInfo->Modify();
 			// Get list of verts to update
 			TMap<FIntPoint, float> BrushInfo;
 			int32 X1, Y1, X2, Y2;
-			if (!Brush->ApplyBrush(MousePositions,BrushInfo, X1, Y1, X2, Y2))
+			if (!Brush->ApplyBrush(MousePositions, BrushInfo, X1, Y1, X2, Y2))
 			{
 				return;
 			}
 
 			// Invert when holding Shift
-			bool bInvert = MousePositions[MousePositions.Num()-1].bShiftDown;
+			bool bInvert = MousePositions[MousePositions.Num() - 1].bShiftDown;
 
-			if ( Brush->GetBrushType() == FLandscapeBrush::BT_Component )
+			if (Brush->GetBrushType() == FLandscapeBrush::BT_Component)
 			{
 				// Todo hold selection... static?
 				TSet<ULandscapeComponent*> NewComponents;
-				LandscapeInfo->GetComponentsInRegion(X1+1,Y1+1,X2-1,Y2-1,NewComponents);
+				LandscapeInfo->GetComponentsInRegion(X1 + 1, Y1 + 1, X2 - 1, Y2 - 1, NewComponents);
 
 				if (!bInitializedComponentInvert)
 				{
 					// Get the component under the mouse location. Copied from FLandscapeBrushComponent::ApplyBrush()
 					const int32 MouseX = FMath::RoundToInt(MousePositions[0].PositionX);
 					const int32 MouseY = FMath::RoundToInt(MousePositions[0].PositionY);
-					const int32 MouseComponentIndexX = (MouseX >= 0.f) ? MouseX / LandscapeInfo->ComponentSizeQuads : (MouseX+1) / LandscapeInfo->ComponentSizeQuads - 1;
-					const int32 MouseComponentIndexY = (MouseY >= 0.f) ? MouseY / LandscapeInfo->ComponentSizeQuads : (MouseY+1) / LandscapeInfo->ComponentSizeQuads - 1;
+					const int32 MouseComponentIndexX = (MouseX >= 0.f) ? MouseX / LandscapeInfo->ComponentSizeQuads : (MouseX + 1) / LandscapeInfo->ComponentSizeQuads - 1;
+					const int32 MouseComponentIndexY = (MouseY >= 0.f) ? MouseY / LandscapeInfo->ComponentSizeQuads : (MouseY + 1) / LandscapeInfo->ComponentSizeQuads - 1;
 					ULandscapeComponent* MouseComponent = LandscapeInfo->XYtoComponentMap.FindRef(FIntPoint(MouseComponentIndexX, MouseComponentIndexY));
 
 					if (MouseComponent != NULL)
@@ -96,7 +96,7 @@ public:
 				{
 					Objects.Add(*It);
 				}
-				FPropertyEditorModule& PropertyModule = FModuleManager::Get().LoadModuleChecked<FPropertyEditorModule>( TEXT("PropertyEditor") );
+				FPropertyEditorModule& PropertyModule = FModuleManager::Get().LoadModuleChecked<FPropertyEditorModule>(TEXT("PropertyEditor"));
 				PropertyModule.UpdatePropertyViews(Objects);
 			}
 			else // Select various shape regions
@@ -109,19 +109,19 @@ public:
 				// Tablet pressure
 				float Pressure = ViewportClient->Viewport->IsPenActive() ? ViewportClient->Viewport->GetTabletPressure() : 1.f;
 
-				Cache.CacheData(X1,Y1,X2,Y2);
+				Cache.CacheData(X1, Y1, X2, Y2);
 				TArray<uint8> Data;
-				Cache.GetCachedData(X1,Y1,X2,Y2,Data);
+				Cache.GetCachedData(X1, Y1, X2, Y2, Data);
 
 				TSet<ULandscapeComponent*> NewComponents;
 				// Remove invalid regions
-				LandscapeInfo->GetComponentsInRegion(X1,Y1,X2,Y2,NewComponents);
+				LandscapeInfo->GetComponentsInRegion(X1, Y1, X2, Y2, NewComponents);
 				LandscapeInfo->UpdateSelectedComponents(NewComponents, false);
 
-				for( auto It = BrushInfo.CreateIterator(); It; ++It )
+				for (auto It = BrushInfo.CreateIterator(); It; ++It)
 				{
 					FIntPoint Key = It.Key();
-					if( It.Value() > 0.f && LandscapeInfo->IsValidPosition(Key.X, Key.Y) )
+					if (It.Value() > 0.f && LandscapeInfo->IsValidPosition(Key.X, Key.Y))
 					{
 						float PaintValue = It.Value() * UISettings->ToolStrength * Pressure;
 						float Value = LandscapeInfo->SelectedRegion.FindRef(Key);
@@ -142,18 +142,18 @@ public:
 							LandscapeInfo->SelectedRegion.Add(Key, FMath::Min(Value + PaintValue, 1.0f));
 						}
 
-						Data[(Key.X-X1) + (Key.Y-Y1)*(1+X2-X1)] = FMath::Clamp<int32>(FMath::RoundToInt(LandscapeInfo->SelectedRegion.FindRef(Key) * 255), 0, 255);
+						Data[(Key.X - X1) + (Key.Y - Y1)*(1 + X2 - X1)] = FMath::Clamp<int32>(FMath::RoundToInt(LandscapeInfo->SelectedRegion.FindRef(Key) * 255), 0, 255);
 					}
 				}
 
-				Cache.SetCachedData(X1,Y1,X2,Y2,Data);
+				Cache.SetCachedData(X1, Y1, X2, Y2, Data);
 				Cache.Flush();
 			}
 		}
 	}
 
 protected:
-	class ULandscapeInfo* LandscapeInfo;
+	ULandscapeInfo* LandscapeInfo;
 	FLandscapeDataCache Cache;
 };
 
@@ -164,28 +164,23 @@ template<class TStrokeClass>
 class FLandscapeToolSelect : public FLandscapeToolBase<TStrokeClass>
 {
 public:
-	FLandscapeToolSelect(class FEdModeLandscape* InEdMode)
+	FLandscapeToolSelect(FEdModeLandscape* InEdMode)
 		: FLandscapeToolBase<TStrokeClass>(InEdMode)
 	{}
 
-	virtual const TCHAR* GetToolName() override { return TEXT("Selection"); }
+	virtual const TCHAR* GetToolName() override { return TEXT("Select"); }
 	virtual FText GetDisplayName() override { return NSLOCTEXT("UnrealEd", "LandscapeMode_Selection", "Component Selection"); };
 	virtual void SetEditRenderType() override { GLandscapeEditRenderMode = ELandscapeEditRenderMode::SelectComponent | (GLandscapeEditRenderMode & ELandscapeEditRenderMode::BitMaskForMask); }
 	virtual bool SupportsMask() override { return false; }
 
 	virtual FLandscapeTool::EToolType GetToolType() override { return this->TT_Mask; }
-
-	virtual bool IsValidForTarget(const FLandscapeToolTarget& Target) override
-	{
-		return true; // applied to all...
-	}
 };
 
 template<class TStrokeClass>
 class FLandscapeToolMask : public FLandscapeToolSelect<TStrokeClass>
 {
 public:
-	FLandscapeToolMask(class FEdModeLandscape* InEdMode)
+	FLandscapeToolMask(FEdModeLandscape* InEdMode)
 		: FLandscapeToolSelect<TStrokeClass>(InEdMode)
 	{}
 
@@ -199,13 +194,13 @@ class FLandscapeToolStrokeVisibility
 {
 public:
 	FLandscapeToolStrokeVisibility(FEdModeLandscape* InEdMode, const FLandscapeToolTarget& InTarget)
-		:	LandscapeInfo(InTarget.LandscapeInfo.Get())
-		,	Cache(InTarget)
+		: LandscapeInfo(InTarget.LandscapeInfo.Get())
+		, Cache(InTarget)
 	{}
 
 	virtual void Apply(FEditorViewportClient* ViewportClient, FLandscapeBrush* Brush, const ULandscapeEditorObject* UISettings, const TArray<FLandscapeToolMousePosition>& MousePositions)
 	{
-		if( LandscapeInfo )
+		if (LandscapeInfo)
 		{
 			LandscapeInfo->Modify();
 			// Get list of verts to update
@@ -217,7 +212,7 @@ public:
 			}
 
 			// Invert when holding Shift
-			bool bInvert = MousePositions[MousePositions.Num()-1].bShiftDown;
+			bool bInvert = MousePositions[MousePositions.Num() - 1].bShiftDown;
 
 			X1 -= 1;
 			Y1 -= 1;
@@ -227,18 +222,18 @@ public:
 			// Tablet pressure
 			float Pressure = ViewportClient->Viewport->IsPenActive() ? ViewportClient->Viewport->GetTabletPressure() : 1.f;
 
-			Cache.CacheData(X1,Y1,X2,Y2);
+			Cache.CacheData(X1, Y1, X2, Y2);
 			TArray<uint8> Data;
-			Cache.GetCachedData(X1,Y1,X2,Y2,Data);
+			Cache.GetCachedData(X1, Y1, X2, Y2, Data);
 
-			for( auto It = BrushInfo.CreateConstIterator(); It; ++It )
+			for (auto It = BrushInfo.CreateConstIterator(); It; ++It)
 			{
 				int32 X, Y;
 				ALandscape::UnpackKey(It.Key(), X, Y);
 
-				if( It.Value() > 0.f )
+				if (It.Value() > 0.f)
 				{
-					float Value = Data[(X-X1) + (Y-Y1)*(1+X2-X1)];
+					float Value = Data[(X - X1) + (Y - Y1)*(1 + X2 - X1)];
 					if (bInvert)
 					{
 						Value = 0;
@@ -248,17 +243,17 @@ public:
 						Value = 255; // Just on and off for visibility, for masking...
 					}
 
-					Data[(X-X1) + (Y-Y1)*(1+X2-X1)] = FMath::Clamp<int32>(FMath::RoundToInt(Value), 0, 255);
+					Data[(X - X1) + (Y - Y1)*(1 + X2 - X1)] = FMath::Clamp<int32>(FMath::RoundToInt(Value), 0, 255);
 				}
 			}
 
-			Cache.SetCachedData(X1,Y1,X2,Y2,Data);
+			Cache.SetCachedData(X1, Y1, X2, Y2, Data);
 			Cache.Flush();
 		}
 	}
 
 protected:
-	class ULandscapeInfo* LandscapeInfo;
+	ULandscapeInfo* LandscapeInfo;
 	FLandscapeVisCache Cache;
 };
 
@@ -268,7 +263,7 @@ protected:
 class FLandscapeToolVisibility : public FLandscapeToolBase<FLandscapeToolStrokeVisibility>
 {
 public:
-	FLandscapeToolVisibility(class FEdModeLandscape* InEdMode)
+	FLandscapeToolVisibility(FEdModeLandscape* InEdMode)
 		: FLandscapeToolBase<FLandscapeToolStrokeVisibility>(InEdMode)
 	{}
 
@@ -278,9 +273,9 @@ public:
 	virtual void SetEditRenderType() override { GLandscapeEditRenderMode = ELandscapeEditRenderMode::None | (GLandscapeEditRenderMode & ELandscapeEditRenderMode::BitMaskForMask); }
 	virtual bool SupportsMask() override { return false; }
 
-	virtual bool IsValidForTarget(const FLandscapeToolTarget& Target) override
+	virtual ELandscapeToolTargetTypeMask::Type GetSupportedTargetTypes() override
 	{
-		return true; // applied to all...
+		return ELandscapeToolTargetTypeMask::Visibility;
 	}
 };
 
@@ -295,7 +290,7 @@ public:
 	{
 		ALandscape* Landscape = LandscapeInfo ? LandscapeInfo->LandscapeActor.Get() : NULL;
 
-		if( Landscape )
+		if (Landscape)
 		{
 			Landscape->Modify();
 			LandscapeInfo->Modify();
@@ -347,7 +342,7 @@ public:
 				bool bIsAllCurrentLevel = true;
 				for (TSet<ULandscapeComponent*>::TIterator It(SelectedComponents); It; ++It)
 				{
-					if ( (*It)->GetLandscapeProxy()->GetLevel() != World->GetCurrentLevel() )
+					if ((*It)->GetLandscapeProxy()->GetLevel() != World->GetCurrentLevel())
 					{
 						bIsAllCurrentLevel = false;
 					}
@@ -367,7 +362,7 @@ public:
 				for (TSet<ULandscapeComponent*>::TIterator It(SelectedComponents); It; ++It)
 				{
 					UMaterialInterface* LandscapeMaterial = (*It)->GetLandscapeMaterial();
-					if ( LandscapeMaterial && LandscapeMaterial->GetOutermost() == (*It)->GetOutermost() )
+					if (LandscapeMaterial && LandscapeMaterial->GetOutermost() == (*It)->GetOutermost())
 					{
 						ULandscapeComponent* Comp = *It;
 						RenameObjects.AddUnique(LandscapeMaterial);
@@ -379,16 +374,16 @@ public:
 
 				if (RenameObjects.Num())
 				{
-					if (FMessageDialog::Open( EAppMsgType::OkCancel,
+					if (FMessageDialog::Open(EAppMsgType::OkCancel,
 						FText::Format(
-							NSLOCTEXT("UnrealEd", "LandscapeMoveToStreamingLevel_SharedResources", "The following items must be moved out of the persistent level and into a package that can be shared between multiple levels:\n\n{0}"),
-								FText::FromString(MsgBoxList) ) ))
+						NSLOCTEXT("UnrealEd", "LandscapeMoveToStreamingLevel_SharedResources", "The following items must be moved out of the persistent level and into a package that can be shared between multiple levels:\n\n{0}"),
+						FText::FromString(MsgBoxList))))
 					{
 						FString Path = Landscape->GetOutermost()->GetName() + TEXT("_sharedassets/");
 						bool bSucceed = ObjectTools::RenameObjects(RenameObjects, false, TEXT(""), Path);
 						if (!bSucceed)
 						{
-							FMessageDialog::Open( EAppMsgType::Ok, NSLOCTEXT("UnrealEd", "LandscapeMoveToStreamingLevel_RenameFailed", "Move To Streaming Level did not succeed because shared resources could not be moved to a new package.") );
+							FMessageDialog::Open(EAppMsgType::Ok, NSLOCTEXT("UnrealEd", "LandscapeMoveToStreamingLevel_RenameFailed", "Move To Streaming Level did not succeed because shared resources could not be moved to a new package."));
 							return;
 						}
 					}
@@ -398,7 +393,7 @@ public:
 					}
 				}
 
-				GWarn->BeginSlowTask( LOCTEXT("BeginMovingLandscapeComponentsToCurrentLevelTask", "Moving Landscape components to current level"), true);
+				GWarn->BeginSlowTask(LOCTEXT("BeginMovingLandscapeComponentsToCurrentLevelTask", "Moving Landscape components to current level"), true);
 
 				TSet<ALandscapeProxy*> SelectProxies;
 				TSet<UTexture2D*> OldTextureSet;
@@ -408,10 +403,10 @@ public:
 
 				int32 Progress = 0;
 				LandscapeInfo->SortSelectedComponents();
-				int32 ComponentSizeVerts = Landscape->NumSubsections * (Landscape->SubsectionSizeQuads+1);
-				int32 NeedHeightmapSize = 1<<FMath::CeilLogTwo( ComponentSizeVerts );
+				int32 ComponentSizeVerts = Landscape->NumSubsections * (Landscape->SubsectionSizeQuads + 1);
+				int32 NeedHeightmapSize = 1 << FMath::CeilLogTwo(ComponentSizeVerts);
 
-				for(TSet<ULandscapeComponent*>::TConstIterator It(SelectedComponents);It;++It)
+				for (TSet<ULandscapeComponent*>::TConstIterator It(SelectedComponents); It; ++It)
 				{
 					ULandscapeComponent* Comp = *It;
 					SelectProxies.Add(Comp->GetLandscapeProxy());
@@ -431,21 +426,21 @@ public:
 				int32 TotalProgress = TargetSelectedComponents.Num() * TargetSelectedCollisionComponents.Num();
 
 				// Check which ones are need for height map change
-				for(TSet<ULandscapeComponent*>::TConstIterator It(TargetSelectedComponents);It;++It)
+				for (TSet<ULandscapeComponent*>::TConstIterator It(TargetSelectedComponents); It; ++It)
 				{
 					ULandscapeComponent* Comp = *It;
-					Comp->Modify();				
+					Comp->Modify();
 					OldTextureSet.Add(Comp->HeightmapTexture);
 				}
 
 				// Need to split all the component which share Heightmap with selected components
 				// Search neighbor only
-				for(TSet<ULandscapeComponent*>::TConstIterator It(TargetSelectedComponents);It;++It)
+				for (TSet<ULandscapeComponent*>::TConstIterator It(TargetSelectedComponents); It; ++It)
 				{
 					ULandscapeComponent* Comp = *It;
 					int32 SearchX = Comp->HeightmapTexture->Source.GetSizeX() / NeedHeightmapSize;
 					int32 SearchY = Comp->HeightmapTexture->Source.GetSizeY() / NeedHeightmapSize;
-					FIntPoint ComponentBase = Comp->GetSectionBase()/Comp->ComponentSizeQuads;
+					FIntPoint ComponentBase = Comp->GetSectionBase() / Comp->ComponentSizeQuads;
 
 					for (int32 Y = 0; Y < SearchY; ++Y)
 					{
@@ -454,8 +449,8 @@ public:
 							// Search for four directions...
 							for (int32 Dir = 0; Dir < 4; ++Dir)
 							{
-								int32 XDir = (Dir>>1) ? 1 : -1;
-								int32 YDir = (Dir%2) ? 1 : -1;
+								int32 XDir = (Dir >> 1) ? 1 : -1;
+								int32 YDir = (Dir % 2) ? 1 : -1;
 								ULandscapeComponent* Neighbor = LandscapeInfo->XYtoComponentMap.FindRef(ComponentBase + FIntPoint(XDir*X, YDir*Y));
 								if (Neighbor && Neighbor->HeightmapTexture == Comp->HeightmapTexture && !HeightmapUpdateComponents.Contains(Neighbor))
 								{
@@ -472,14 +467,14 @@ public:
 				}
 
 				// Changing Heightmap format for selected components
-				for(TSet<ULandscapeComponent*>::TConstIterator It(HeightmapUpdateComponents);It;++It)
+				for (TSet<ULandscapeComponent*>::TConstIterator It(HeightmapUpdateComponents); It; ++It)
 				{
 					ULandscapeComponent* Comp = *It;
 					ALandscape::SplitHeightmap(Comp, (Comp->HeightmapScaleBias.X > 0.f));
 				}
 
 				// Delete if it is no referenced textures...
-				for(TSet<UTexture2D*>::TIterator It(OldTextureSet);It;++It)
+				for (TSet<UTexture2D*>::TIterator It(OldTextureSet); It; ++It)
 				{
 					(*It)->SetFlags(RF_Transactional);
 					(*It)->Modify();
@@ -493,7 +488,7 @@ public:
 					LandscapeProxy = World->SpawnActor<ALandscapeProxy>();
 					// copy shared properties to this new proxy
 					LandscapeProxy->GetSharedProperties(Landscape);
-					
+
 					// set proxy location
 					// by default first component location
 					ULandscapeComponent* FirstComponent = *TargetSelectedComponents.CreateConstIterator();
@@ -507,7 +502,7 @@ public:
 					}
 				}
 
-				for(TSet<ALandscapeProxy*>::TIterator It(SelectProxies);It;++It)
+				for (TSet<ALandscapeProxy*>::TIterator It(SelectProxies); It; ++It)
 				{
 					(*It)->Modify();
 				}
@@ -527,14 +522,14 @@ public:
 
 						// Code from ULandscapeComponent::ReallocateWeightmaps
 						// Move to other channels left
-						while( TotalNeededChannels > 0 )
+						while (TotalNeededChannels > 0)
 						{
 							// UE_LOG(LogLandscape, Log, TEXT("Still need %d channels"), TotalNeededChannels);
 
 							UTexture2D* CurrentWeightmapTexture = NULL;
 							FLandscapeWeightmapUsage* CurrentWeightmapUsage = NULL;
 
-							if( TotalNeededChannels < 4 )
+							if (TotalNeededChannels < 4)
 							{
 								// UE_LOG(LogLandscape, Log, TEXT("Looking for nearest"));
 
@@ -543,15 +538,15 @@ public:
 								for (TMap<UTexture2D*, struct FLandscapeWeightmapUsage>::TIterator WeightmapIt(LandscapeProxy->WeightmapUsageMap); WeightmapIt; ++WeightmapIt)
 								{
 									FLandscapeWeightmapUsage* TryWeightmapUsage = &WeightmapIt.Value();
-									if( TryWeightmapUsage->FreeChannelCount() >= TotalNeededChannels )
+									if (TryWeightmapUsage->FreeChannelCount() >= TotalNeededChannels)
 									{
 										// See if this candidate is closer than any others we've found
-										for( int32 ChanIdx=0;ChanIdx<4;ChanIdx++ )
+										for (int32 ChanIdx = 0; ChanIdx < 4; ChanIdx++)
 										{
-											if( TryWeightmapUsage->ChannelUsage[ChanIdx] != NULL  )
+											if (TryWeightmapUsage->ChannelUsage[ChanIdx] != NULL)
 											{
 												int32 TryDistanceSquared = (TryWeightmapUsage->ChannelUsage[ChanIdx]->GetSectionBase() - Comp->GetSectionBase()).SizeSquared();
-												if( TryDistanceSquared < BestDistanceSquared )
+												if (TryDistanceSquared < BestDistanceSquared)
 												{
 													CurrentWeightmapTexture = WeightmapIt.Key();
 													CurrentWeightmapUsage = TryWeightmapUsage;
@@ -563,14 +558,14 @@ public:
 								}
 							}
 
-							bool NeedsUpdateResource=false;
+							bool NeedsUpdateResource = false;
 							// No suitable weightmap texture
-							if( CurrentWeightmapTexture == NULL )
+							if (CurrentWeightmapTexture == NULL)
 							{
 								Comp->MarkPackageDirty();
 
 								// Weightmap is sized the same as the component
-								int32 WeightmapSize = (Comp->SubsectionSizeQuads+1) * Comp->NumSubsections;
+								int32 WeightmapSize = (Comp->SubsectionSizeQuads + 1) * Comp->NumSubsections;
 
 								// We need a new weightmap texture
 								CurrentWeightmapTexture = Comp->GetLandscapeProxy()->CreateLandscapeTexture(WeightmapSize, WeightmapSize, TEXTUREGROUP_Terrain_Weightmap, TSF_BGRA8);
@@ -586,27 +581,27 @@ public:
 
 							NewWeightmapTextures.Add(CurrentWeightmapTexture);
 
-							for( int32 ChanIdx=0;ChanIdx<4 && TotalNeededChannels > 0;ChanIdx++ )
+							for (int32 ChanIdx = 0; ChanIdx < 4 && TotalNeededChannels > 0; ChanIdx++)
 							{
 								// UE_LOG(LogLandscape, Log, TEXT("Finding allocation for layer %d"), CurrentLayer);
 
-								if( CurrentWeightmapUsage->ChannelUsage[ChanIdx] == NULL  )
+								if (CurrentWeightmapUsage->ChannelUsage[ChanIdx] == NULL)
 								{
 									// Use this allocation
 									FWeightmapLayerAllocationInfo& AllocInfo = Comp->WeightmapLayerAllocations[CurrentLayer];
 
-									if( AllocInfo.WeightmapTextureIndex == 255 )
+									if (AllocInfo.WeightmapTextureIndex == 255)
 									{
 										// New layer - zero out the data for this texture channel
-										LandscapeEdit.ZeroTextureChannel( CurrentWeightmapTexture, ChanIdx );
+										LandscapeEdit.ZeroTextureChannel(CurrentWeightmapTexture, ChanIdx);
 									}
 									else
 									{
 										UTexture2D* OldWeightmapTexture = Comp->WeightmapTextures[AllocInfo.WeightmapTextureIndex];
 
 										// Copy the data
-										LandscapeEdit.CopyTextureChannel( CurrentWeightmapTexture, ChanIdx, OldWeightmapTexture, AllocInfo.WeightmapTextureChannel );
-										LandscapeEdit.ZeroTextureChannel( OldWeightmapTexture, AllocInfo.WeightmapTextureChannel );
+										LandscapeEdit.CopyTextureChannel(CurrentWeightmapTexture, ChanIdx, OldWeightmapTexture, AllocInfo.WeightmapTextureChannel);
+										LandscapeEdit.ZeroTextureChannel(OldWeightmapTexture, AllocInfo.WeightmapTextureChannel);
 
 										// Remove the old allocation
 										FLandscapeWeightmapUsage* OldWeightmapUsage = Comp->GetLandscapeProxy()->WeightmapUsageMap.Find(OldWeightmapTexture);
@@ -615,7 +610,7 @@ public:
 
 									// Assign the new allocation
 									CurrentWeightmapUsage->ChannelUsage[ChanIdx] = Comp;
-									AllocInfo.WeightmapTextureIndex = NewWeightmapTextures.Num()-1;
+									AllocInfo.WeightmapTextureIndex = NewWeightmapTextures.Num() - 1;
 									AllocInfo.WeightmapTextureChannel = ChanIdx;
 									CurrentLayer++;
 									TotalNeededChannels--;
@@ -627,7 +622,7 @@ public:
 						Comp->WeightmapTextures = NewWeightmapTextures;
 
 						// Update the mipmaps for the textures we edited
-						for( int32 Idx=0;Idx<Comp->WeightmapTextures.Num();Idx++)
+						for (int32 Idx = 0; Idx < Comp->WeightmapTextures.Num(); Idx++)
 						{
 							UTexture2D* WeightmapTexture = Comp->WeightmapTextures[Idx];
 							FLandscapeTextureDataInfo* WeightmapDataInfo = LandscapeEdit.GetTextureDataInfo(WeightmapTexture);
@@ -635,7 +630,7 @@ public:
 							int32 NumMips = WeightmapTexture->Source.GetNumMips();
 							TArray<FColor*> WeightmapTextureMipData;
 							WeightmapTextureMipData.AddUninitialized(NumMips);
-							for( int32 MipIdx=0;MipIdx<NumMips;MipIdx++ )
+							for (int32 MipIdx = 0; MipIdx < NumMips; MipIdx++)
 							{
 								WeightmapTextureMipData[MipIdx] = (FColor*)WeightmapDataInfo->GetMipData(MipIdx);
 							}
@@ -649,7 +644,7 @@ public:
 
 				// Move the components to the Proxy actor
 				// This does not use the MoveSelectedActorsToCurrentLevel path as there is no support to only move certain components.
-				for(TSet<ULandscapeComponent*>::TConstIterator It(TargetSelectedComponents);It;++It)
+				for (TSet<ULandscapeComponent*>::TConstIterator It(TargetSelectedComponents); It; ++It)
 				{
 					// Need to move or recreate all related data (Height map, Weight map, maybe collision components, allocation info)
 					ULandscapeComponent* Comp = *It;
@@ -658,22 +653,22 @@ public:
 					Comp->DetachFromParent(true);
 					Comp->InvalidateLightingCache();
 					Comp->Rename(NULL, LandscapeProxy);
-					LandscapeProxy->LandscapeComponents.Add( Comp );
-					Comp->AttachTo( LandscapeProxy->GetRootComponent(), NAME_None, EAttachLocation::KeepWorldPosition );
+					LandscapeProxy->LandscapeComponents.Add(Comp);
+					Comp->AttachTo(LandscapeProxy->GetRootComponent(), NAME_None, EAttachLocation::KeepWorldPosition);
 					Comp->UpdateMaterialInstances();
 
 					FFormatNamedArguments Args;
-					Args.Add( TEXT("ComponentName"), FText::FromString( Comp->GetName() ) );
-					GWarn->StatusUpdate( Progress++, TotalProgress, FText::Format( LOCTEXT("MovingComponentStatus", "Moving Component: {ComponentName}"), Args ) );
+					Args.Add(TEXT("ComponentName"), FText::FromString(Comp->GetName()));
+					GWarn->StatusUpdate(Progress++, TotalProgress, FText::Format(LOCTEXT("MovingComponentStatus", "Moving Component: {ComponentName}"), Args));
 				}
 
-				for(TArray<ULandscapeHeightfieldCollisionComponent*>::TConstIterator It(TargetSelectedCollisionComponents);It;++It)
+				for (TArray<ULandscapeHeightfieldCollisionComponent*>::TConstIterator It(TargetSelectedCollisionComponents); It; ++It)
 				{
 					// Need to move or recreate all related data (Height map, Weight map, maybe collision components, allocation info)
 					ULandscapeHeightfieldCollisionComponent* Comp = *It;
 
 					// Move any foliage associated
-					AInstancedFoliageActor* OldIFA = AInstancedFoliageActor::GetInstancedFoliageActorForLevel( Comp->GetLandscapeProxy()->GetLevel() );
+					AInstancedFoliageActor* OldIFA = AInstancedFoliageActor::GetInstancedFoliageActorForLevel(Comp->GetLandscapeProxy()->GetLevel());
 					if (OldIFA)
 					{
 						OldIFA->MoveInstancesForComponentToCurrentLevel(Comp);
@@ -684,15 +679,15 @@ public:
 					Comp->DetachFromParent(true);
 					Comp->Rename(NULL, LandscapeProxy);
 					LandscapeProxy->CollisionComponents.Add(*It);
-					Comp->AttachTo( LandscapeProxy->GetRootComponent(), NAME_None, EAttachLocation::KeepWorldPosition );
+					Comp->AttachTo(LandscapeProxy->GetRootComponent(), NAME_None, EAttachLocation::KeepWorldPosition);
 
 					FFormatNamedArguments Args;
-					Args.Add( TEXT("ComponentName"), FText::FromString( Comp->GetName() ) );
-					GWarn->StatusUpdate( Progress++, TotalProgress, FText::Format( LOCTEXT("MovingComponentStatus", "Moving Component: {ComponentName}"), Args ) );
+					Args.Add(TEXT("ComponentName"), FText::FromString(Comp->GetName()));
+					GWarn->StatusUpdate(Progress++, TotalProgress, FText::Format(LOCTEXT("MovingComponentStatus", "Moving Component: {ComponentName}"), Args));
 				}
 
 				GEditor->SelectNone(false, true);
-				GEditor->SelectActor( LandscapeProxy, true, false, true );
+				GEditor->SelectActor(LandscapeProxy, true, false, true);
 
 				GEditor->SelectNone(false, true);
 
@@ -702,7 +697,7 @@ public:
 					LandscapeProxy->RegisterAllComponents();
 				}
 
-				for(TSet<ALandscapeProxy*>::TIterator It(SelectProxies);It;++It)
+				for (TSet<ALandscapeProxy*>::TIterator It(SelectProxies); It; ++It)
 				{
 					if ((*It)->GetRootComponent()->IsRegistered())
 					{
@@ -723,7 +718,7 @@ public:
 	}
 
 protected:
-	class ULandscapeInfo* LandscapeInfo;
+	ULandscapeInfo* LandscapeInfo;
 };
 
 // 
@@ -732,8 +727,8 @@ protected:
 class FLandscapeToolMoveToLevel : public FLandscapeToolBase<FLandscapeToolStrokeMoveToLevel>
 {
 public:
-	FLandscapeToolMoveToLevel(class FEdModeLandscape* InEdMode)
-		: FLandscapeToolBase<FLandscapeToolStrokeMoveToLevel>(InEdMode)	
+	FLandscapeToolMoveToLevel(FEdModeLandscape* InEdMode)
+		: FLandscapeToolBase<FLandscapeToolStrokeMoveToLevel>(InEdMode)
 	{}
 
 	virtual const TCHAR* GetToolName() override { return TEXT("MoveToLevel"); }
@@ -741,11 +736,6 @@ public:
 
 	virtual void SetEditRenderType() override { GLandscapeEditRenderMode = ELandscapeEditRenderMode::SelectComponent | (GLandscapeEditRenderMode & ELandscapeEditRenderMode::BitMaskForMask); }
 	virtual bool SupportsMask() override { return false; }
-
-	virtual bool IsValidForTarget(const FLandscapeToolTarget& Target) override
-	{
-		return true; // applied to all...
-	}
 };
 
 
@@ -753,15 +743,15 @@ class FLandscapeToolStrokeAddComponent
 {
 public:
 	FLandscapeToolStrokeAddComponent(FEdModeLandscape* InEdMode, const FLandscapeToolTarget& InTarget)
-		:	EdMode(InEdMode)
-		,	LandscapeInfo(InTarget.LandscapeInfo.Get())
-		,	HeightCache(InTarget)
+		: EdMode(InEdMode)
+		, LandscapeInfo(InTarget.LandscapeInfo.Get())
+		, HeightCache(InTarget)
 	{}
 
 	virtual void Apply(FEditorViewportClient* ViewportClient, FLandscapeBrush* Brush, const ULandscapeEditorObject* UISettings, const TArray<FLandscapeToolMousePosition>& MousePositions)
 	{
 		ALandscapeProxy* Landscape = LandscapeInfo ? LandscapeInfo->GetCurrentLevelLandscapeProxy(true) : NULL;
-		if( Landscape && EdMode->LandscapeRenderAddCollision)
+		if (Landscape && EdMode->LandscapeRenderAddCollision)
 		{
 			// Get list of verts to update
 			TMap<FIntPoint, float> BrushInfo;
@@ -777,25 +767,25 @@ public:
 			X2 += 1;
 			Y2 += 1;
 
-			HeightCache.CacheData(X1,Y1,X2,Y2);
+			HeightCache.CacheData(X1, Y1, X2, Y2);
 			TArray<uint16> Data;
-			HeightCache.GetCachedData(X1,Y1,X2,Y2,Data);
+			HeightCache.GetCachedData(X1, Y1, X2, Y2, Data);
 
 			// Find component range for this block of data, non shared vertices
-			int32 ComponentIndexX1 = (X1+1 >= 0) ? (X1+1) / Landscape->ComponentSizeQuads : (X1+2) / Landscape->ComponentSizeQuads - 1;
-			int32 ComponentIndexY1 = (Y1+1 >= 0) ? (Y1+1) / Landscape->ComponentSizeQuads : (Y1+2) / Landscape->ComponentSizeQuads - 1;
-			int32 ComponentIndexX2 = (X2-2 >= 0) ? (X2-2) / Landscape->ComponentSizeQuads : (X2-1) / Landscape->ComponentSizeQuads - 1;
-			int32 ComponentIndexY2 = (Y2-2 >= 0) ? (Y2-2) / Landscape->ComponentSizeQuads : (Y2-1) / Landscape->ComponentSizeQuads - 1;
+			int32 ComponentIndexX1 = (X1 + 1 >= 0) ? (X1 + 1) / Landscape->ComponentSizeQuads : (X1 + 2) / Landscape->ComponentSizeQuads - 1;
+			int32 ComponentIndexY1 = (Y1 + 1 >= 0) ? (Y1 + 1) / Landscape->ComponentSizeQuads : (Y1 + 2) / Landscape->ComponentSizeQuads - 1;
+			int32 ComponentIndexX2 = (X2 - 2 >= 0) ? (X2 - 2) / Landscape->ComponentSizeQuads : (X2 - 1) / Landscape->ComponentSizeQuads - 1;
+			int32 ComponentIndexY2 = (Y2 - 2 >= 0) ? (Y2 - 2) / Landscape->ComponentSizeQuads : (Y2 - 1) / Landscape->ComponentSizeQuads - 1;
 
 			TArray<ULandscapeComponent*> NewComponents;
 			Landscape->Modify();
 			LandscapeInfo->Modify();
-			for( int32 ComponentIndexY=ComponentIndexY1;ComponentIndexY<=ComponentIndexY2;ComponentIndexY++ )
+			for (int32 ComponentIndexY = ComponentIndexY1; ComponentIndexY <= ComponentIndexY2; ComponentIndexY++)
 			{
-				for( int32 ComponentIndexX=ComponentIndexX1;ComponentIndexX<=ComponentIndexX2;ComponentIndexX++ )
+				for (int32 ComponentIndexX = ComponentIndexX1; ComponentIndexX <= ComponentIndexX2; ComponentIndexX++)
 				{
-					ULandscapeComponent* Component = LandscapeInfo->XYtoComponentMap.FindRef(FIntPoint(ComponentIndexX,ComponentIndexY));
-					if( !Component )
+					ULandscapeComponent* Component = LandscapeInfo->XYtoComponentMap.FindRef(FIntPoint(ComponentIndexX, ComponentIndexY));
+					if (!Component)
 					{
 						// Add New component...
 						FIntPoint ComponentBase = FIntPoint(ComponentIndexX, ComponentIndexY)*Landscape->ComponentSizeQuads;
@@ -803,7 +793,7 @@ public:
 						Landscape->LandscapeComponents.Add(LandscapeComponent);
 						NewComponents.Add(LandscapeComponent);
 						LandscapeComponent->Init(
-							ComponentBase.X,ComponentBase.Y,
+							ComponentBase.X, ComponentBase.Y,
 							Landscape->ComponentSizeQuads,
 							Landscape->NumSubsections,
 							Landscape->SubsectionSizeQuads
@@ -814,14 +804,14 @@ public:
 						LandscapeComponent->bCastStaticShadow = Landscape->bCastStaticShadow;
 						LandscapeComponent->bCastShadowAsTwoSided = Landscape->bCastShadowAsTwoSided;
 
-						int32 ComponentVerts = (Landscape->SubsectionSizeQuads+1) * Landscape->NumSubsections;
+						int32 ComponentVerts = (Landscape->SubsectionSizeQuads + 1) * Landscape->NumSubsections;
 						// Update Weightmap Scale Bias
-						LandscapeComponent->WeightmapScaleBias = FVector4( 1.f / (float)ComponentVerts , 1.f / (float)ComponentVerts, 0.5f / (float)ComponentVerts , 0.5f / (float)ComponentVerts );
-						LandscapeComponent->WeightmapSubsectionOffset =  (float)(LandscapeComponent->SubsectionSizeQuads+1) / (float)ComponentVerts ;
+						LandscapeComponent->WeightmapScaleBias = FVector4(1.f / (float)ComponentVerts, 1.f / (float)ComponentVerts, 0.5f / (float)ComponentVerts, 0.5f / (float)ComponentVerts);
+						LandscapeComponent->WeightmapSubsectionOffset = (float)(LandscapeComponent->SubsectionSizeQuads + 1) / (float)ComponentVerts;
 
 						TArray<FColor> HeightData;
-						HeightData.Empty( FMath::Square(ComponentVerts) );
-						HeightData.AddZeroed( FMath::Square(ComponentVerts) );
+						HeightData.Empty(FMath::Square(ComponentVerts));
+						HeightData.AddZeroed(FMath::Square(ComponentVerts));
 						LandscapeComponent->InitHeightmapData(HeightData, true);
 						LandscapeComponent->UpdateMaterialInstances();
 					}
@@ -830,10 +820,10 @@ public:
 
 			Landscape->RegisterAllComponents();
 
-			HeightCache.SetCachedData(X1,Y1,X2,Y2,Data);
+			HeightCache.SetCachedData(X1, Y1, X2, Y2, Data);
 			HeightCache.Flush();
 
-			for ( int32 Idx = 0; Idx < NewComponents.Num(); Idx++ )
+			for (int32 Idx = 0; Idx < NewComponents.Num(); Idx++)
 			{
 				// Update LODBias from neighbors
 				FIntPoint ComponentBase = NewComponents[Idx]->GetSectionBase() / NewComponents[Idx]->ComponentSizeQuads;
@@ -867,19 +857,19 @@ public:
 				if (CollisionComp)
 				{
 					CollisionComp->MarkRenderStateDirty();
-					CollisionComp->RecreateCollision();	
+					CollisionComp->RecreateCollision();
 				}
 			}
-			
+
 			EdMode->LandscapeRenderAddCollision = NULL;
 
-			GEngine->BroadcastOnActorMoved( Landscape );
+			GEngine->BroadcastOnActorMoved(Landscape);
 		}
 	}
 
 protected:
-	class FEdModeLandscape* EdMode;
-	class ULandscapeInfo* LandscapeInfo;
+	FEdModeLandscape* EdMode;
+	ULandscapeInfo* LandscapeInfo;
 	FLandscapeHeightCache HeightCache;
 };
 
@@ -889,7 +879,7 @@ protected:
 class FLandscapeToolAddComponent : public FLandscapeToolBase<FLandscapeToolStrokeAddComponent>
 {
 public:
-	FLandscapeToolAddComponent(class FEdModeLandscape* InEdMode)
+	FLandscapeToolAddComponent(FEdModeLandscape* InEdMode)
 		: FLandscapeToolBase<FLandscapeToolStrokeAddComponent>(InEdMode)
 	{}
 
@@ -899,9 +889,11 @@ public:
 	virtual void SetEditRenderType() override { GLandscapeEditRenderMode = ELandscapeEditRenderMode::None | (GLandscapeEditRenderMode & ELandscapeEditRenderMode::BitMaskForMask); }
 	virtual bool SupportsMask() override { return false; }
 
-	virtual bool IsValidForTarget(const FLandscapeToolTarget& Target) override
+	virtual void ExitTool() override
 	{
-		return true; // applied to all...
+		FLandscapeToolBase<FLandscapeToolStrokeAddComponent>::ExitTool();
+
+		EdMode->LandscapeRenderAddCollision = nullptr;
 	}
 };
 
@@ -909,15 +901,15 @@ class FLandscapeToolStrokeDeleteComponent
 {
 public:
 	FLandscapeToolStrokeDeleteComponent(FEdModeLandscape* InEdMode, const FLandscapeToolTarget& InTarget)
-		:	LandscapeInfo(InTarget.LandscapeInfo.Get())
+		: LandscapeInfo(InTarget.LandscapeInfo.Get())
 	{}
 
 	virtual void Apply(FEditorViewportClient* ViewportClient, FLandscapeBrush* Brush, const ULandscapeEditorObject* UISettings, const TArray<FLandscapeToolMousePosition>& MousePositions)
 	{
-		if( LandscapeInfo )
+		if (LandscapeInfo)
 		{
 			LandscapeInfo->Modify();
-			
+
 			auto SelectedComponents = LandscapeInfo->GetSelectedComponents();
 			if (!SelectedComponents.Num())
 			{
@@ -928,21 +920,21 @@ public:
 				{
 					return;
 				}
-				LandscapeInfo->GetComponentsInRegion(X1+1,Y1+1,X2-1,Y2-1, SelectedComponents);
+				LandscapeInfo->GetComponentsInRegion(X1 + 1, Y1 + 1, X2 - 1, Y2 - 1, SelectedComponents);
 			}
 
-			int32 ComponentSizeVerts = LandscapeInfo->ComponentNumSubsections * (LandscapeInfo->SubsectionSizeQuads+1);
-			int32 NeedHeightmapSize = 1<<FMath::CeilLogTwo( ComponentSizeVerts );
+			int32 ComponentSizeVerts = LandscapeInfo->ComponentNumSubsections * (LandscapeInfo->SubsectionSizeQuads + 1);
+			int32 NeedHeightmapSize = 1 << FMath::CeilLogTwo(ComponentSizeVerts);
 
 			TSet<ULandscapeComponent*> HeightmapUpdateComponents;
 			// Need to split all the component which share Heightmap with selected components
 			// Search neighbor only
-			for(TSet<ULandscapeComponent*>::TConstIterator It(SelectedComponents);It;++It)
+			for (TSet<ULandscapeComponent*>::TConstIterator It(SelectedComponents); It; ++It)
 			{
 				ULandscapeComponent* Comp = *It;
 				int32 SearchX = Comp->HeightmapTexture->Source.GetSizeX() / NeedHeightmapSize;
 				int32 SearchY = Comp->HeightmapTexture->Source.GetSizeY() / NeedHeightmapSize;
-				FIntPoint ComponentBase = Comp->GetSectionBase()/Comp->ComponentSizeQuads;
+				FIntPoint ComponentBase = Comp->GetSectionBase() / Comp->ComponentSizeQuads;
 
 				for (int32 Y = 0; Y < SearchY; ++Y)
 				{
@@ -951,8 +943,8 @@ public:
 						// Search for four directions...
 						for (int32 Dir = 0; Dir < 4; ++Dir)
 						{
-							int32 XDir = (Dir>>1) ? 1 : -1;
-							int32 YDir = (Dir%2) ? 1 : -1;
+							int32 XDir = (Dir >> 1) ? 1 : -1;
+							int32 YDir = (Dir % 2) ? 1 : -1;
 							ULandscapeComponent* Neighbor = LandscapeInfo->XYtoComponentMap.FindRef(ComponentBase + FIntPoint(XDir*X, YDir*Y));
 							if (Neighbor && Neighbor->HeightmapTexture == Comp->HeightmapTexture && !HeightmapUpdateComponents.Contains(Neighbor))
 							{
@@ -965,7 +957,7 @@ public:
 			}
 
 			// Changing Heightmap format for selected components
-			for(TSet<ULandscapeComponent*>::TConstIterator It(HeightmapUpdateComponents);It;++It)
+			for (TSet<ULandscapeComponent*>::TConstIterator It(HeightmapUpdateComponents); It; ++It)
 			{
 				ULandscapeComponent* Comp = *It;
 				ALandscape::SplitHeightmap(Comp, false);
@@ -973,7 +965,7 @@ public:
 
 			TArray<FIntPoint> DeletedNeighborKeys;
 			// Check which ones are need for height map change
-			for(TSet<ULandscapeComponent*>::TIterator It(SelectedComponents);It;++It)
+			for (TSet<ULandscapeComponent*>::TIterator It(SelectedComponents); It; ++It)
 			{
 				ULandscapeComponent* Comp = *It;
 				ALandscapeProxy* Proxy = Comp->GetLandscapeProxy();
@@ -1009,7 +1001,7 @@ public:
 				}
 
 				// Remove Selected Region in deleted Component
-				for (int32 Y = 0; Y < Comp->ComponentSizeQuads; ++Y )
+				for (int32 Y = 0; Y < Comp->ComponentSizeQuads; ++Y)
 				{
 					for (int32 X = 0; X < Comp->ComponentSizeQuads; ++X)
 					{
@@ -1041,15 +1033,15 @@ public:
 					Comp->XYOffsetmapTexture->ClearFlags(RF_Standalone);
 				}
 
-				FIntPoint Key = Comp->GetSectionBase()/Comp->ComponentSizeQuads;
-				DeletedNeighborKeys.AddUnique(Key + FIntPoint(-1,-1));
-				DeletedNeighborKeys.AddUnique(Key + FIntPoint(+0,-1));
-				DeletedNeighborKeys.AddUnique(Key + FIntPoint(+1,-1));
-				DeletedNeighborKeys.AddUnique(Key + FIntPoint(-1,+0));
-				DeletedNeighborKeys.AddUnique(Key + FIntPoint(+1,+0));
-				DeletedNeighborKeys.AddUnique(Key + FIntPoint(-1,+1));
-				DeletedNeighborKeys.AddUnique(Key + FIntPoint(+0,+1));
-				DeletedNeighborKeys.AddUnique(Key + FIntPoint(+1,+1));
+				FIntPoint Key = Comp->GetSectionBase() / Comp->ComponentSizeQuads;
+				DeletedNeighborKeys.AddUnique(Key + FIntPoint(-1, -1));
+				DeletedNeighborKeys.AddUnique(Key + FIntPoint(+0, -1));
+				DeletedNeighborKeys.AddUnique(Key + FIntPoint(+1, -1));
+				DeletedNeighborKeys.AddUnique(Key + FIntPoint(-1, +0));
+				DeletedNeighborKeys.AddUnique(Key + FIntPoint(+1, +0));
+				DeletedNeighborKeys.AddUnique(Key + FIntPoint(-1, +1));
+				DeletedNeighborKeys.AddUnique(Key + FIntPoint(+0, +1));
+				DeletedNeighborKeys.AddUnique(Key + FIntPoint(+1, +1));
 
 				ULandscapeHeightfieldCollisionComponent* CollisionComp = Comp->CollisionComponent.Get();
 				if (CollisionComp)
@@ -1087,7 +1079,7 @@ public:
 	}
 
 protected:
-	class ULandscapeInfo* LandscapeInfo;
+	ULandscapeInfo* LandscapeInfo;
 };
 
 // 
@@ -1096,7 +1088,7 @@ protected:
 class FLandscapeToolDeleteComponent : public FLandscapeToolBase<FLandscapeToolStrokeDeleteComponent>
 {
 public:
-	FLandscapeToolDeleteComponent(class FEdModeLandscape* InEdMode)
+	FLandscapeToolDeleteComponent(FEdModeLandscape* InEdMode)
 		: FLandscapeToolBase<FLandscapeToolStrokeDeleteComponent>(InEdMode)
 	{}
 
@@ -1105,11 +1097,6 @@ public:
 
 	virtual void SetEditRenderType() override { GLandscapeEditRenderMode = ELandscapeEditRenderMode::SelectComponent | (GLandscapeEditRenderMode & ELandscapeEditRenderMode::BitMaskForMask); }
 	virtual bool SupportsMask() override { return false; }
-
-	virtual bool IsValidForTarget(const FLandscapeToolTarget& Target) override
-	{
-		return true; // applied to all...
-	}
 };
 
 template<class ToolTarget>
@@ -1117,14 +1104,14 @@ class FLandscapeToolStrokeCopy
 {
 public:
 	FLandscapeToolStrokeCopy(FEdModeLandscape* InEdMode, const FLandscapeToolTarget& InTarget)
-		:	EdMode(InEdMode)
-		,	LandscapeInfo(InTarget.LandscapeInfo.Get())
-		,	Cache(InTarget)
-		,	HeightCache(InTarget)
-		,	WeightCache(InTarget)
+		: EdMode(InEdMode)
+		, LandscapeInfo(InTarget.LandscapeInfo.Get())
+		, Cache(InTarget)
+		, HeightCache(InTarget)
+		, WeightCache(InTarget)
 	{}
 
-	struct FGizmoPreData 
+	struct FGizmoPreData
 	{
 		float Ratio;
 		float Data;
@@ -1170,16 +1157,16 @@ public:
 
 			if (bApplyToAll)
 			{
-				HeightCache.CacheData(X1,Y1,X2,Y2);
-				HeightCache.GetCachedData(X1,Y1,X2,Y2,HeightData);
+				HeightCache.CacheData(X1, Y1, X2, Y2);
+				HeightCache.GetCachedData(X1, Y1, X2, Y2, HeightData);
 
-				WeightCache.CacheData(X1,Y1,X2,Y2);
-				WeightCache.GetCachedData(X1,Y1,X2,Y2, WeightDatas, LayerNum);	
+				WeightCache.CacheData(X1, Y1, X2, Y2);
+				WeightCache.GetCachedData(X1, Y1, X2, Y2, WeightDatas, LayerNum);
 			}
 			else
 			{
-				Cache.CacheData(X1,Y1,X2,Y2);
-				Cache.GetCachedData(X1,Y1,X2,Y2,Data);
+				Cache.CacheData(X1, Y1, X2, Y2);
+				Cache.GetCachedData(X1, Y1, X2, Y2, Data);
 			}
 
 			const float ScaleXY = LandscapeInfo->DrawScale.X;
@@ -1191,7 +1178,7 @@ public:
 			Gizmo->CachedScaleXY = ScaleXY;
 
 			// Rasterize Gizmo regions
-			int32 SizeX = FMath::CeilToInt(Width  / ScaleXY);
+			int32 SizeX = FMath::CeilToInt(Width / ScaleXY);
 			int32 SizeY = FMath::CeilToInt(Height / ScaleXY);
 
 			const float W = (Width - ScaleXY) / (2 * ScaleXY);
@@ -1212,12 +1199,12 @@ public:
 			{
 				for (int32 X = 0; X < SizeX; ++X)
 				{
-					FVector LandscapeLocal = GizmoLocalToLandscape.TransformPosition(FVector( -W + X, -H + Y, 0 ));
+					FVector LandscapeLocal = GizmoLocalToLandscape.TransformPosition(FVector(-W + X, -H + Y, 0));
 					int32 LX = FMath::FloorToInt(LandscapeLocal.X);
 					int32 LY = FMath::FloorToInt(LandscapeLocal.Y);
 
 					{
-						for (int32 i = -1; (!bApplyToAll && i < 0) || i < LayerNum; ++i )
+						for (int32 i = -1; (!bApplyToAll && i < 0) || i < LayerNum; ++i)
 						{
 							FGizmoPreData GizmoPreData[NeighborNum];
 
@@ -1228,11 +1215,11 @@ public:
 									int32 x = FMath::Clamp(LX + LocalX, X1, X2);
 									int32 y = FMath::Clamp(LY + LocalY, Y1, Y2);
 									GizmoPreData[LocalX + LocalY * 2].Ratio = LandscapeInfo->SelectedRegion.FindRef(FIntPoint(x, y));
-									int32 index = (x-X1) + (y-Y1)*(1+X2-X1);
+									int32 index = (x - X1) + (y - Y1)*(1 + X2 - X1);
 
 									if (bApplyToAll)
 									{
-										if ( i < 0 )
+										if (i < 0)
 										{
 											GizmoPreData[LocalX + LocalY * 2].Data = Gizmo->GetNormalizedHeight(HeightData[index]);
 										}
@@ -1244,7 +1231,7 @@ public:
 									else
 									{
 										typename ToolTarget::CacheClass::DataType OriginalValue = Data[index];
-										if ( EdMode->CurrentToolTarget.TargetType == ELandscapeToolTargetType::Heightmap )
+										if (EdMode->CurrentToolTarget.TargetType == ELandscapeToolTargetType::Heightmap)
 										{
 											GizmoPreData[LocalX + LocalY * 2].Data = Gizmo->GetNormalizedHeight(OriginalValue);
 										}
@@ -1259,7 +1246,7 @@ public:
 							FGizmoPreData LerpedData;
 							float FracX = LandscapeLocal.X - LX;
 							float FracY = LandscapeLocal.Y - LY;
-							LerpedData.Ratio = bFullCopy ? 1.f : 
+							LerpedData.Ratio = bFullCopy ? 1.f :
 								FMath::Lerp(
 								FMath::Lerp(GizmoPreData[0].Ratio, GizmoPreData[1].Ratio, FracX),
 								FMath::Lerp(GizmoPreData[2].Ratio, GizmoPreData[3].Ratio, FracX),
@@ -1359,7 +1346,7 @@ public:
 			{
 				if (!bApplyToAll)
 				{
-					if ( EdMode->CurrentToolTarget.TargetType == ELandscapeToolTargetType::Heightmap )
+					if (EdMode->CurrentToolTarget.TargetType == ELandscapeToolTargetType::Heightmap)
 					{
 						Gizmo->DataType = ELandscapeGizmoType(Gizmo->DataType | LGT_Height);
 					}
@@ -1370,7 +1357,7 @@ public:
 				}
 				else
 				{
-					if ( LayerNum > 0 )
+					if (LayerNum > 0)
 					{
 						Gizmo->DataType = ELandscapeGizmoType(Gizmo->DataType | LGT_Height);
 						Gizmo->DataType = ELandscapeGizmoType(Gizmo->DataType | LGT_Weight);
@@ -1384,7 +1371,7 @@ public:
 				Gizmo->SampleData(SizeX, SizeY);
 
 				// Update LayerInfos
-				for (auto It = LayerInfoSet.CreateConstIterator(); It; ++It )
+				for (auto It = LayerInfoSet.CreateConstIterator(); It; ++It)
 				{
 					Gizmo->LayerInfos.Add(*It);
 				}
@@ -1407,8 +1394,8 @@ public:
 	}
 
 protected:
-	class FEdModeLandscape* EdMode;
-	class ULandscapeInfo* LandscapeInfo;
+	FEdModeLandscape* EdMode;
+	ULandscapeInfo* LandscapeInfo;
 	typename ToolTarget::CacheClass Cache;
 	FLandscapeHeightCache HeightCache;
 	FLandscapeFullWeightCache WeightCache;
@@ -1418,10 +1405,10 @@ protected:
 // FLandscapeToolCopy
 //
 template<class ToolTarget>
-class FLandscapeToolCopy : public FLandscapeToolBase<FLandscapeToolStrokeCopy<ToolTarget> >
+class FLandscapeToolCopy : public FLandscapeToolBase<FLandscapeToolStrokeCopy<ToolTarget>>
 {
 public:
-	FLandscapeToolCopy(class FEdModeLandscape* InEdMode)
+	FLandscapeToolCopy(FEdModeLandscape* InEdMode)
 		: FLandscapeToolBase<FLandscapeToolStrokeCopy<ToolTarget> >(InEdMode)
 	{}
 
@@ -1429,19 +1416,19 @@ public:
 	virtual FText GetDisplayName() override { return NSLOCTEXT("UnrealEd", "LandscapeMode_Copy", "Copy"); };
 
 	virtual void SetEditRenderType() override
-	{ 
+	{
 		GLandscapeEditRenderMode = ELandscapeEditRenderMode::Gizmo | (GLandscapeEditRenderMode & ELandscapeEditRenderMode::BitMaskForMask);
 		GLandscapeEditRenderMode |= (this->EdMode && this->EdMode->CurrentToolTarget.LandscapeInfo.IsValid() && this->EdMode->CurrentToolTarget.LandscapeInfo->SelectedRegion.Num()) ? ELandscapeEditRenderMode::SelectRegion : ELandscapeEditRenderMode::SelectComponent;
 	}
 
-	virtual bool IsValidForTarget(const FLandscapeToolTarget& Target) override
+	virtual ELandscapeToolTargetTypeMask::Type GetSupportedTargetTypes() override
 	{
-		return Target.TargetType == ToolTarget::TargetType;
+		return ELandscapeToolTargetTypeMask::FromType(ToolTarget::TargetType);
 	}
 
-	virtual bool BeginTool( FEditorViewportClient* ViewportClient, const FLandscapeToolTarget& InTarget, const FVector& InHitLocation ) override
+	virtual bool BeginTool(FEditorViewportClient* ViewportClient, const FLandscapeToolTarget& InTarget, const FVector& InHitLocation) override
 	{
-		check(this->MousePositions.Num()==0);
+		check(this->MousePositions.Num() == 0);
 		this->bToolActive = true;
 		this->ToolStroke = new FLandscapeToolStrokeCopy<ToolTarget>(this->EdMode, InTarget);
 
@@ -1454,9 +1441,9 @@ public:
 		return true;
 	}
 
-	virtual void EndTool(FEditorViewportClient* ViewportClient) override 
+	virtual void EndTool(FEditorViewportClient* ViewportClient) override
 	{
-		if( this->bToolActive && this->MousePositions.Num() )
+		if (this->bToolActive && this->MousePositions.Num())
 		{
 			this->ToolStroke->Apply(ViewportClient, this->EdMode->CurrentBrush, this->EdMode->UISettings, this->MousePositions);
 			this->MousePositions.Empty();
@@ -1468,10 +1455,10 @@ public:
 		this->EdMode->GizmoBrush->EndStroke();
 	}
 
-	virtual bool MouseMove( FEditorViewportClient* ViewportClient, FViewport* Viewport, int x, int y ) override
+	virtual bool MouseMove(FEditorViewportClient* ViewportClient, FViewport* Viewport, int x, int y) override
 	{
 		return true;
-	}	
+	}
 
 };
 
@@ -1480,11 +1467,11 @@ class FLandscapeToolStrokePaste
 {
 public:
 	FLandscapeToolStrokePaste(FEdModeLandscape* InEdMode, const FLandscapeToolTarget& InTarget)
-		:	EdMode(InEdMode)
-		,	LandscapeInfo(InTarget.LandscapeInfo.Get())
-		,	Cache(InTarget)
-		,	HeightCache(InTarget)
-		,	WeightCache(InTarget)
+		: EdMode(InEdMode)
+		, LandscapeInfo(InTarget.LandscapeInfo.Get())
+		, Cache(InTarget)
+		, HeightCache(InTarget)
+		, WeightCache(InTarget)
 	{}
 
 	virtual void Apply(FEditorViewportClient* ViewportClient, FLandscapeBrush* Brush, const ULandscapeEditorObject* UISettings, const TArray<FLandscapeToolMousePosition>& MousePositions)
@@ -1530,19 +1517,19 @@ public:
 
 			if (bApplyToAll)
 			{
-				HeightCache.CacheData(X1,Y1,X2,Y2);
-				HeightCache.GetCachedData(X1,Y1,X2,Y2,HeightData);
+				HeightCache.CacheData(X1, Y1, X2, Y2);
+				HeightCache.GetCachedData(X1, Y1, X2, Y2, HeightData);
 
 				if (LayerNum > 0)
 				{
-					WeightCache.CacheData(X1,Y1,X2,Y2);
-					WeightCache.GetCachedData(X1,Y1,X2,Y2, WeightDatas, LayerNum);
+					WeightCache.CacheData(X1, Y1, X2, Y2);
+					WeightCache.GetCachedData(X1, Y1, X2, Y2, WeightDatas, LayerNum);
 				}
 			}
 			else
 			{
-				Cache.CacheData(X1,Y1,X2,Y2);
-				Cache.GetCachedData(X1,Y1,X2,Y2,Data);
+				Cache.CacheData(X1, Y1, X2, Y2);
+				Cache.GetCachedData(X1, Y1, X2, Y2, Data);
 			}
 
 			const float Width = Gizmo->GetWidth();
@@ -1561,18 +1548,18 @@ public:
 			//FMatrix LToW = Landscape->LocalToWorld();
 			FVector BaseLocation = WToL.TransformPosition(Gizmo->GetActorLocation());
 			//FMatrix LandscapeLocalToGizmo = FRotationTranslationMatrix(FRotator(0, Gizmo->Rotation.Yaw, 0), FVector(BaseLocation.X - W + 0.5, BaseLocation.Y - H + 0.5, 0));
-			FMatrix LandscapeToGizmoLocal = 
-				(FTranslationMatrix(FVector((- W + 0.5)*SignX, (- H + 0.5)*SignY, 0)) * FScaleRotationTranslationMatrix(FVector(SignX, SignY, 1.f), FRotator(0, Gizmo->GetActorRotation().Yaw, 0), FVector(BaseLocation.X, BaseLocation.Y, 0))).Inverse();
+			FMatrix LandscapeToGizmoLocal =
+				(FTranslationMatrix(FVector((-W + 0.5)*SignX, (-H + 0.5)*SignY, 0)) * FScaleRotationTranslationMatrix(FVector(SignX, SignY, 1.f), FRotator(0, Gizmo->GetActorRotation().Yaw, 0), FVector(BaseLocation.X, BaseLocation.Y, 0))).Inverse();
 
-			for( auto It = BrushInfo.CreateIterator(); It; ++It )
+			for (auto It = BrushInfo.CreateIterator(); It; ++It)
 			{
 				int32 X, Y;
 				ALandscape::UnpackKey(It.Key(), X, Y);
 
-				if( It.Value() > 0.f )
+				if (It.Value() > 0.f)
 				{
 					// Value before we apply our painting
-					int32 index = (X-X1) + (Y-Y1)*(1+X2-X1);
+					int32 index = (X - X1) + (Y - Y1)*(1 + X2 - X1);
 					float PaintAmount = (Brush->GetBrushType() == FLandscapeBrush::BT_Gizmo) ? It.Value() : It.Value() * EdMode->UISettings->ToolStrength * Pressure;
 
 					FVector GizmoLocal = LandscapeToGizmoLocal.TransformPosition(FVector(X, Y, 0));
@@ -1586,14 +1573,14 @@ public:
 					float FracY = GizmoLocal.Y - LY;
 
 					FGizmoSelectData* Data00 = Gizmo->SelectedData.Find(ALandscape::MakeKey(LX, LY));
-					FGizmoSelectData* Data10 = Gizmo->SelectedData.Find(ALandscape::MakeKey(LX+1, LY));
-					FGizmoSelectData* Data01 = Gizmo->SelectedData.Find(ALandscape::MakeKey(LX, LY+1));
-					FGizmoSelectData* Data11 = Gizmo->SelectedData.Find(ALandscape::MakeKey(LX+1, LY+1));
+					FGizmoSelectData* Data10 = Gizmo->SelectedData.Find(ALandscape::MakeKey(LX + 1, LY));
+					FGizmoSelectData* Data01 = Gizmo->SelectedData.Find(ALandscape::MakeKey(LX, LY + 1));
+					FGizmoSelectData* Data11 = Gizmo->SelectedData.Find(ALandscape::MakeKey(LX + 1, LY + 1));
 
 					{
-						for (int32 i = -1; (!bApplyToAll && i < 0) || i < LayerNum; ++i )
+						for (int32 i = -1; (!bApplyToAll && i < 0) || i < LayerNum; ++i)
 						{
-							if ( (bApplyToAll && (i < 0)) || (!bApplyToAll && EdMode->CurrentToolTarget.TargetType == ELandscapeToolTargetType::Heightmap) )
+							if ((bApplyToAll && (i < 0)) || (!bApplyToAll && EdMode->CurrentToolTarget.TargetType == ELandscapeToolTargetType::Heightmap))
 							{
 								float OriginalValue;
 								if (bApplyToAll)
@@ -1610,15 +1597,15 @@ public:
 								float DestValue = FLandscapeHeightCache::ClampValue(
 									LandscapeDataAccess::GetTexHeight(
 									FMath::Lerp(
-									FMath::Lerp(Data00 ? FMath::Lerp(Value, Gizmo->GetLandscapeHeight(Data00->HeightData), Data00->Ratio) : Value, 
+									FMath::Lerp(Data00 ? FMath::Lerp(Value, Gizmo->GetLandscapeHeight(Data00->HeightData), Data00->Ratio) : Value,
 									Data10 ? FMath::Lerp(Value, Gizmo->GetLandscapeHeight(Data10->HeightData), Data10->Ratio) : Value, FracX),
-									FMath::Lerp(Data01 ? FMath::Lerp(Value, Gizmo->GetLandscapeHeight(Data01->HeightData), Data01->Ratio) : Value, 
+									FMath::Lerp(Data01 ? FMath::Lerp(Value, Gizmo->GetLandscapeHeight(Data01->HeightData), Data01->Ratio) : Value,
 									Data11 ? FMath::Lerp(Value, Gizmo->GetLandscapeHeight(Data11->HeightData), Data11->Ratio) : Value, FracX),
 									FracY
-									) ) 
+									))
 									);
 
-								switch(EdMode->UISettings->PasteMode)
+								switch (EdMode->UISettings->PasteMode)
 								{
 								case ELandscapeToolNoiseMode::Add:
 									PaintAmount = OriginalValue < DestValue ? PaintAmount : 0.f;
@@ -1632,11 +1619,11 @@ public:
 
 								if (bApplyToAll)
 								{
-									HeightData[index] = FMath::Lerp( OriginalValue, DestValue, PaintAmount );
+									HeightData[index] = FMath::Lerp(OriginalValue, DestValue, PaintAmount);
 								}
 								else
 								{
-									Data[index] = FMath::Lerp( OriginalValue, DestValue, PaintAmount );
+									Data[index] = FMath::Lerp(OriginalValue, DestValue, PaintAmount);
 								}
 							}
 							else
@@ -1656,20 +1643,20 @@ public:
 
 								float DestValue = FLandscapeAlphaCache::ClampValue(
 									FMath::Lerp(
-									FMath::Lerp(Data00 ? FMath::Lerp(OriginalValue, Data00->WeightDataMap.FindRef(LayerInfo), Data00->Ratio) : OriginalValue, 
+									FMath::Lerp(Data00 ? FMath::Lerp(OriginalValue, Data00->WeightDataMap.FindRef(LayerInfo), Data00->Ratio) : OriginalValue,
 									Data10 ? FMath::Lerp(OriginalValue, Data10->WeightDataMap.FindRef(LayerInfo), Data10->Ratio) : OriginalValue, FracX),
-									FMath::Lerp(Data01 ? FMath::Lerp(OriginalValue, Data01->WeightDataMap.FindRef(LayerInfo), Data01->Ratio) : OriginalValue, 
+									FMath::Lerp(Data01 ? FMath::Lerp(OriginalValue, Data01->WeightDataMap.FindRef(LayerInfo), Data01->Ratio) : OriginalValue,
 									Data11 ? FMath::Lerp(OriginalValue, Data11->WeightDataMap.FindRef(LayerInfo), Data11->Ratio) : OriginalValue, FracX),
 									FracY
 									));
 
 								if (bApplyToAll)
 								{
-									WeightDatas[index*LayerNum + i] = FMath::Lerp( OriginalValue, DestValue, PaintAmount );
+									WeightDatas[index*LayerNum + i] = FMath::Lerp(OriginalValue, DestValue, PaintAmount);
 								}
 								else
 								{
-									Data[index] = FMath::Lerp( OriginalValue, DestValue, PaintAmount );
+									Data[index] = FMath::Lerp(OriginalValue, DestValue, PaintAmount);
 								}
 							}
 						}
@@ -1684,18 +1671,18 @@ public:
 
 			if (bApplyToAll)
 			{
-				HeightCache.SetCachedData(X1,Y1,X2,Y2,HeightData);
+				HeightCache.SetCachedData(X1, Y1, X2, Y2, HeightData);
 				HeightCache.Flush();
 				if (WeightDatas.Num())
 				{
 					// Set the layer data, bypassing painting restrictions because it doesn't work well when altering multiple layers
-					WeightCache.SetCachedData(X1,Y1,X2,Y2,WeightDatas, LayerNum, ELandscapeLayerPaintingRestriction::None);
+					WeightCache.SetCachedData(X1, Y1, X2, Y2, WeightDatas, LayerNum, ELandscapeLayerPaintingRestriction::None);
 				}
 				WeightCache.Flush();
 			}
 			else
 			{
-				Cache.SetCachedData(X1,Y1,X2,Y2,Data);
+				Cache.SetCachedData(X1, Y1, X2, Y2, Data);
 				Cache.Flush();
 			}
 
@@ -1704,8 +1691,8 @@ public:
 	}
 
 protected:
-	class FEdModeLandscape* EdMode;
-	class ULandscapeInfo* LandscapeInfo;
+	FEdModeLandscape* EdMode;
+	ULandscapeInfo* LandscapeInfo;
 	typename ToolTarget::CacheClass Cache;
 	FLandscapeHeightCache HeightCache;
 	FLandscapeFullWeightCache WeightCache;
@@ -1715,19 +1702,19 @@ protected:
 // FLandscapeToolPaste
 //
 template<class ToolTarget>
-class FLandscapeToolPaste : public FLandscapeToolBase<FLandscapeToolStrokePaste<ToolTarget> >
+class FLandscapeToolPaste : public FLandscapeToolBase<FLandscapeToolStrokePaste<ToolTarget>>
 {
 public:
-	FLandscapeToolPaste(class FEdModeLandscape* InEdMode)
+	FLandscapeToolPaste(FEdModeLandscape* InEdMode)
 		: FLandscapeToolBase<FLandscapeToolStrokePaste<ToolTarget> >(InEdMode)
-		,	bUseGizmoRegion(false)
+		, bUseGizmoRegion(false)
 	{}
 
 	virtual const TCHAR* GetToolName() override { return TEXT("Paste"); }
 	virtual FText GetDisplayName() override { return NSLOCTEXT("UnrealEd", "LandscapeMode_Region", "Region Copy/Paste"); };
 
-	virtual void SetEditRenderType() override 
-	{ 
+	virtual void SetEditRenderType() override
+	{
 		GLandscapeEditRenderMode = ELandscapeEditRenderMode::Gizmo | (GLandscapeEditRenderMode & ELandscapeEditRenderMode::BitMaskForMask);
 		GLandscapeEditRenderMode |= (this->EdMode && this->EdMode->CurrentToolTarget.LandscapeInfo.IsValid() && this->EdMode->CurrentToolTarget.LandscapeInfo->SelectedRegion.Num()) ? ELandscapeEditRenderMode::SelectRegion : ELandscapeEditRenderMode::SelectComponent;
 	}
@@ -1737,12 +1724,12 @@ public:
 		bUseGizmoRegion = InbUseGizmoRegion;
 	}
 
-	virtual bool IsValidForTarget(const FLandscapeToolTarget& Target) override
+	virtual ELandscapeToolTargetTypeMask::Type GetSupportedTargetTypes() override
 	{
-		return Target.TargetType == ToolTarget::TargetType;
+		return ELandscapeToolTargetTypeMask::FromType(ToolTarget::TargetType);
 	}
 
-	virtual bool BeginTool( FEditorViewportClient* ViewportClient, const FLandscapeToolTarget& InTarget, const FVector& InHitLocation ) override
+	virtual bool BeginTool(FEditorViewportClient* ViewportClient, const FLandscapeToolTarget& InTarget, const FVector& InHitLocation) override
 	{
 		this->bToolActive = true;
 		this->ToolStroke = new FLandscapeToolStrokePaste<ToolTarget>(this->EdMode, InTarget);
@@ -1765,7 +1752,7 @@ public:
 
 	virtual void EndTool(FEditorViewportClient* ViewportClient) override
 	{
-		if( this->bToolActive && this->MousePositions.Num() )
+		if (this->bToolActive && this->MousePositions.Num())
 		{
 			this->ToolStroke->Apply(ViewportClient, this->EdMode->CurrentBrush, this->EdMode->UISettings, this->MousePositions);
 			this->MousePositions.Empty();
@@ -1784,7 +1771,7 @@ public:
 		}
 	}
 
-	virtual bool MouseMove( FEditorViewportClient* ViewportClient, FViewport* Viewport, int32 x, int32 y ) override
+	virtual bool MouseMove(FEditorViewportClient* ViewportClient, FViewport* Viewport, int32 x, int32 y) override
 	{
 		if (bUseGizmoRegion)
 		{
@@ -1793,7 +1780,7 @@ public:
 
 		FLandscapeToolBase<FLandscapeToolStrokePaste<ToolTarget> >::MouseMove(ViewportClient, Viewport, x, y);
 		return true;
-	}	
+	}
 protected:
 	bool bUseGizmoRegion;
 };
@@ -1802,62 +1789,73 @@ template<class ToolTarget>
 class FLandscapeToolCopyPaste : public FLandscapeToolPaste<ToolTarget>
 {
 public:
-	FLandscapeToolCopyPaste(class FEdModeLandscape* InEdMode)
+	FLandscapeToolCopyPaste(FEdModeLandscape* InEdMode)
 		: FLandscapeToolPaste<ToolTarget>(InEdMode)
+		, CopyTool(InEdMode)
 	{
-		pCopyTool = new FLandscapeToolCopy<ToolTarget>(InEdMode);
 	}
 
-	~FLandscapeToolCopyPaste()
+	// Just hybrid of Copy and Paste tool
+	virtual const TCHAR* GetToolName() override { return TEXT("CopyPaste"); }
+	virtual FText GetDisplayName() override { return NSLOCTEXT("UnrealEd", "LandscapeMode_Region", "Region Copy/Paste"); };
+
+	// Copy tool doesn't use any view information, so just do it as one function
+	void Copy()
 	{
-		if (pCopyTool)
-		{
-			delete pCopyTool;
-		}
+		CopyTool->BeginTool(nullptr, EdMode->CurrentToolTarget, FVector::ZeroVector);
+		CopyTool->EndTool(nullptr);
 	}
 
-	// Just hybrid of Copy and Paste tool
-	// Copy tool doesn't use any view information, so just do it as one function
-	virtual void Process(int32 Index, int32 Arg) override
+	void Paste()
 	{
-		switch(Index)
-		{
-		case 0: // Copy
-			{
-				if (pCopyTool && this->EdMode)
-				{
-					pCopyTool->BeginTool(NULL, this->EdMode->CurrentToolTarget, FVector::ZeroVector );
-					pCopyTool->EndTool(NULL);
-				}
-			}
-			break;
-		case 1:// Paste
-			{
-				if (this->EdMode)
-				{
-					this->SetGizmoMode(true);
-					this->BeginTool(NULL, this->EdMode->CurrentToolTarget, FVector::ZeroVector );
-					this->EndTool(NULL);
-					this->SetGizmoMode(false);
-				}
-			}
-		}
+		SetGizmoMode(true);
+		BeginTool(nullptr, EdMode->CurrentToolTarget, FVector::ZeroVector);
+		EndTool(nullptr);
+		SetGizmoMode(false);
 	}
 
 protected:
-	FLandscapeToolCopy<ToolTarget>* pCopyTool;
+	TUniqueObj<FLandscapeToolCopy<ToolTarget>> CopyTool;
 };
 
+void FEdModeLandscape::CopyDataToGizmo()
+{
+	// For Copy operation...
+	if (CopyPasteTool /*&& CopyPasteTool == CurrentTool*/)
+	{
+		CopyPasteTool->Copy();
+	}
+	if (CurrentGizmoActor.IsValid())
+	{
+		GEditor->SelectNone(false, true);
+		GEditor->SelectActor(CurrentGizmoActor.Get(), true, true, true);
+	}
+}
+
+void FEdModeLandscape::PasteDataFromGizmo()
+{
+	// For Paste for Gizmo Region operation...
+	if (CopyPasteTool /*&& CopyPasteTool == CurrentTool*/)
+	{
+		CopyPasteTool->Paste();
+	}
+	if (CurrentGizmoActor.IsValid())
+	{
+		GEditor->SelectNone(false, true);
+		GEditor->SelectActor(CurrentGizmoActor.Get(), true, true, true);
+	}
+}
+
 // 
 // FLandscapeToolNewLandscape
 //
 class FLandscapeToolNewLandscape : public FLandscapeTool
 {
 public:
-	class FEdModeLandscape* EdMode;
+	FEdModeLandscape* EdMode;
 	ENewLandscapePreviewMode::Type NewLandscapePreviewMode;
 
-	FLandscapeToolNewLandscape(class FEdModeLandscape* InEdMode)
+	FLandscapeToolNewLandscape(FEdModeLandscape* InEdMode)
 		: FLandscapeTool()
 		, EdMode(InEdMode)
 		, NewLandscapePreviewMode(ENewLandscapePreviewMode::NewLandscape)
@@ -1870,11 +1868,6 @@ public:
 	virtual void SetEditRenderType() override { GLandscapeEditRenderMode = ELandscapeEditRenderMode::None | (GLandscapeEditRenderMode & ELandscapeEditRenderMode::BitMaskForMask); }
 	virtual bool SupportsMask() override { return false; }
 
-	virtual bool IsValidForTarget(const FLandscapeToolTarget& Target) override
-	{
-		return true; // ignores target...
-	}
-
 	virtual void EnterTool()
 	{
 		EdMode->NewLandscapePreviewMode = NewLandscapePreviewMode;
@@ -1886,7 +1879,7 @@ public:
 		EdMode->NewLandscapePreviewMode = ENewLandscapePreviewMode::None;
 	}
 
-	virtual bool BeginTool( FEditorViewportClient* ViewportClient, const FLandscapeToolTarget& Target, const FVector& InHitLocation )
+	virtual bool BeginTool(FEditorViewportClient* ViewportClient, const FLandscapeToolTarget& Target, const FVector& InHitLocation)
 	{
 		// does nothing
 		return false;
@@ -1897,7 +1890,7 @@ public:
 		// does nothing
 	}
 
-	virtual bool MouseMove( FEditorViewportClient* ViewportClient, FViewport* Viewport, int32 x, int32 y )
+	virtual bool MouseMove(FEditorViewportClient* ViewportClient, FViewport* Viewport, int32 x, int32 y)
 	{
 		// does nothing
 		return false;
@@ -1911,9 +1904,9 @@ public:
 class FLandscapeToolResizeLandscape : public FLandscapeTool
 {
 public:
-	class FEdModeLandscape* EdMode;
+	FEdModeLandscape* EdMode;
 
-	FLandscapeToolResizeLandscape(class FEdModeLandscape* InEdMode)
+	FLandscapeToolResizeLandscape(FEdModeLandscape* InEdMode)
 		: FLandscapeTool()
 		, EdMode(InEdMode)
 	{
@@ -1925,11 +1918,6 @@ public:
 	virtual void SetEditRenderType() override { GLandscapeEditRenderMode = ELandscapeEditRenderMode::None | (GLandscapeEditRenderMode & ELandscapeEditRenderMode::BitMaskForMask); }
 	virtual bool SupportsMask() override { return false; }
 
-	virtual bool IsValidForTarget(const FLandscapeToolTarget& Target) override
-	{
-		return true; // ignores target...
-	}
-
 	virtual void EnterTool()
 	{
 		const int32 ComponentSizeQuads = EdMode->CurrentToolTarget.LandscapeInfo->ComponentSizeQuads;
@@ -1945,9 +1933,9 @@ public:
 			EdMode->UISettings->ResizeLandscape_Original_ComponentCount = FIntPoint::ZeroValue;
 			EdMode->UISettings->ResizeLandscape_ComponentCount = FIntPoint::ZeroValue;
 		}
-		EdMode->UISettings->ResizeLandscape_Original_QuadsPerSection      = EdMode->CurrentToolTarget.LandscapeInfo->SubsectionSizeQuads;
+		EdMode->UISettings->ResizeLandscape_Original_QuadsPerSection = EdMode->CurrentToolTarget.LandscapeInfo->SubsectionSizeQuads;
 		EdMode->UISettings->ResizeLandscape_Original_SectionsPerComponent = EdMode->CurrentToolTarget.LandscapeInfo->ComponentNumSubsections;
-		EdMode->UISettings->ResizeLandscape_QuadsPerSection      = EdMode->UISettings->ResizeLandscape_Original_QuadsPerSection;
+		EdMode->UISettings->ResizeLandscape_QuadsPerSection = EdMode->UISettings->ResizeLandscape_Original_QuadsPerSection;
 		EdMode->UISettings->ResizeLandscape_SectionsPerComponent = EdMode->UISettings->ResizeLandscape_Original_SectionsPerComponent;
 	}
 
@@ -1955,7 +1943,7 @@ public:
 	{
 	}
 
-	virtual bool BeginTool( FEditorViewportClient* ViewportClient, const FLandscapeToolTarget& Target, const FVector& InHitLocation )
+	virtual bool BeginTool(FEditorViewportClient* ViewportClient, const FLandscapeToolTarget& Target, const FVector& InHitLocation)
 	{
 		// does nothing
 		return false;
@@ -1966,7 +1954,7 @@ public:
 		// does nothing
 	}
 
-	virtual bool MouseMove( FEditorViewportClient* ViewportClient, FViewport* Viewport, int32 x, int32 y )
+	virtual bool MouseMove(FEditorViewportClient* ViewportClient, FViewport* Viewport, int32 x, int32 y)
 	{
 		// does nothing
 		return false;
@@ -1975,85 +1963,82 @@ public:
 
 //////////////////////////////////////////////////////////////////////////
 
-void FEdModeLandscape::InitializeToolSet_NewLandscape()
+void FEdModeLandscape::InitializeTool_NewLandscape()
 {
-	FLandscapeToolSet* ToolSet_NewLandscape = new(LandscapeToolSets) FLandscapeToolSet(TEXT("ToolSet_NewLandscape"));
-	ToolSet_NewLandscape->AddTool(new FLandscapeToolNewLandscape(this));
-
-	ToolSet_NewLandscape->ValidBrushes.Add("BrushSet_Dummy");
+	auto Tool_NewLandscape = MakeUnique<FLandscapeToolNewLandscape>(this);
+	Tool_NewLandscape->ValidBrushes.Add("BrushSet_Dummy");
+	LandscapeTools.Add(MoveTemp(Tool_NewLandscape));
 }
 
-void FEdModeLandscape::InitializeToolSet_ResizeLandscape()
+void FEdModeLandscape::InitializeTool_ResizeLandscape()
 {
-	FLandscapeToolSet* ToolSet_ResizeLandscape = new(LandscapeToolSets) FLandscapeToolSet(TEXT("ToolSet_ResizeLandscape"));
-	ToolSet_ResizeLandscape->AddTool(new FLandscapeToolResizeLandscape(this));
-
-	ToolSet_ResizeLandscape->ValidBrushes.Add("BrushSet_Dummy");
+	auto Tool_ResizeLandscape = MakeUnique<FLandscapeToolResizeLandscape>(this);
+	Tool_ResizeLandscape->ValidBrushes.Add("BrushSet_Dummy");
+	LandscapeTools.Add(MoveTemp(Tool_ResizeLandscape));
 }
 
-void FEdModeLandscape::IntializeToolSet_Select()
+void FEdModeLandscape::InitializeTool_Select()
 {
-	FLandscapeToolSet* ToolSet_Select = new(LandscapeToolSets) FLandscapeToolSet(TEXT("ToolSet_Select"));
-	ToolSet_Select->AddTool(new FLandscapeToolSelect<FLandscapeToolStrokeSelect>(this));
-
-	ToolSet_Select->ValidBrushes.Add("BrushSet_Component");
+	auto Tool_Select = MakeUnique<FLandscapeToolSelect<FLandscapeToolStrokeSelect>>(this);
+	Tool_Select->ValidBrushes.Add("BrushSet_Component");
+	LandscapeTools.Add(MoveTemp(Tool_Select));
 }
 
-void FEdModeLandscape::IntializeToolSet_AddComponent()
+void FEdModeLandscape::InitializeTool_AddComponent()
 {
-	FLandscapeToolSet* ToolSet_AddComponent = new(LandscapeToolSets) FLandscapeToolSet(TEXT("ToolSet_AddComponent"));
-	ToolSet_AddComponent->AddTool(new FLandscapeToolAddComponent(this));
-
-	ToolSet_AddComponent->ValidBrushes.Add("BrushSet_Component");
+	auto Tool_AddComponent = MakeUnique<FLandscapeToolAddComponent>(this);
+	Tool_AddComponent->ValidBrushes.Add("BrushSet_Component");
+	LandscapeTools.Add(MoveTemp(Tool_AddComponent));
 }
 
-void FEdModeLandscape::IntializeToolSet_DeleteComponent()
+void FEdModeLandscape::InitializeTool_DeleteComponent()
 {
-	FLandscapeToolSet* ToolSet_DeleteComponent = new(LandscapeToolSets) FLandscapeToolSet(TEXT("ToolSet_DeleteComponent"));
-	ToolSet_DeleteComponent->AddTool(new FLandscapeToolDeleteComponent(this));
-
-	ToolSet_DeleteComponent->ValidBrushes.Add("BrushSet_Component");
+	auto Tool_DeleteComponent = MakeUnique<FLandscapeToolDeleteComponent>(this);
+	Tool_DeleteComponent->ValidBrushes.Add("BrushSet_Component");
+	LandscapeTools.Add(MoveTemp(Tool_DeleteComponent));
 }
 
-void FEdModeLandscape::IntializeToolSet_MoveToLevel()
+void FEdModeLandscape::InitializeTool_MoveToLevel()
 {
-	FLandscapeToolSet* ToolSet_MoveToLevel = new(LandscapeToolSets) FLandscapeToolSet(TEXT("ToolSet_MoveToLevel"));
-	ToolSet_MoveToLevel->AddTool(new FLandscapeToolMoveToLevel(this));
-
-	ToolSet_MoveToLevel->ValidBrushes.Add("BrushSet_Component");
+	auto Tool_MoveToLevel = MakeUnique<FLandscapeToolMoveToLevel>(this);
+	Tool_MoveToLevel->ValidBrushes.Add("BrushSet_Component");
+	LandscapeTools.Add(MoveTemp(Tool_MoveToLevel));
 }
 
-void FEdModeLandscape::IntializeToolSet_Mask()
+void FEdModeLandscape::InitializeTool_Mask()
 {
-	FLandscapeToolSet* ToolSet_Mask = new(LandscapeToolSets) FLandscapeToolSet(TEXT("ToolSet_Mask"));
-	ToolSet_Mask->AddTool(new FLandscapeToolMask<FLandscapeToolStrokeSelect>(this));
-
-	ToolSet_Mask->ValidBrushes.Add("BrushSet_Circle");
-	ToolSet_Mask->ValidBrushes.Add("BrushSet_Alpha");
-	ToolSet_Mask->ValidBrushes.Add("BrushSet_Pattern");
+	auto Tool_Mask = MakeUnique<FLandscapeToolMask<FLandscapeToolStrokeSelect>>(this);
+	Tool_Mask->ValidBrushes.Add("BrushSet_Circle");
+	Tool_Mask->ValidBrushes.Add("BrushSet_Alpha");
+	Tool_Mask->ValidBrushes.Add("BrushSet_Pattern");
+	LandscapeTools.Add(MoveTemp(Tool_Mask));
 }
 
-void FEdModeLandscape::IntializeToolSet_CopyPaste()
+void FEdModeLandscape::InitializeTool_CopyPaste()
 {
-	FLandscapeToolSet* ToolSet_CopyPaste = new(LandscapeToolSets) FLandscapeToolSet(TEXT("ToolSet_CopyPaste"));
-	ToolSet_CopyPaste->AddTool(new FLandscapeToolCopyPaste<FHeightmapToolTarget>(this));
-	ToolSet_CopyPaste->AddTool(new FLandscapeToolCopyPaste<FWeightmapToolTarget>(this));
-	CopyPasteToolSet = ToolSet_CopyPaste;
-
-	ToolSet_CopyPaste->ValidBrushes.Add("BrushSet_Circle");
-	ToolSet_CopyPaste->ValidBrushes.Add("BrushSet_Alpha");
-	ToolSet_CopyPaste->ValidBrushes.Add("BrushSet_Pattern");
-	ToolSet_CopyPaste->ValidBrushes.Add("BrushSet_Gizmo");
+	auto Tool_CopyPaste_Heightmap = MakeUnique<FLandscapeToolCopyPaste<FHeightmapToolTarget>>(this);
+	Tool_CopyPaste_Heightmap->ValidBrushes.Add("BrushSet_Circle");
+	Tool_CopyPaste_Heightmap->ValidBrushes.Add("BrushSet_Alpha");
+	Tool_CopyPaste_Heightmap->ValidBrushes.Add("BrushSet_Pattern");
+	Tool_CopyPaste_Heightmap->ValidBrushes.Add("BrushSet_Gizmo");
+	CopyPasteTool = Tool_CopyPaste_Heightmap.Get();
+	LandscapeTools.Add(MoveTemp(Tool_CopyPaste_Heightmap));
+
+	//auto Tool_CopyPaste_Weightmap = MakeUnique<FLandscapeToolCopyPaste<FWeightmapToolTarget>>(this);
+	//Tool_CopyPaste_Weightmap->ValidBrushes.Add("BrushSet_Circle");
+	//Tool_CopyPaste_Weightmap->ValidBrushes.Add("BrushSet_Alpha");
+	//Tool_CopyPaste_Weightmap->ValidBrushes.Add("BrushSet_Pattern");
+	//Tool_CopyPaste_Weightmap->ValidBrushes.Add("BrushSet_Gizmo");
+	//LandscapeTools.Add(MoveTemp(Tool_CopyPaste_Weightmap));
 }
 
-void FEdModeLandscape::IntializeToolSet_Visibility()
+void FEdModeLandscape::InitializeTool_Visibility()
 {
-	FLandscapeToolSet* ToolSet_Visibility = new(LandscapeToolSets) FLandscapeToolSet(TEXT("ToolSet_Visibility"));
-	ToolSet_Visibility->AddTool(new FLandscapeToolVisibility(this));
-
-	ToolSet_Visibility->ValidBrushes.Add("BrushSet_Circle");
-	ToolSet_Visibility->ValidBrushes.Add("BrushSet_Alpha");
-	ToolSet_Visibility->ValidBrushes.Add("BrushSet_Pattern");
+	auto Tool_Visibility = MakeUnique<FLandscapeToolVisibility>(this);
+	Tool_Visibility->ValidBrushes.Add("BrushSet_Circle");
+	Tool_Visibility->ValidBrushes.Add("BrushSet_Alpha");
+	Tool_Visibility->ValidBrushes.Add("BrushSet_Pattern");
+	LandscapeTools.Add(MoveTemp(Tool_Visibility));
 }
 
 #undef LOCTEXT_NAMESPACE
\ No newline at end of file
diff --git a/Engine/Source/Editor/LandscapeEditor/Private/LandscapeEdModeErosionTools.cpp b/Engine/Source/Editor/LandscapeEditor/Private/LandscapeEdModeErosionTools.cpp
--- a/Engine/Source/Editor/LandscapeEditor/Private/LandscapeEdModeErosionTools.cpp
+++ b/Engine/Source/Editor/LandscapeEditor/Private/LandscapeEdModeErosionTools.cpp
@@ -20,15 +20,15 @@ class FLandscapeToolStrokeErosionBase
 {
 public:
 	FLandscapeToolStrokeErosionBase(FEdModeLandscape* InEdMode, const FLandscapeToolTarget& InTarget)
-	:	LandscapeInfo(InTarget.LandscapeInfo.Get())
-	,	HeightCache(InTarget)
-	,	WeightCache(InTarget)
-	,	bWeightApplied(InTarget.TargetType != ELandscapeToolTargetType::Heightmap)
+		: LandscapeInfo(InTarget.LandscapeInfo.Get())
+		, HeightCache(InTarget)
+		, WeightCache(InTarget)
+		, bWeightApplied(InTarget.TargetType != ELandscapeToolTargetType::Heightmap)
 	{}
 
 	virtual void Apply(FEditorViewportClient* ViewportClient, FLandscapeBrush* Brush, const ULandscapeEditorObject* UISettings, const TArray<FLandscapeToolMousePosition>& MousePositions) = 0;
 protected:
-	class ULandscapeInfo* LandscapeInfo;
+	ULandscapeInfo* LandscapeInfo;
 	FLandscapeHeightCache HeightCache;
 	FLandscapeFullWeightCache WeightCache;
 	bool bWeightApplied;
@@ -38,14 +38,13 @@ template<class TStrokeClass>
 class FLandscapeToolErosionBase : public FLandscapeToolBase<TStrokeClass>
 {
 public:
-	FLandscapeToolErosionBase(class FEdModeLandscape* InEdMode)
-	:	FLandscapeToolBase<TStrokeClass>(InEdMode)
+	FLandscapeToolErosionBase(FEdModeLandscape* InEdMode)
+		: FLandscapeToolBase<TStrokeClass>(InEdMode)
 	{}
 
-	virtual bool IsValidForTarget(const FLandscapeToolTarget& Target) override 
+	virtual ELandscapeToolTargetTypeMask::Type GetSupportedTargetTypes() override
 	{
-		// Erosion is applied to all layers
-		return true;
+		return ELandscapeToolTargetTypeMask::Heightmap;
 	}
 };
 
@@ -57,12 +56,12 @@ class FLandscapeToolStrokeErosion : public FLandscapeToolStrokeErosionBase
 {
 public:
 	FLandscapeToolStrokeErosion(FEdModeLandscape* InEdMode, const FLandscapeToolTarget& InTarget)
-	:	FLandscapeToolStrokeErosionBase(InEdMode, InTarget)
+		: FLandscapeToolStrokeErosionBase(InEdMode, InTarget)
 	{}
 
 	virtual void Apply(FEditorViewportClient* ViewportClient, FLandscapeBrush* Brush, const ULandscapeEditorObject* UISettings, const TArray<FLandscapeToolMousePosition>& MousePositions) override
 	{
-		if( !LandscapeInfo )
+		if (!LandscapeInfo)
 		{
 			return;
 		}
@@ -70,7 +69,7 @@ public:
 		// Get list of verts to update
 		TMap<FIntPoint, float> BrushInfo;
 		int32 X1, Y1, X2, Y2;
-		if( !Brush->ApplyBrush(MousePositions, BrushInfo, X1, Y1, X2, Y2) || MousePositions.Num()==0 )
+		if (!Brush->ApplyBrush(MousePositions, BrushInfo, X1, Y1, X2, Y2) || MousePositions.Num() == 0)
 		{
 			return;
 		}
@@ -89,13 +88,13 @@ public:
 		const int32 Thickness = UISettings->ErodeSurfaceThickness;
 		const int32 LayerNum = LandscapeInfo->Layers.Num();
 
-		HeightCache.CacheData(X1,Y1,X2,Y2);
+		HeightCache.CacheData(X1, Y1, X2, Y2);
 		TArray<uint16> HeightData;
-		HeightCache.GetCachedData(X1,Y1,X2,Y2,HeightData);
+		HeightCache.GetCachedData(X1, Y1, X2, Y2, HeightData);
 
 		TArray<uint8> WeightDatas; // Weight*Layers...
-		WeightCache.CacheData(X1,Y1,X2,Y2);
-		WeightCache.GetCachedData(X1,Y1,X2,Y2, WeightDatas, LayerNum);	
+		WeightCache.CacheData(X1, Y1, X2, Y2);
+		WeightCache.GetCachedData(X1, Y1, X2, Y2, WeightDatas, LayerNum);
 
 		// Apply the brush	
 		uint16 Thresh = UISettings->ErodeThresh;
@@ -112,15 +111,15 @@ public:
 		for (int32 i = 0; i < Iteration; i++)
 		{
 			bHasChanged = false;
-			for( auto It = BrushInfo.CreateIterator(); It; ++It )
+			for (auto It = BrushInfo.CreateIterator(); It; ++It)
 			{
 				int32 X, Y;
 				ALandscape::UnpackKey(It.Key(), X, Y);
 
-				if( It.Value() > 0.f )
+				if (It.Value() > 0.f)
 				{
-					int32 Center = (X-X1) + (Y-Y1)*(1+X2-X1);
-					int32 Neighbor[NeighborNum] = {(X-1-X1) + (Y-Y1)*(1+X2-X1), (X+1-X1) + (Y-Y1)*(1+X2-X1), (X-X1) + (Y-1-Y1)*(1+X2-X1), (X-X1) + (Y+1-Y1)*(1+X2-X1)};
+					int32 Center = (X - X1) + (Y - Y1)*(1 + X2 - X1);
+					int32 Neighbor[NeighborNum] = { (X - 1 - X1) + (Y - Y1)*(1 + X2 - X1), (X + 1 - X1) + (Y - Y1)*(1 + X2 - X1), (X - X1) + (Y - 1 - Y1)*(1 + X2 - X1), (X - X1) + (Y + 1 - Y1)*(1 + X2 - X1) };
 					uint32 SlopeTotal = 0;
 					uint16 SlopeMax = Thresh;
 
@@ -178,14 +177,14 @@ public:
 											for (int32 LayerIdx = 0; LayerIdx < LayerNum; LayerIdx++)
 											{
 												float CenterWeight = (float)(WeightDatas[Center*LayerNum + LayerIdx]) / 255.f;
-												float Weight = (float)(WeightDatas[Neighbor[Idx]*LayerNum + LayerIdx]) / 255.f;
+												float Weight = (float)(WeightDatas[Neighbor[Idx] * LayerNum + LayerIdx]) / 255.f;
 												NeighborWeight[Idx*LayerNum + LayerIdx] = Weight*(float)Thickness + CenterWeight*WeightDiff*WeightTransfer; // transferred + original...
 												TotalWeight += NeighborWeight[Idx*LayerNum + LayerIdx];
 											}
 											// Need to normalize weight...
 											for (int32 LayerIdx = 0; LayerIdx < LayerNum; LayerIdx++)
 											{
-												WeightDatas[Neighbor[Idx]*LayerNum + LayerIdx] = (uint8)(255.f * NeighborWeight[Idx*LayerNum + LayerIdx] / TotalWeight);
+												WeightDatas[Neighbor[Idx] * LayerNum + LayerIdx] = (uint8)(255.f * NeighborWeight[Idx*LayerNum + LayerIdx] / TotalWeight);
 											}
 										}
 									}
@@ -230,24 +229,24 @@ public:
 		}
 
 		// Make some noise...
-		for( auto It = BrushInfo.CreateIterator(); It; ++It )
+		for (auto It = BrushInfo.CreateIterator(); It; ++It)
 		{
 			int32 X, Y;
 			ALandscape::UnpackKey(It.Key(), X, Y);
 
-			if( It.Value() > 0.f )
+			if (It.Value() > 0.f)
 			{
 				FNoiseParameter NoiseParam(0, UISettings->ErosionNoiseScale, It.Value() * Thresh * UISettings->ToolStrength * BrushSizeAdjust);
 				float PaintAmount = ELandscapeToolNoiseMode::Conversion((ELandscapeToolNoiseMode::Type)UISettings->ErosionNoiseMode.GetValue(), NoiseParam.NoiseAmount, NoiseParam.Sample(X, Y));
-				HeightData[(X-X1) + (Y-Y1)*(1+X2-X1)] = FLandscapeHeightCache::ClampValue(HeightData[(X-X1) + (Y-Y1)*(1+X2-X1)] + PaintAmount);
+				HeightData[(X - X1) + (Y - Y1)*(1 + X2 - X1)] = FLandscapeHeightCache::ClampValue(HeightData[(X - X1) + (Y - Y1)*(1 + X2 - X1)] + PaintAmount);
 			}
 		}
 
-		HeightCache.SetCachedData(X1,Y1,X2,Y2,HeightData);
+		HeightCache.SetCachedData(X1, Y1, X2, Y2, HeightData);
 		HeightCache.Flush();
 		if (bWeightApplied)
 		{
-			WeightCache.SetCachedData(X1,Y1,X2,Y2,WeightDatas, LayerNum, ELandscapeLayerPaintingRestriction::None);
+			WeightCache.SetCachedData(X1, Y1, X2, Y2, WeightDatas, LayerNum, ELandscapeLayerPaintingRestriction::None);
 		}
 		WeightCache.Flush();
 	}
@@ -256,8 +255,8 @@ public:
 class FLandscapeToolErosion : public FLandscapeToolErosionBase<FLandscapeToolStrokeErosion>
 {
 public:
-	FLandscapeToolErosion(class FEdModeLandscape* InEdMode)
-	:	FLandscapeToolErosionBase(InEdMode)
+	FLandscapeToolErosion(FEdModeLandscape* InEdMode)
+		: FLandscapeToolErosionBase(InEdMode)
 	{}
 
 	virtual const TCHAR* GetToolName() override { return TEXT("Erosion"); }
@@ -273,12 +272,12 @@ class FLandscapeToolStrokeHydraErosion : public FLandscapeToolStrokeErosionBase
 {
 public:
 	FLandscapeToolStrokeHydraErosion(FEdModeLandscape* InEdMode, const FLandscapeToolTarget& InTarget)
-	:	FLandscapeToolStrokeErosionBase(InEdMode, InTarget)
+		: FLandscapeToolStrokeErosionBase(InEdMode, InTarget)
 	{}
 
 	virtual void Apply(FEditorViewportClient* ViewportClient, FLandscapeBrush* Brush, const ULandscapeEditorObject* UISettings, const TArray<FLandscapeToolMousePosition>& MousePositions) override
 	{
-		if( !LandscapeInfo )
+		if (!LandscapeInfo)
 		{
 			return;
 		}
@@ -286,7 +285,7 @@ public:
 		// Get list of verts to update
 		TMap<FIntPoint, float> BrushInfo;
 		int32 X1, Y1, X2, Y2;
-		if( !Brush->ApplyBrush(MousePositions,BrushInfo, X1, Y1, X2, Y2) || MousePositions.Num()==0 )
+		if (!Brush->ApplyBrush(MousePositions, BrushInfo, X1, Y1, X2, Y2) || MousePositions.Num() == 0)
 		{
 			return;
 		}
@@ -309,48 +308,48 @@ public:
 		const float EvaporateRatio = 0.5;
 		const float SedimentCapacity = 0.10 * UISettings->SedimentCapacity; //DissolvingRatio; //0.01;
 
-		HeightCache.CacheData(X1,Y1,X2,Y2);
+		HeightCache.CacheData(X1, Y1, X2, Y2);
 		TArray<uint16> HeightData;
-		HeightCache.GetCachedData(X1,Y1,X2,Y2,HeightData);
+		HeightCache.GetCachedData(X1, Y1, X2, Y2, HeightData);
 
 		// Apply the brush
 		TArray<uint16> WaterData;
-		WaterData.Empty((1+X2-X1)*(1+Y2-Y1));
-		WaterData.AddZeroed((1+X2-X1)*(1+Y2-Y1));
+		WaterData.Empty((1 + X2 - X1)*(1 + Y2 - Y1));
+		WaterData.AddZeroed((1 + X2 - X1)*(1 + Y2 - Y1));
 		TArray<uint16> SedimentData;
-		SedimentData.Empty((1+X2-X1)*(1+Y2-Y1));
-		SedimentData.AddZeroed((1+X2-X1)*(1+Y2-Y1));
+		SedimentData.Empty((1 + X2 - X1)*(1 + Y2 - Y1));
+		SedimentData.AddZeroed((1 + X2 - X1)*(1 + Y2 - Y1));
 
 		// Only initial raining works better...
 		FNoiseParameter NoiseParam(0, UISettings->RainDistScale, RainAmount);
-		for( auto It = BrushInfo.CreateIterator(); It; ++It )
+		for (auto It = BrushInfo.CreateIterator(); It; ++It)
 		{
 			int32 X, Y;
 			ALandscape::UnpackKey(It.Key(), X, Y);
 
-			if( It.Value() >= 1.f )
+			if (It.Value() >= 1.f)
 			{
 				float PaintAmount = ELandscapeToolNoiseMode::Conversion((ELandscapeToolNoiseMode::Type)UISettings->RainDistMode.GetValue(), NoiseParam.NoiseAmount, NoiseParam.Sample(X, Y));
 				if (PaintAmount > 0) // Raining only for positive region...
-					WaterData[(X-X1) + (Y-Y1)*(1+X2-X1)] += PaintAmount;
+					WaterData[(X - X1) + (Y - Y1)*(1 + X2 - X1)] += PaintAmount;
 			}
 		}
 
 		for (int32 i = 0; i < Iteration; i++)
 		{
 			bool bWaterExist = false;
-			for( auto It = BrushInfo.CreateIterator(); It; ++It )
+			for (auto It = BrushInfo.CreateIterator(); It; ++It)
 			{
 				int32 X, Y;
 				ALandscape::UnpackKey(It.Key(), X, Y);
 
-				if( It.Value() > 0.f)
+				if (It.Value() > 0.f)
 				{
-					int32 Center = (X-X1) + (Y-Y1)*(1+X2-X1);
+					int32 Center = (X - X1) + (Y - Y1)*(1 + X2 - X1);
 
 					int32 Neighbor[NeighborNum] = {
-						(X-1-X1) + (Y-Y1)*(1+X2-X1), (X+1-X1) + (Y-Y1)*(1+X2-X1), (X-X1) + (Y-1-Y1)*(1+X2-X1), (X-X1) + (Y+1-Y1)*(1+X2-X1)
-						,(X-1-X1) + (Y-1-Y1)*(1+X2-X1), (X+1-X1) + (Y+1-Y1)*(1+X2-X1), (X+1-X1) + (Y-1-Y1)*(1+X2-X1), (X-1-X1) + (Y+1-Y1)*(1+X2-X1)
+						(X - 1 - X1) + (Y - Y1)*(1 + X2 - X1), (X + 1 - X1) + (Y - Y1)*(1 + X2 - X1), (X - X1) + (Y - 1 - Y1)*(1 + X2 - X1), (X - X1) + (Y + 1 - Y1)*(1 + X2 - X1)
+						, (X - 1 - X1) + (Y - 1 - Y1)*(1 + X2 - X1), (X + 1 - X1) + (Y + 1 - Y1)*(1 + X2 - X1), (X + 1 - X1) + (Y - 1 - Y1)*(1 + X2 - X1), (X - 1 - X1) + (Y + 1 - Y1)*(1 + X2 - X1)
 					};
 
 					// Dissolving...				
@@ -363,16 +362,16 @@ public:
 
 					uint32 TotalHeightDiff = 0;
 					uint32 TotalAltitudeDiff = 0;
-					uint32 AltitudeDiff[NeighborNum] = {0};
+					uint32 AltitudeDiff[NeighborNum] = { 0 };
 					uint32 TotalWaterDiff = 0;
 					uint32 TotalSedimentDiff = 0;
 
-					uint32 Altitude = HeightData[Center]+WaterData[Center];
+					uint32 Altitude = HeightData[Center] + WaterData[Center];
 					float AverageAltitude = 0;
 					uint32 LowerNeighbor = 0;
 					for (int32 Idx = 0; Idx < NeighborNum; Idx++)
 					{
-						uint32 NeighborAltitude = HeightData[Neighbor[Idx]]+WaterData[Neighbor[Idx]];
+						uint32 NeighborAltitude = HeightData[Neighbor[Idx]] + WaterData[Neighbor[Idx]];
 						if (Altitude > NeighborAltitude)
 						{
 							AltitudeDiff[Idx] = Altitude - NeighborAltitude;
@@ -430,13 +429,13 @@ public:
 						if (SedimentDiff > 0)
 						{
 							SedimentData[Center] -= SedimentDiff;
-							HeightData[Center] = FMath::Clamp<uint16>(HeightData[Center]+SedimentDiff, 0, 65535);
+							HeightData[Center] = FMath::Clamp<uint16>(HeightData[Center] + SedimentDiff, 0, 65535);
 						}
 					}
-				}	
+				}
 			}
 
-			if (!bWaterExist) 
+			if (!bWaterExist)
 			{
 				break;
 			}
@@ -448,7 +447,7 @@ public:
 			LowPassFilter<uint16>(X1, Y1, X2, Y2, BrushInfo, HeightData, UISettings->HErosionDetailScale, 1.0f);
 		}
 
-		HeightCache.SetCachedData(X1,Y1,X2,Y2,HeightData);
+		HeightCache.SetCachedData(X1, Y1, X2, Y2, HeightData);
 		HeightCache.Flush();
 	}
 };
@@ -456,11 +455,11 @@ public:
 class FLandscapeToolHydraErosion : public FLandscapeToolErosionBase<FLandscapeToolStrokeHydraErosion>
 {
 public:
-	FLandscapeToolHydraErosion(class FEdModeLandscape* InEdMode)
-	:	FLandscapeToolErosionBase(InEdMode)
+	FLandscapeToolHydraErosion(FEdModeLandscape* InEdMode)
+		: FLandscapeToolErosionBase(InEdMode)
 	{}
 
-	virtual const TCHAR* GetToolName() override { return TEXT("HydraulicErosion"); }
+	virtual const TCHAR* GetToolName() override { return TEXT("HydraErosion"); } // formerly HydraulicErosion
 	virtual FText GetDisplayName() override { return NSLOCTEXT("UnrealEd", "LandscapeMode_HydraErosion", "Hydraulic Erosion"); };
 
 };
@@ -468,22 +467,20 @@ public:
 //
 // Toolset initialization
 //
-void FEdModeLandscape::IntializeToolSet_Erosion()
+void FEdModeLandscape::InitializeTool_Erosion()
 {
-	FLandscapeToolSet* ToolSet_Erosion = new(LandscapeToolSets) FLandscapeToolSet(TEXT("ToolSet_Erosion"));
-	ToolSet_Erosion->AddTool(new FLandscapeToolErosion(this));
-
-	ToolSet_Erosion->ValidBrushes.Add("BrushSet_Circle");
-	ToolSet_Erosion->ValidBrushes.Add("BrushSet_Alpha");
-	ToolSet_Erosion->ValidBrushes.Add("BrushSet_Pattern");
+	auto Tool_Erosion = MakeUnique<FLandscapeToolErosion>(this);
+	Tool_Erosion->ValidBrushes.Add("BrushSet_Circle");
+	Tool_Erosion->ValidBrushes.Add("BrushSet_Alpha");
+	Tool_Erosion->ValidBrushes.Add("BrushSet_Pattern");
+	LandscapeTools.Add(MoveTemp(Tool_Erosion));
 }
 
-void FEdModeLandscape::IntializeToolSet_HydraErosion()
+void FEdModeLandscape::InitializeTool_HydraErosion()
 {
-	FLandscapeToolSet* ToolSet_HydraErosion = new(LandscapeToolSets) FLandscapeToolSet(TEXT("ToolSet_HydraErosion"));
-	ToolSet_HydraErosion->AddTool(new FLandscapeToolHydraErosion(this));
-
-	ToolSet_HydraErosion->ValidBrushes.Add("BrushSet_Circle");
-	ToolSet_HydraErosion->ValidBrushes.Add("BrushSet_Alpha");
-	ToolSet_HydraErosion->ValidBrushes.Add("BrushSet_Pattern");
+	auto Tool_HydraErosion = MakeUnique<FLandscapeToolHydraErosion>(this);
+	Tool_HydraErosion->ValidBrushes.Add("BrushSet_Circle");
+	Tool_HydraErosion->ValidBrushes.Add("BrushSet_Alpha");
+	Tool_HydraErosion->ValidBrushes.Add("BrushSet_Pattern");
+	LandscapeTools.Add(MoveTemp(Tool_HydraErosion));
 }
