Remove FLandscapeToolSet

diff --git a/Engine/Source/Editor/LandscapeEditor/Private/LandscapeEdModePaintTools.cpp b/Engine/Source/Editor/LandscapeEditor/Private/LandscapeEdModePaintTools.cpp
--- a/Engine/Source/Editor/LandscapeEditor/Private/LandscapeEdModePaintTools.cpp
+++ b/Engine/Source/Editor/LandscapeEditor/Private/LandscapeEdModePaintTools.cpp
@@ -13,21 +13,21 @@
 #include "LandscapeEdModeTools.h"
 
 
-const int32 FNoiseParameter::Permutations[256] = 
+const int32 FNoiseParameter::Permutations[256] =
 {
-	151,160,137,91,90,15,
-	131,13,201,95,96,53,194,233,7,225,140,36,103,30,69,142,8,99,37,240,21,10,23,
-	190, 6,148,247,120,234,75,0,26,197,62,94,252,219,203,117,35,11,32,57,177,33,
-	88,237,149,56,87,174,20,125,136,171,168, 68,175,74,165,71,134,139,48,27,166,
-	77,146,158,231,83,111,229,122,60,211,133,230,220,105,92,41,55,46,245,40,244,
-	102,143,54, 65,25,63,161, 1,216,80,73,209,76,132,187,208, 89,18,169,200,196,
-	135,130,116,188,159,86,164,100,109,198,173,186, 3,64,52,217,226,250,124,123,
-	5,202,38,147,118,126,255,82,85,212,207,206,59,227,47,16,58,17,182,189,28,42,
-	223,183,170,213,119,248,152, 2,44,154,163, 70,221,153,101,155,167, 43,172,9,
-	129,22,39,253, 19,98,108,110,79,113,224,232,178,185, 112,104,218,246,97,228,
-	251,34,242,193,238,210,144,12,191,179,162,241, 81,51,145,235,249,14,239,107,
-	49,192,214, 31,181,199,106,157,184, 84,204,176,115,121,50,45,127, 4,150,254,
-	138,236,205,93,222,114,67,29,24,72,243,141,128,195,78,66,215,61,156,180
+	151, 160, 137, 91, 90, 15,
+	131, 13, 201, 95, 96, 53, 194, 233, 7, 225, 140, 36, 103, 30, 69, 142, 8, 99, 37, 240, 21, 10, 23,
+	190, 6, 148, 247, 120, 234, 75, 0, 26, 197, 62, 94, 252, 219, 203, 117, 35, 11, 32, 57, 177, 33,
+	88, 237, 149, 56, 87, 174, 20, 125, 136, 171, 168, 68, 175, 74, 165, 71, 134, 139, 48, 27, 166,
+	77, 146, 158, 231, 83, 111, 229, 122, 60, 211, 133, 230, 220, 105, 92, 41, 55, 46, 245, 40, 244,
+	102, 143, 54, 65, 25, 63, 161, 1, 216, 80, 73, 209, 76, 132, 187, 208, 89, 18, 169, 200, 196,
+	135, 130, 116, 188, 159, 86, 164, 100, 109, 198, 173, 186, 3, 64, 52, 217, 226, 250, 124, 123,
+	5, 202, 38, 147, 118, 126, 255, 82, 85, 212, 207, 206, 59, 227, 47, 16, 58, 17, 182, 189, 28, 42,
+	223, 183, 170, 213, 119, 248, 152, 2, 44, 154, 163, 70, 221, 153, 101, 155, 167, 43, 172, 9,
+	129, 22, 39, 253, 19, 98, 108, 110, 79, 113, 224, 232, 178, 185, 112, 104, 218, 246, 97, 228,
+	251, 34, 242, 193, 238, 210, 144, 12, 191, 179, 162, 241, 81, 51, 145, 235, 249, 14, 239, 107,
+	49, 192, 214, 31, 181, 199, 106, 157, 184, 84, 204, 176, 115, 121, 50, 45, 127, 4, 150, 254,
+	138, 236, 205, 93, 222, 114, 67, 29, 24, 72, 243, 141, 128, 195, 78, 66, 215, 61, 156, 180
 };
 
 
@@ -42,12 +42,12 @@ class FLandscapeToolPaintBase : public FLandscapeToolBase<TStrokeClass>
 {
 public:
 	FLandscapeToolPaintBase(FEdModeLandscape* InEdMode)
-		:	FLandscapeToolBase<TStrokeClass>(InEdMode)
+		: FLandscapeToolBase<TStrokeClass>(InEdMode)
 	{}
 
-	virtual bool IsValidForTarget(const FLandscapeToolTarget& Target)
+	virtual ELandscapeToolTargetTypeMask::Type GetSupportedTargetTypes() override
 	{
-		return Target.TargetType == TToolTarget::TargetType;
+		return ELandscapeToolTargetTypeMask::FromType(TToolTarget::TargetType);
 	}
 };
 
@@ -57,15 +57,15 @@ class FLandscapeToolStrokePaintBase
 {
 public:
 	FLandscapeToolStrokePaintBase(FEdModeLandscape* InEdMode, const FLandscapeToolTarget& InTarget)
-	:	LandscapeInfo(InTarget.LandscapeInfo.Get())
-	,	Cache(InTarget)
+		: LandscapeInfo(InTarget.LandscapeInfo.Get())
+		, Cache(InTarget)
 	{
 	}
 
 	virtual void Apply(FEditorViewportClient* ViewportClient, FLandscapeBrush* Brush, const ULandscapeEditorObject* UISettings, const TArray<FLandscapeToolMousePosition>& MousePositions) = 0;
 protected:
 	typename ToolTarget::CacheClass Cache;
-	class ULandscapeInfo* LandscapeInfo;
+	ULandscapeInfo* LandscapeInfo;
 };
 
 
@@ -82,7 +82,7 @@ class FLandscapeToolStrokePaint : public FLandscapeToolStrokePaintBase<ToolTarge
 	TMap<FIntPoint, float> TotalInfluenceMap;	// amount of time and weight the brush has spent on each vertex.
 public:
 	FLandscapeToolStrokePaint(FEdModeLandscape* InEdMode, const FLandscapeToolTarget& InTarget)
-	:	FLandscapeToolStrokePaintBase<ToolTarget>(InEdMode, InTarget)
+		: FLandscapeToolStrokePaintBase<ToolTarget>(InEdMode, InTarget)
 	{}
 
 	virtual void Apply(FEditorViewportClient* ViewportClient, FLandscapeBrush* Brush, const ULandscapeEditorObject* UISettings, const TArray<FLandscapeToolMousePosition>& MousePositions) override
@@ -90,11 +90,11 @@ public:
 		// Get list of verts to update
 		TMap<FIntPoint, float> BrushInfo;
 		int32 X1, Y1, X2, Y2;
-		if( !Brush->ApplyBrush(MousePositions, BrushInfo, X1, Y1, X2, Y2) || MousePositions.Num()==0 )
+		if (!Brush->ApplyBrush(MousePositions, BrushInfo, X1, Y1, X2, Y2) || MousePositions.Num() == 0)
 		{
 			return;
 		}
-		
+
 		// Tablet pressure
 		float Pressure = ViewportClient->Viewport->IsPenActive() ? ViewportClient->Viewport->GetTabletPressure() : 1.f;
 
@@ -104,50 +104,50 @@ public:
 		X2 += 1;
 		Y2 += 1;
 
-		this->Cache.CacheData(X1,Y1,X2,Y2);
+		this->Cache.CacheData(X1, Y1, X2, Y2);
 
 		// Invert when holding Shift
-		bool bInvert = MousePositions[MousePositions.Num()-1].bShiftDown;
+		bool bInvert = MousePositions[MousePositions.Num() - 1].bShiftDown;
 		//UE_LOG(LogLandscape, Log, TEXT("bInvert = %d"), bInvert);
 		bool bUseClayBrush = UISettings->bUseClayBrush && ToolTarget::TargetType == ELandscapeToolTargetType::Heightmap;
 		bool bUseWeightTargetValue = UISettings->bUseWeightTargetValue && ToolTarget::TargetType == ELandscapeToolTargetType::Weightmap;
 
 		// The data we'll be writing to
 		TArray<typename ToolTarget::CacheClass::DataType> Data;
-		this->Cache.GetCachedData(X1,Y1,X2,Y2,Data);
+		this->Cache.GetCachedData(X1, Y1, X2, Y2, Data);
 
 		// The source data we use for editing. 
 		// For heightmaps we use a cached snapshot, for weightmaps we use the live data.
 		TArray<typename ToolTarget::CacheClass::DataType>* SourceDataArrayPtr = &Data;
 		TArray<typename ToolTarget::CacheClass::DataType> OriginalData;
 
-		if( ToolTarget::TargetType == ELandscapeToolTargetType::Heightmap )
+		if (ToolTarget::TargetType == ELandscapeToolTargetType::Heightmap)
 		{
 			// Heightmaps use the original data rather than the data edited during the stroke.
-			this->Cache.GetOriginalData(X1,Y1,X2,Y2,OriginalData);
+			this->Cache.GetOriginalData(X1, Y1, X2, Y2, OriginalData);
 			SourceDataArrayPtr = &OriginalData;
 		}
 		else
-		if( !bUseWeightTargetValue )
-		{
+			if (!bUseWeightTargetValue)
+			{
 			// When painting weights (and not using target value mode), we use a source value that tends more
 			// to the current value as we paint over the same region multiple times.
-			this->Cache.GetOriginalData(X1,Y1,X2,Y2,OriginalData);
+			this->Cache.GetOriginalData(X1, Y1, X2, Y2, OriginalData);
 			SourceDataArrayPtr = &OriginalData;
 
-			for( int32 Y=Y1;Y<Y2;Y++ )
+			for (int32 Y = Y1; Y < Y2; Y++)
 			{
-				for( int32 X=X1;X<X2;X++ )
+				for (int32 X = X1; X < X2; X++)
 				{
-					float VertexInfluence = TotalInfluenceMap.FindRef(ALandscape::MakeKey(X,Y));
+					float VertexInfluence = TotalInfluenceMap.FindRef(ALandscape::MakeKey(X, Y));
 
-					typename ToolTarget::CacheClass::DataType& CurrentValue = DATA_AT(Data,X,Y);
-					typename ToolTarget::CacheClass::DataType& SourceValue = DATA_AT(OriginalData,X,Y);
+					typename ToolTarget::CacheClass::DataType& CurrentValue = DATA_AT(Data, X, Y);
+					typename ToolTarget::CacheClass::DataType& SourceValue = DATA_AT(OriginalData, X, Y);
 
-					SourceValue = FMath::Lerp<typename ToolTarget::CacheClass::DataType>( SourceValue, CurrentValue, FMath::Min<float>(VertexInfluence * 0.05f, 1.f) );
+					SourceValue = FMath::Lerp<typename ToolTarget::CacheClass::DataType>(SourceValue, CurrentValue, FMath::Min<float>(VertexInfluence * 0.05f, 1.f));
 				}
 			}
-		}
+			}
 
 		FMatrix ToWorld = ToolTarget::ToWorldMatrix(this->LandscapeInfo);
 		FMatrix FromWorld = ToolTarget::FromWorldMatrix(this->LandscapeInfo);
@@ -159,154 +159,163 @@ public:
 		FPlane BrushPlane;
 		TArray<FVector> Normals;
 
-		if( bUseClayBrush )
+		if (bUseClayBrush)
 		{
 			// Calculate normals for brush verts in data space
 			Normals.Empty(OriginalData.Num());
 			Normals.AddZeroed(OriginalData.Num());
 
-			for( int32 Y=Y1;Y<Y2;Y++ )
+			for (int32 Y = Y1; Y < Y2; Y++)
 			{
-				for( int32 X=X1;X<X2;X++ )
+				for (int32 X = X1; X < X2; X++)
 				{
-					FVector Vert00 = ToWorld.TransformPosition( FVector((float)X+0.f,(float)Y+0.f,DATA_AT(OriginalData,X+0,Y+0)) );
-					FVector Vert01 = ToWorld.TransformPosition( FVector((float)X+0.f,(float)Y+1.f,DATA_AT(OriginalData,X+0,Y+1)) );
-					FVector Vert10 = ToWorld.TransformPosition( FVector((float)X+1.f,(float)Y+0.f,DATA_AT(OriginalData,X+1,Y+0)) );
-					FVector Vert11 = ToWorld.TransformPosition( FVector((float)X+1.f,(float)Y+1.f,DATA_AT(OriginalData,X+1,Y+1)) );
+					FVector Vert00 = ToWorld.TransformPosition(FVector((float)X + 0.f, (float)Y + 0.f, DATA_AT(OriginalData, X + 0, Y + 0)));
+					FVector Vert01 = ToWorld.TransformPosition(FVector((float)X + 0.f, (float)Y + 1.f, DATA_AT(OriginalData, X + 0, Y + 1)));
+					FVector Vert10 = ToWorld.TransformPosition(FVector((float)X + 1.f, (float)Y + 0.f, DATA_AT(OriginalData, X + 1, Y + 0)));
+					FVector Vert11 = ToWorld.TransformPosition(FVector((float)X + 1.f, (float)Y + 1.f, DATA_AT(OriginalData, X + 1, Y + 1)));
 
-					FVector FaceNormal1 = ((Vert00-Vert10) ^ (Vert10-Vert11)).SafeNormal();
-					FVector FaceNormal2 = ((Vert11-Vert01) ^ (Vert01-Vert00)).SafeNormal(); 
+					FVector FaceNormal1 = ((Vert00 - Vert10) ^ (Vert10 - Vert11)).SafeNormal();
+					FVector FaceNormal2 = ((Vert11 - Vert01) ^ (Vert01 - Vert00)).SafeNormal();
 
 					// contribute to the vertex normals.
-					DATA_AT(Normals,X+1,Y+0) += FaceNormal1;
-					DATA_AT(Normals,X+0,Y+1) += FaceNormal2;
-					DATA_AT(Normals,X+0,Y+0) += FaceNormal1 + FaceNormal2;
-					DATA_AT(Normals,X+1,Y+1) += FaceNormal1 + FaceNormal2;
+					DATA_AT(Normals, X + 1, Y + 0) += FaceNormal1;
+					DATA_AT(Normals, X + 0, Y + 1) += FaceNormal2;
+					DATA_AT(Normals, X + 0, Y + 0) += FaceNormal1 + FaceNormal2;
+					DATA_AT(Normals, X + 1, Y + 1) += FaceNormal1 + FaceNormal2;
 				}
 			}
-			for( int32 Y=Y1;Y<=Y2;Y++ )
+			for (int32 Y = Y1; Y <= Y2; Y++)
 			{
-				for( int32 X=X1;X<=X2;X++ )
+				for (int32 X = X1; X <= X2; X++)
 				{
-					DATA_AT(Normals,X,Y) = DATA_AT(Normals,X,Y).SafeNormal();
+					DATA_AT(Normals, X, Y) = DATA_AT(Normals, X, Y).SafeNormal();
 				}
 			}
 
 			// Find brush centroid location
-			FVector AveragePoint(0.f,0.f,0.f);
-			FVector AverageNormal(0.f,0.f,0.f);
+			FVector AveragePoint(0.f, 0.f, 0.f);
+			FVector AverageNormal(0.f, 0.f, 0.f);
 			float TotalWeight = 0.f;
-			for( auto It = BrushInfo.CreateIterator(); It; ++It )
+			for (auto It = BrushInfo.CreateIterator(); It; ++It)
 			{
 				int32 X, Y;
 				ALandscape::UnpackKey(It.Key(), X, Y);
 				float Weight = It.Value();
 
-				AveragePoint += FVector( (float)X * Weight, (float)Y * Weight, (float)DATA_AT(OriginalData,FMath::FloorToInt(X),FMath::FloorToInt(Y)) * Weight );
+				AveragePoint += FVector((float)X * Weight, (float)Y * Weight, (float)DATA_AT(OriginalData, FMath::FloorToInt(X), FMath::FloorToInt(Y)) * Weight);
 
-				FVector SampleNormal = DATA_AT(Normals,X,Y);
+				FVector SampleNormal = DATA_AT(Normals, X, Y);
 				AverageNormal += SampleNormal * Weight;
 
 				TotalWeight += Weight;
 			}
 
-			if( TotalWeight > 0.f )
+			if (TotalWeight > 0.f)
 			{
 				AveragePoint /= TotalWeight;
 				AverageNormal = AverageNormal.SafeNormal();
 			}
 
 			// Convert to world space
-			FVector AverageLocation = ToWorld.TransformPosition( AveragePoint );
-			FVector StrengthVector = ToWorld.TransformVector(FVector(0,0,UISettings->ToolStrength * Pressure * AdjustedStrength));
+			FVector AverageLocation = ToWorld.TransformPosition(AveragePoint);
+			FVector StrengthVector = ToWorld.TransformVector(FVector(0, 0, UISettings->ToolStrength * Pressure * AdjustedStrength));
 
 			// Brush pushes out in the normal direction
 			FVector OffsetVector = AverageNormal * StrengthVector.Z;
-			if( bInvert )
+			if (bInvert)
 			{
 				OffsetVector *= -1;
 			}
 
 			// World space brush plane
-			BrushPlane = FPlane( AverageLocation + OffsetVector, AverageNormal );
+			BrushPlane = FPlane(AverageLocation + OffsetVector, AverageNormal);
 		}
 
 		// Apply the brush	
-		for( auto It = BrushInfo.CreateIterator(); It; ++It )
+		for (auto It = BrushInfo.CreateIterator(); It; ++It)
 		{
 			int32 X, Y;
 			ALandscape::UnpackKey(It.Key(), X, Y);
 
 			// Update influence map
 			float VertexInfluence = TotalInfluenceMap.FindRef(It.Key());
-			TotalInfluenceMap.Add( It.Key(), VertexInfluence + It.Value() );
+			TotalInfluenceMap.Add(It.Key(), VertexInfluence + It.Value());
 
 			float PaintAmount = It.Value() * UISettings->ToolStrength * Pressure * AdjustedStrength;
-			typename ToolTarget::CacheClass::DataType& CurrentValue = DATA_AT(Data,X,Y);
-			const typename ToolTarget::CacheClass::DataType& SourceValue = DATA_AT(*SourceDataArrayPtr,X,Y);
+			typename ToolTarget::CacheClass::DataType& CurrentValue = DATA_AT(Data, X, Y);
+			const typename ToolTarget::CacheClass::DataType& SourceValue = DATA_AT(*SourceDataArrayPtr, X, Y);
 
-			if( bUseWeightTargetValue )
+			if (bUseWeightTargetValue)
 			{
-				if( bInvert )
+				if (bInvert)
 				{
-					CurrentValue = FMath::Lerp( CurrentValue, DestValue, PaintAmount / AdjustedStrength );
+					CurrentValue = FMath::Lerp(CurrentValue, DestValue, PaintAmount / AdjustedStrength);
 				}
 				else
 				{
-					CurrentValue = FMath::Lerp( CurrentValue, DestValue, PaintAmount / AdjustedStrength );
+					CurrentValue = FMath::Lerp(CurrentValue, DestValue, PaintAmount / AdjustedStrength);
 				}
 			}
 			else
-			if( bUseClayBrush )
-			{
+				if (bUseClayBrush)
+				{
 				// Brush application starts from original world location at start of stroke
-				FVector WorldLoc = ToWorld.TransformPosition(FVector(X,Y,SourceValue));
+				FVector WorldLoc = ToWorld.TransformPosition(FVector(X, Y, SourceValue));
 
 				// Calculate new location on the brush plane
 				WorldLoc.Z = (BrushPlane.W - BrushPlane.X*WorldLoc.X - BrushPlane.Y*WorldLoc.Y) / BrushPlane.Z;
 
 				// Painted amount lerps based on brush falloff.
-				float PaintValue = FMath::Lerp<float>( (float)SourceValue, FromWorld.TransformPosition(WorldLoc).Z, It.Value() );
+				float PaintValue = FMath::Lerp<float>((float)SourceValue, FromWorld.TransformPosition(WorldLoc).Z, It.Value());
 
-				if( bInvert )
+				if (bInvert)
 				{
-					CurrentValue = ToolTarget::CacheClass::ClampValue( FMath::Min<int32>(FMath::RoundToInt(PaintValue), CurrentValue) );
+					CurrentValue = ToolTarget::CacheClass::ClampValue(FMath::Min<int32>(FMath::RoundToInt(PaintValue), CurrentValue));
 				}
 				else
 				{
-					CurrentValue = ToolTarget::CacheClass::ClampValue( FMath::Max<int32>(FMath::RoundToInt(PaintValue), CurrentValue) );
+					CurrentValue = ToolTarget::CacheClass::ClampValue(FMath::Max<int32>(FMath::RoundToInt(PaintValue), CurrentValue));
 				}
-			}
-			else
-			{
-				if( bInvert )
-				{
-					CurrentValue = ToolTarget::CacheClass::ClampValue( FMath::Min<int32>(SourceValue - FMath::RoundToInt(PaintAmount), CurrentValue) );
 				}
 				else
 				{
-					CurrentValue = ToolTarget::CacheClass::ClampValue( FMath::Max<int32>(SourceValue + FMath::RoundToInt(PaintAmount), CurrentValue) );
+					if (bInvert)
+					{
+						CurrentValue = ToolTarget::CacheClass::ClampValue(FMath::Min<int32>(SourceValue - FMath::RoundToInt(PaintAmount), CurrentValue));
+					}
+					else
+					{
+						CurrentValue = ToolTarget::CacheClass::ClampValue(FMath::Max<int32>(SourceValue + FMath::RoundToInt(PaintAmount), CurrentValue));
+					}
 				}
-			}
 		}
 
-		this->Cache.SetCachedData(X1,Y1,X2,Y2,Data, UISettings->PaintingRestriction);
+		this->Cache.SetCachedData(X1, Y1, X2, Y2, Data, UISettings->PaintingRestriction);
 		this->Cache.Flush();
 	}
 #undef DATA_AT
 };
 
-template<class ToolTarget>
-class FLandscapeToolPaint : public FLandscapeToolPaintBase<ToolTarget, FLandscapeToolStrokePaint<ToolTarget> >
+class FLandscapeToolPaint : public FLandscapeToolPaintBase<FWeightmapToolTarget, FLandscapeToolStrokePaint<FWeightmapToolTarget>>
 {
 public:
-	FLandscapeToolPaint(class FEdModeLandscape* InEdMode)
-	:	FLandscapeToolPaintBase<ToolTarget, FLandscapeToolStrokePaint<ToolTarget> >(InEdMode)
+	FLandscapeToolPaint(FEdModeLandscape* InEdMode)
+		: FLandscapeToolPaintBase<FWeightmapToolTarget, FLandscapeToolStrokePaint<FWeightmapToolTarget>>(InEdMode)
 	{}
 
 	virtual const TCHAR* GetToolName() override { return TEXT("Paint"); }
 	virtual FText GetDisplayName() override { return NSLOCTEXT("UnrealEd", "LandscapeMode_Paint", "Paint"); };
+};
+
+class FLandscapeToolSculpt : public FLandscapeToolPaintBase<FHeightmapToolTarget, FLandscapeToolStrokePaint<FHeightmapToolTarget>>
+{
+public:
+	FLandscapeToolSculpt(FEdModeLandscape* InEdMode)
+		: FLandscapeToolPaintBase<FHeightmapToolTarget, FLandscapeToolStrokePaint<FHeightmapToolTarget>>(InEdMode)
+	{}
 
+	virtual const TCHAR* GetToolName() override { return TEXT("Sculpt"); }
+	virtual FText GetDisplayName() override { return NSLOCTEXT("UnrealEd", "LandscapeMode_Sculpt", "Sculpt"); };
 };
 
 // 
@@ -318,7 +327,7 @@ class FLandscapeToolStrokeSmooth : public FLandscapeToolStrokePaintBase<ToolTarg
 {
 public:
 	FLandscapeToolStrokeSmooth(FEdModeLandscape* InEdMode, const FLandscapeToolTarget& InTarget)
-	:	FLandscapeToolStrokePaintBase<ToolTarget>(InEdMode, InTarget)
+		: FLandscapeToolStrokePaintBase<ToolTarget>(InEdMode, InTarget)
 	{}
 
 	virtual void Apply(FEditorViewportClient* ViewportClient, FLandscapeBrush* Brush, const ULandscapeEditorObject* UISettings, const TArray<FLandscapeToolMousePosition>& MousePositions) override
@@ -342,10 +351,10 @@ public:
 		X2 += 1;
 		Y2 += 1;
 
-		this->Cache.CacheData(X1,Y1,X2,Y2);
+		this->Cache.CacheData(X1, Y1, X2, Y2);
 
 		TArray<typename ToolTarget::CacheClass::DataType> Data;
-		this->Cache.GetCachedData(X1,Y1,X2,Y2,Data);
+		this->Cache.GetCachedData(X1, Y1, X2, Y2, Data);
 
 		// Apply the brush
 		if (UISettings->bDetailSmooth)
@@ -355,27 +364,27 @@ public:
 		else
 		{
 			// Filter size is SmoothFilterKernelScale * BrushSize, But avoids overflow value or under 1
-			int32 HalfFilterSize = FMath::Clamp<int32>(UISettings->SmoothFilterKernelScale * FMath::Max(Y2-Y1+1, X2-X1+1) * 0.25f, 1, 127);
+			int32 HalfFilterSize = FMath::Clamp<int32>(UISettings->SmoothFilterKernelScale * FMath::Max(Y2 - Y1 + 1, X2 - X1 + 1) * 0.25f, 1, 127);
 
-			for( auto It = BrushInfo.CreateIterator(); It; ++It )
+			for (auto It = BrushInfo.CreateIterator(); It; ++It)
 			{
 				int32 X, Y;
 				ALandscape::UnpackKey(It.Key(), X, Y);
 
-				if( It.Value() > 0.f )
+				if (It.Value() > 0.f)
 				{
 					int32 FilterValue = 0;
 					int32 FilterSamplingNumber = 0;
 
-					for( int32 y = Y - HalfFilterSize; y <= Y + HalfFilterSize; y++ )
+					for (int32 y = Y - HalfFilterSize; y <= Y + HalfFilterSize; y++)
 					{
 						int32 YY = FMath::Clamp<int32>(y, Y1, Y2);
-						for( int32 x = X - HalfFilterSize; x <= X + HalfFilterSize; x++ )
+						for (int32 x = X - HalfFilterSize; x <= X + HalfFilterSize; x++)
 						{
 							int32 XX = FMath::Clamp<int32>(x, X1, X2);
 							if (BrushInfo.Find(FIntPoint(XX, YY)))
 							{
-								FilterValue += Data[(XX-X1) + (YY-Y1)*(1+X2-X1)];
+								FilterValue += Data[(XX - X1) + (YY - Y1)*(1 + X2 - X1)];
 								FilterSamplingNumber++;
 							}
 						}
@@ -383,24 +392,24 @@ public:
 
 					FilterValue /= FilterSamplingNumber;
 
-					int32 HeightDataIndex = (X-X1) + (Y-Y1)*(1+X2-X1);
-					Data[HeightDataIndex] = FMath::Lerp( Data[HeightDataIndex], (typename ToolTarget::CacheClass::DataType)FilterValue, It.Value() * UISettings->ToolStrength * Pressure );
-				}	
+					int32 HeightDataIndex = (X - X1) + (Y - Y1)*(1 + X2 - X1);
+					Data[HeightDataIndex] = FMath::Lerp(Data[HeightDataIndex], (typename ToolTarget::CacheClass::DataType)FilterValue, It.Value() * UISettings->ToolStrength * Pressure);
+				}
 			}
 		}
 
-		this->Cache.SetCachedData(X1,Y1,X2,Y2,Data, UISettings->PaintingRestriction);
+		this->Cache.SetCachedData(X1, Y1, X2, Y2, Data, UISettings->PaintingRestriction);
 		this->Cache.Flush();
 	}
 
 };
 
 template<class ToolTarget>
-class FLandscapeToolSmooth : public FLandscapeToolPaintBase<ToolTarget, FLandscapeToolStrokeSmooth<ToolTarget> >
+class FLandscapeToolSmooth : public FLandscapeToolPaintBase<ToolTarget, FLandscapeToolStrokeSmooth<ToolTarget>>
 {
 public:
-	FLandscapeToolSmooth(class FEdModeLandscape* InEdMode)
-	:	FLandscapeToolPaintBase<ToolTarget, FLandscapeToolStrokeSmooth<ToolTarget> >(InEdMode)
+	FLandscapeToolSmooth(FEdModeLandscape* InEdMode)
+		: FLandscapeToolPaintBase<ToolTarget, FLandscapeToolStrokeSmooth<ToolTarget> >(InEdMode)
 	{}
 
 	virtual const TCHAR* GetToolName() override { return TEXT("Smooth"); }
@@ -412,7 +421,7 @@ public:
 // FLandscapeToolFlatten
 //
 template<class ToolTarget>
-class FLandscapeToolStrokeFlatten : public FLandscapeToolStrokePaintBase<ToolTarget>
+class FLandscapeToolStrokeFlatten : public FLandscapeToolStrokePaintBase < ToolTarget >
 {
 	typename ToolTarget::CacheClass::DataType FlattenHeight;
 
@@ -423,9 +432,9 @@ class FLandscapeToolStrokeFlatten : public FLandscapeToolStrokePaintBase<ToolTar
 
 public:
 	FLandscapeToolStrokeFlatten(FEdModeLandscape* InEdMode, const FLandscapeToolTarget& InTarget)
-	:	FLandscapeToolStrokePaintBase<ToolTarget>(InEdMode, InTarget)
-	,	bInitializedFlattenHeight(false)
-	,	bTargetIsHeightmap(InTarget.TargetType == ELandscapeToolTargetType::Heightmap)
+		: FLandscapeToolStrokePaintBase<ToolTarget>(InEdMode, InTarget)
+		, bInitializedFlattenHeight(false)
+		, bTargetIsHeightmap(InTarget.TargetType == ELandscapeToolTargetType::Heightmap)
 	{
 		if (InEdMode->UISettings->bUseFlattenTarget && bTargetIsHeightmap)
 		{
@@ -448,14 +457,14 @@ public:
 			int32 FlattenHeightX = FMath::FloorToInt(FlattenX);
 			int32 FlattenHeightY = FMath::FloorToInt(FlattenY);
 
-			this->Cache.CacheData(FlattenHeightX,FlattenHeightY,FlattenHeightX+1,FlattenHeightY+1);
+			this->Cache.CacheData(FlattenHeightX, FlattenHeightY, FlattenHeightX + 1, FlattenHeightY + 1);
 			float HeightValue = this->Cache.GetValue(FlattenX, FlattenY);
 			FlattenHeight = HeightValue;
 
 			if (UISettings->bUseSlopeFlatten && bTargetIsHeightmap)
 			{
 				FlattenNormal = this->Cache.GetNormal(FlattenHeightX, FlattenHeightY);
-				FlattenPlaneDist = -(FlattenNormal | FVector(FlattenX, FlattenY, HeightValue) );
+				FlattenPlaneDist = -(FlattenNormal | FVector(FlattenX, FlattenY, HeightValue));
 			}
 
 			bInitializedFlattenHeight = true;
@@ -479,25 +488,25 @@ public:
 		X2 += 1;
 		Y2 += 1;
 
-		this->Cache.CacheData(X1,Y1,X2,Y2);
+		this->Cache.CacheData(X1, Y1, X2, Y2);
 
 		TArray<typename ToolTarget::CacheClass::DataType> HeightData;
-		this->Cache.GetCachedData(X1,Y1,X2,Y2,HeightData);
+		this->Cache.GetCachedData(X1, Y1, X2, Y2, HeightData);
 
 		// Apply the brush
-		for( auto It = BrushInfo.CreateIterator(); It; ++It )
+		for (auto It = BrushInfo.CreateIterator(); It; ++It)
 		{
 			int32 X, Y;
 			ALandscape::UnpackKey(It.Key(), X, Y);
 
-			if( It.Value() > 0.f )
+			if (It.Value() > 0.f)
 			{
-				int32 HeightDataIndex = (X-X1) + (Y-Y1)*(1+X2-X1);
+				int32 HeightDataIndex = (X - X1) + (Y - Y1)*(1 + X2 - X1);
 
 				if (!(UISettings->bUseSlopeFlatten && bTargetIsHeightmap))
 				{
 					int32 Delta = HeightData[HeightDataIndex] - FlattenHeight;
-					switch(UISettings->FlattenMode)
+					switch (UISettings->FlattenMode)
 					{
 					case ELandscapeToolNoiseMode::Add:
 						if (Delta < 0)
@@ -513,46 +522,46 @@ public:
 						break;
 					default:
 					case ELandscapeToolNoiseMode::Both:
-						HeightData[HeightDataIndex] = FMath::Lerp( HeightData[HeightDataIndex], FlattenHeight, It.Value() * UISettings->ToolStrength * Pressure );
+						HeightData[HeightDataIndex] = FMath::Lerp(HeightData[HeightDataIndex], FlattenHeight, It.Value() * UISettings->ToolStrength * Pressure);
 						break;
 					}
 				}
 				else
 				{
-					typename ToolTarget::CacheClass::DataType DestValue = -( FlattenNormal.X * X + FlattenNormal.Y * Y + FlattenPlaneDist ) / FlattenNormal.Z;
+					typename ToolTarget::CacheClass::DataType DestValue = -(FlattenNormal.X * X + FlattenNormal.Y * Y + FlattenPlaneDist) / FlattenNormal.Z;
 					//float PlaneDist = FlattenNormal | FVector(X, Y, HeightData(HeightDataIndex)) + FlattenPlaneDist;
 					float PlaneDist = HeightData[HeightDataIndex] - DestValue;
 					DestValue = HeightData[HeightDataIndex] - PlaneDist * It.Value() * UISettings->ToolStrength * Pressure;
-					switch(UISettings->FlattenMode)
+					switch (UISettings->FlattenMode)
 					{
 					case ELandscapeToolNoiseMode::Add:
 						if (PlaneDist < 0)
 						{
-							HeightData[HeightDataIndex] = FMath::Lerp( HeightData[HeightDataIndex], DestValue, It.Value() * UISettings->ToolStrength * Pressure );
+							HeightData[HeightDataIndex] = FMath::Lerp(HeightData[HeightDataIndex], DestValue, It.Value() * UISettings->ToolStrength * Pressure);
 						}
 						break;
 					case ELandscapeToolNoiseMode::Sub:
 						if (PlaneDist > 0)
 						{
-							HeightData[HeightDataIndex] = FMath::Lerp( HeightData[HeightDataIndex], DestValue, It.Value() * UISettings->ToolStrength * Pressure );
+							HeightData[HeightDataIndex] = FMath::Lerp(HeightData[HeightDataIndex], DestValue, It.Value() * UISettings->ToolStrength * Pressure);
 						}
 						break;
 					default:
 					case ELandscapeToolNoiseMode::Both:
-						HeightData[HeightDataIndex] = FMath::Lerp( HeightData[HeightDataIndex], DestValue, It.Value() * UISettings->ToolStrength * Pressure );
+						HeightData[HeightDataIndex] = FMath::Lerp(HeightData[HeightDataIndex], DestValue, It.Value() * UISettings->ToolStrength * Pressure);
 						break;
 					}
 				}
 			}
 		}
 
-		this->Cache.SetCachedData(X1,Y1,X2,Y2,HeightData, UISettings->PaintingRestriction);
+		this->Cache.SetCachedData(X1, Y1, X2, Y2, HeightData, UISettings->PaintingRestriction);
 		this->Cache.Flush();
 	}
 };
 
 template<class ToolTarget>
-class FLandscapeToolFlatten : public FLandscapeToolPaintBase<ToolTarget, FLandscapeToolStrokeFlatten<ToolTarget>>
+class FLandscapeToolFlatten : public FLandscapeToolPaintBase < ToolTarget, FLandscapeToolStrokeFlatten<ToolTarget> >
 {
 protected:
 	FVector LastMousePosition;
@@ -560,11 +569,11 @@ protected:
 	UStaticMeshComponent* MeshComponent;
 
 public:
-	FLandscapeToolFlatten(class FEdModeLandscape* InEdMode)
-	:	FLandscapeToolPaintBase<ToolTarget, FLandscapeToolStrokeFlatten<ToolTarget> >(InEdMode)
-	,	LastMousePosition(FVector::ZeroVector)
-	,	PlaneMesh(LoadObject<UStaticMesh>(NULL, TEXT("/Engine/EditorLandscapeResources/FlattenPlaneMesh.FlattenPlaneMesh")))
-	,	MeshComponent(NULL)
+	FLandscapeToolFlatten(FEdModeLandscape* InEdMode)
+		: FLandscapeToolPaintBase<ToolTarget, FLandscapeToolStrokeFlatten<ToolTarget> >(InEdMode)
+		, LastMousePosition(FVector::ZeroVector)
+		, PlaneMesh(LoadObject<UStaticMesh>(NULL, TEXT("/Engine/EditorLandscapeResources/FlattenPlaneMesh.FlattenPlaneMesh")))
+		, MeshComponent(NULL)
 	{
 		check(PlaneMesh);
 	}
@@ -572,7 +581,7 @@ public:
 	virtual const TCHAR* GetToolName() override { return TEXT("Flatten"); }
 	virtual FText GetDisplayName() override { return NSLOCTEXT("UnrealEd", "LandscapeMode_Flatten", "Flatten"); };
 
-	virtual void Tick(FEditorViewportClient* ViewportClient,float DeltaTime)
+	virtual void Tick(FEditorViewportClient* ViewportClient, float DeltaTime)
 	{
 		FLandscapeToolPaintBase<ToolTarget, FLandscapeToolStrokeFlatten<ToolTarget>>::Tick(ViewportClient, DeltaTime);
 
@@ -636,7 +645,7 @@ class FLandscapeToolStrokeNoise : public FLandscapeToolStrokePaintBase<ToolTarge
 {
 public:
 	FLandscapeToolStrokeNoise(FEdModeLandscape* InEdMode, const FLandscapeToolTarget& InTarget)
-	:	FLandscapeToolStrokePaintBase<ToolTarget>(InEdMode, InTarget)
+		: FLandscapeToolStrokePaintBase<ToolTarget>(InEdMode, InTarget)
 	{}
 
 	virtual void Apply(FEditorViewportClient* ViewportClient, FLandscapeBrush* Brush, const ULandscapeEditorObject* UISettings, const TArray<FLandscapeToolMousePosition>& MousePositions) override
@@ -660,9 +669,9 @@ public:
 		X2 += 1;
 		Y2 += 1;
 
-		this->Cache.CacheData(X1,Y1,X2,Y2);
+		this->Cache.CacheData(X1, Y1, X2, Y2);
 		TArray<typename ToolTarget::CacheClass::DataType> Data;
-		this->Cache.GetCachedData(X1,Y1,X2,Y2,Data);
+		this->Cache.GetCachedData(X1, Y1, X2, Y2, Data);
 
 		float BrushSizeAdjust = 1.0f;
 		if (ToolTarget::TargetType != ELandscapeToolTargetType::Weightmap && UISettings->BrushRadius < UISettings->MaximumValueRadius)
@@ -673,14 +682,14 @@ public:
 		bool bUseWeightTargetValue = UISettings->bUseWeightTargetValue && ToolTarget::TargetType == ELandscapeToolTargetType::Weightmap;
 
 		// Apply the brush
-		for( auto It = BrushInfo.CreateIterator(); It; ++It )
+		for (auto It = BrushInfo.CreateIterator(); It; ++It)
 		{
 			int32 X, Y;
 			ALandscape::UnpackKey(It.Key(), X, Y);
 
-			if( It.Value() > 0.f )
+			if (It.Value() > 0.f)
 			{
-				float OriginalValue = Data[(X-X1) + (Y-Y1)*(1+X2-X1)];
+				float OriginalValue = Data[(X - X1) + (Y - Y1)*(1 + X2 - X1)];
 				if (bUseWeightTargetValue)
 				{
 					FNoiseParameter NoiseParam(0, UISettings->NoiseScale, 255.f / 2.f);
@@ -701,29 +710,29 @@ public:
 						}
 						break;
 					}
-					Data[(X-X1) + (Y-Y1)*(1+X2-X1)] = ToolTarget::CacheClass::ClampValue( FMath::RoundToInt(FMath::Lerp( OriginalValue, DestValue, It.Value() * UISettings->ToolStrength * Pressure)) );
+					Data[(X - X1) + (Y - Y1)*(1 + X2 - X1)] = ToolTarget::CacheClass::ClampValue(FMath::RoundToInt(FMath::Lerp(OriginalValue, DestValue, It.Value() * UISettings->ToolStrength * Pressure)));
 				}
 				else
 				{
 					float TotalStrength = It.Value() * UISettings->ToolStrength * Pressure * ToolTarget::StrengthMultiplier(this->LandscapeInfo, UISettings->BrushRadius);
-					FNoiseParameter NoiseParam(0, UISettings->NoiseScale,  TotalStrength * BrushSizeAdjust);
+					FNoiseParameter NoiseParam(0, UISettings->NoiseScale, TotalStrength * BrushSizeAdjust);
 					float PaintAmount = ELandscapeToolNoiseMode::Conversion(UISettings->NoiseMode, NoiseParam.NoiseAmount, NoiseParam.Sample(X, Y));
-					Data[(X-X1) + (Y-Y1)*(1+X2-X1)] = ToolTarget::CacheClass::ClampValue(OriginalValue + PaintAmount);
+					Data[(X - X1) + (Y - Y1)*(1 + X2 - X1)] = ToolTarget::CacheClass::ClampValue(OriginalValue + PaintAmount);
 				}
-			}	
+			}
 		}
 
-		this->Cache.SetCachedData(X1,Y1,X2,Y2,Data, UISettings->PaintingRestriction);
+		this->Cache.SetCachedData(X1, Y1, X2, Y2, Data, UISettings->PaintingRestriction);
 		this->Cache.Flush();
 	}
 };
 
 template<class ToolTarget>
-class FLandscapeToolNoise : public FLandscapeToolPaintBase<ToolTarget, FLandscapeToolStrokeNoise<ToolTarget> >
+class FLandscapeToolNoise : public FLandscapeToolPaintBase < ToolTarget, FLandscapeToolStrokeNoise<ToolTarget> >
 {
 public:
-	FLandscapeToolNoise(class FEdModeLandscape* InEdMode)
-	:	FLandscapeToolPaintBase<ToolTarget, FLandscapeToolStrokeNoise<ToolTarget> >(InEdMode)
+	FLandscapeToolNoise(FEdModeLandscape* InEdMode)
+		: FLandscapeToolPaintBase<ToolTarget, FLandscapeToolStrokeNoise<ToolTarget> >(InEdMode)
 	{}
 
 	virtual const TCHAR* GetToolName() override { return TEXT("Noise"); }
@@ -734,54 +743,64 @@ public:
 //
 // Toolset initialization
 //
-void FEdModeLandscape::IntializeToolSet_Paint()
+void FEdModeLandscape::InitializeTool_Paint()
 {
-	FLandscapeToolSet* ToolSet_Sculpt = new(LandscapeToolSets) FLandscapeToolSet(TEXT("ToolSet_Sculpt"));
-	ToolSet_Sculpt->AddTool(new FLandscapeToolPaint<FHeightmapToolTarget>(this));
-
-	ToolSet_Sculpt->ValidBrushes.Add("BrushSet_Circle");
-	ToolSet_Sculpt->ValidBrushes.Add("BrushSet_Alpha");
-	ToolSet_Sculpt->ValidBrushes.Add("BrushSet_Pattern");
-	ToolSet_Sculpt->ValidBrushes.Add("BrushSet_Component");
-
-	FLandscapeToolSet* ToolSet_Paint = new(LandscapeToolSets) FLandscapeToolSet(TEXT("ToolSet_Paint"));
-	ToolSet_Paint->AddTool(new FLandscapeToolPaint<FWeightmapToolTarget>(this));
-
-	ToolSet_Paint->ValidBrushes.Add("BrushSet_Circle");
-	ToolSet_Paint->ValidBrushes.Add("BrushSet_Alpha");
-	ToolSet_Paint->ValidBrushes.Add("BrushSet_Pattern");
-	ToolSet_Paint->ValidBrushes.Add("BrushSet_Component");
+	auto Tool_Sculpt = MakeUnique<FLandscapeToolSculpt>(this);
+	Tool_Sculpt->ValidBrushes.Add("BrushSet_Circle");
+	Tool_Sculpt->ValidBrushes.Add("BrushSet_Alpha");
+	Tool_Sculpt->ValidBrushes.Add("BrushSet_Pattern");
+	Tool_Sculpt->ValidBrushes.Add("BrushSet_Component");
+	LandscapeTools.Add(MoveTemp(Tool_Sculpt));
+
+	auto Tool_Paint = MakeUnique<FLandscapeToolPaint>(this);
+	Tool_Paint->ValidBrushes.Add("BrushSet_Circle");
+	Tool_Paint->ValidBrushes.Add("BrushSet_Alpha");
+	Tool_Paint->ValidBrushes.Add("BrushSet_Pattern");
+	Tool_Paint->ValidBrushes.Add("BrushSet_Component");
+	LandscapeTools.Add(MoveTemp(Tool_Paint));
 }
 
-void FEdModeLandscape::IntializeToolSet_Smooth()
+void FEdModeLandscape::InitializeTool_Smooth()
 {
-	FLandscapeToolSet* ToolSet_Smooth = new(LandscapeToolSets) FLandscapeToolSet(TEXT("ToolSet_Smooth"));
-	ToolSet_Smooth->AddTool(new FLandscapeToolSmooth<FHeightmapToolTarget>(this));
-	ToolSet_Smooth->AddTool(new FLandscapeToolSmooth<FWeightmapToolTarget>(this));
-
-	ToolSet_Smooth->ValidBrushes.Add("BrushSet_Circle");
-	ToolSet_Smooth->ValidBrushes.Add("BrushSet_Alpha");
-	ToolSet_Smooth->ValidBrushes.Add("BrushSet_Pattern");
+	auto Tool_Smooth_Heightmap = MakeUnique<FLandscapeToolSmooth<FHeightmapToolTarget>>(this);
+	Tool_Smooth_Heightmap->ValidBrushes.Add("BrushSet_Circle");
+	Tool_Smooth_Heightmap->ValidBrushes.Add("BrushSet_Alpha");
+	Tool_Smooth_Heightmap->ValidBrushes.Add("BrushSet_Pattern");
+	LandscapeTools.Add(MoveTemp(Tool_Smooth_Heightmap));
+
+	auto Tool_Smooth_Weightmap = MakeUnique<FLandscapeToolSmooth<FWeightmapToolTarget>>(this);
+	Tool_Smooth_Weightmap->ValidBrushes.Add("BrushSet_Circle");
+	Tool_Smooth_Weightmap->ValidBrushes.Add("BrushSet_Alpha");
+	Tool_Smooth_Weightmap->ValidBrushes.Add("BrushSet_Pattern");
+	LandscapeTools.Add(MoveTemp(Tool_Smooth_Weightmap));
 }
 
-void FEdModeLandscape::IntializeToolSet_Flatten()
+void FEdModeLandscape::InitializeTool_Flatten()
 {
-	FLandscapeToolSet* ToolSet_Flatten = new(LandscapeToolSets) FLandscapeToolSet(TEXT("ToolSet_Flatten"));
-	ToolSet_Flatten->AddTool(new FLandscapeToolFlatten<FHeightmapToolTarget>(this));
-	ToolSet_Flatten->AddTool(new FLandscapeToolFlatten<FWeightmapToolTarget>(this));
-
-	ToolSet_Flatten->ValidBrushes.Add("BrushSet_Circle");
-	ToolSet_Flatten->ValidBrushes.Add("BrushSet_Alpha");
-	ToolSet_Flatten->ValidBrushes.Add("BrushSet_Pattern");
+	auto Tool_Flatten_Heightmap = MakeUnique<FLandscapeToolFlatten<FHeightmapToolTarget>>(this);
+	Tool_Flatten_Heightmap->ValidBrushes.Add("BrushSet_Circle");
+	Tool_Flatten_Heightmap->ValidBrushes.Add("BrushSet_Alpha");
+	Tool_Flatten_Heightmap->ValidBrushes.Add("BrushSet_Pattern");
+	LandscapeTools.Add(MoveTemp(Tool_Flatten_Heightmap));
+
+	auto Tool_Flatten_Weightmap = MakeUnique<FLandscapeToolFlatten<FWeightmapToolTarget>>(this);
+	Tool_Flatten_Weightmap->ValidBrushes.Add("BrushSet_Circle");
+	Tool_Flatten_Weightmap->ValidBrushes.Add("BrushSet_Alpha");
+	Tool_Flatten_Weightmap->ValidBrushes.Add("BrushSet_Pattern");
+	LandscapeTools.Add(MoveTemp(Tool_Flatten_Weightmap));
 }
 
-void FEdModeLandscape::IntializeToolSet_Noise()
+void FEdModeLandscape::InitializeTool_Noise()
 {
-	FLandscapeToolSet* ToolSet_Noise = new(LandscapeToolSets) FLandscapeToolSet(TEXT("ToolSet_Noise"));
-	ToolSet_Noise->AddTool(new FLandscapeToolNoise<FHeightmapToolTarget>(this));
-	ToolSet_Noise->AddTool(new FLandscapeToolNoise<FWeightmapToolTarget>(this));
-
-	ToolSet_Noise->ValidBrushes.Add("BrushSet_Circle");
-	ToolSet_Noise->ValidBrushes.Add("BrushSet_Alpha");
-	ToolSet_Noise->ValidBrushes.Add("BrushSet_Pattern");
+	auto Tool_Noise_Heightmap = MakeUnique<FLandscapeToolNoise<FHeightmapToolTarget>>(this);
+	Tool_Noise_Heightmap->ValidBrushes.Add("BrushSet_Circle");
+	Tool_Noise_Heightmap->ValidBrushes.Add("BrushSet_Alpha");
+	Tool_Noise_Heightmap->ValidBrushes.Add("BrushSet_Pattern");
+	LandscapeTools.Add(MoveTemp(Tool_Noise_Heightmap));
+
+	auto Tool_Noise_Weightmap = MakeUnique<FLandscapeToolNoise<FWeightmapToolTarget>>(this);
+	Tool_Noise_Weightmap->ValidBrushes.Add("BrushSet_Circle");
+	Tool_Noise_Weightmap->ValidBrushes.Add("BrushSet_Alpha");
+	Tool_Noise_Weightmap->ValidBrushes.Add("BrushSet_Pattern");
+	LandscapeTools.Add(MoveTemp(Tool_Noise_Weightmap));
 }
diff --git a/Engine/Source/Editor/LandscapeEditor/Private/LandscapeEdModeRampTool.cpp b/Engine/Source/Editor/LandscapeEditor/Private/LandscapeEdModeRampTool.cpp
--- a/Engine/Source/Editor/LandscapeEditor/Private/LandscapeEdModeRampTool.cpp
+++ b/Engine/Source/Editor/LandscapeEditor/Private/LandscapeEdModeRampTool.cpp
@@ -18,7 +18,7 @@ public:
 	typedef FVector2D InterpolantType;
 
 	/** Initialization constructor. */
-	FLandscapeRampToolHeightRasterPolicy(TArray<uint16>& InData, int32 InMinX, int32 InMinY, int32 InMaxX, int32 InMaxY, bool InbRaiseTerrain, bool InbLowerTerrain):
+	FLandscapeRampToolHeightRasterPolicy(TArray<uint16>& InData, int32 InMinX, int32 InMinY, int32 InMaxX, int32 InMaxY, bool InbRaiseTerrain, bool InbLowerTerrain) :
 		Data(InData),
 		MinX(InMinX),
 		MinY(InMinY),
@@ -40,13 +40,13 @@ protected:
 
 	void ProcessPixel(int32 X, int32 Y, const InterpolantType& Interpolant, bool BackFacing)
 	{
-		const float CosInterpX = (Interpolant.X >= 1 ? 1 : 0.5f - 0.5f * FMath::Cos(Interpolant.X * PI) );
+		const float CosInterpX = (Interpolant.X >= 1 ? 1 : 0.5f - 0.5f * FMath::Cos(Interpolant.X * PI));
 		const float Alpha = CosInterpX;
-		uint16& Dest = Data[(Y-MinY)*(1+MaxX-MinX) + X-MinX];
+		uint16& Dest = Data[(Y - MinY)*(1 + MaxX - MinX) + X - MinX];
 		float Value = FMath::Lerp((float)Dest, Interpolant.Y, Alpha);
 		uint16 DValue = (uint32)FMath::Clamp<float>(Value, 0, LandscapeDataAccess::MaxValue);
-		if ( (bRaiseTerrain && DValue > Dest) ||
-			(bLowerTerrain && DValue < Dest) )
+		if ((bRaiseTerrain && DValue > Dest) ||
+			(bLowerTerrain && DValue < Dest))
 		{
 			Dest = DValue;
 		}
@@ -55,12 +55,12 @@ protected:
 private:
 	TArray<uint16>& Data;
 	int32 MinX, MinY, MaxX, MaxY;
-	uint32 bRaiseTerrain:1, bLowerTerrain:1;
+	uint32 bRaiseTerrain : 1, bLowerTerrain : 1;
 };
 
 class HLandscapeRampToolPointHitProxy : public HHitProxy
 {
-	DECLARE_HIT_PROXY( );
+	DECLARE_HIT_PROXY();
 
 	int8 Point;
 
@@ -81,7 +81,7 @@ IMPLEMENT_HIT_PROXY(HLandscapeRampToolPointHitProxy, HHitProxy)
 class FLandscapeToolRamp : public FLandscapeTool
 {
 protected:
-	class FEdModeLandscape* EdMode;
+	FEdModeLandscape* EdMode;
 	UTexture2D* SpriteTexture;
 	FVector Points[2];
 	int8 NumPoints;
@@ -105,9 +105,9 @@ public:
 	virtual void SetEditRenderType() override { GLandscapeEditRenderMode = ELandscapeEditRenderMode::None | (GLandscapeEditRenderMode & ELandscapeEditRenderMode::BitMaskForMask); }
 	virtual bool SupportsMask() override { return false; }
 
-	virtual bool IsValidForTarget(const FLandscapeToolTarget& Target) override
+	virtual ELandscapeToolTargetTypeMask::Type GetSupportedTargetTypes() override
 	{
-		return Target.TargetType == ELandscapeToolTargetType::Heightmap;
+		return ELandscapeToolTargetTypeMask::Heightmap;
 	}
 
 	virtual void EnterTool() override
@@ -171,7 +171,7 @@ public:
 	{
 		if (HitProxy)
 		{
-			if (HitProxy->IsA(HLandscapeRampToolPointHitProxy::StaticGetType()) )
+			if (HitProxy->IsA(HLandscapeRampToolPointHitProxy::StaticGetType()))
 			{
 				HLandscapeRampToolPointHitProxy* PointHitProxy = (HLandscapeRampToolPointHitProxy*)HitProxy;
 				SelectedPoint = PointHitProxy->Point;
@@ -362,7 +362,7 @@ public:
 		return true;
 	}
 
-	virtual bool IsSelectionAllowed( AActor* InActor, bool bInSelection ) const override
+	virtual bool IsSelectionAllowed(AActor* InActor, bool bInSelection) const override
 	{
 		// Only filter selection not deselection
 		if (bInSelection)
@@ -428,7 +428,7 @@ public:
 
 	virtual void ApplyRamp()
 	{
-		FScopedTransaction Transaction(LOCTEXT("Ramp_Apply", "Landscape Editing: Add ramp") );
+		FScopedTransaction Transaction(LOCTEXT("Ramp_Apply", "Landscape Editing: Add ramp"));
 
 		const ALandscapeProxy* LandscapeProxy = EdMode->CurrentToolTarget.LandscapeInfo->GetLandscapeProxy();
 		const FTransform LandscapeToWorld = LandscapeProxy->LandscapeActorToWorld();
@@ -453,10 +453,10 @@ public:
 		Heights[0] = Points[0].Z * LANDSCAPE_INV_ZSCALE + LandscapeDataAccess::MidValue;
 		Heights[1] = Points[1].Z * LANDSCAPE_INV_ZSCALE + LandscapeDataAccess::MidValue;
 
-		int32 MinX = FMath::CeilToInt( FMath::Min(FMath::Min(OuterVerts[0][0].X, OuterVerts[0][1].X),FMath::Min(OuterVerts[1][0].X, OuterVerts[1][1].X))) - 1; // +/- 1 to make sure we have enough data for calculating correct normals
-		int32 MinY = FMath::CeilToInt( FMath::Min(FMath::Min(OuterVerts[0][0].Y, OuterVerts[0][1].Y),FMath::Min(OuterVerts[1][0].Y, OuterVerts[1][1].Y))) - 1;
-		int32 MaxX = FMath::FloorToInt(FMath::Max(FMath::Max(OuterVerts[0][0].X, OuterVerts[0][1].X),FMath::Max(OuterVerts[1][0].X, OuterVerts[1][1].X))) + 1;
-		int32 MaxY = FMath::FloorToInt(FMath::Max(FMath::Max(OuterVerts[0][0].Y, OuterVerts[0][1].Y),FMath::Max(OuterVerts[1][0].Y, OuterVerts[1][1].Y))) + 1;
+		int32 MinX = FMath::CeilToInt(FMath::Min(FMath::Min(OuterVerts[0][0].X, OuterVerts[0][1].X), FMath::Min(OuterVerts[1][0].X, OuterVerts[1][1].X))) - 1; // +/- 1 to make sure we have enough data for calculating correct normals
+		int32 MinY = FMath::CeilToInt(FMath::Min(FMath::Min(OuterVerts[0][0].Y, OuterVerts[0][1].Y), FMath::Min(OuterVerts[1][0].Y, OuterVerts[1][1].Y))) - 1;
+		int32 MaxX = FMath::FloorToInt(FMath::Max(FMath::Max(OuterVerts[0][0].X, OuterVerts[0][1].X), FMath::Max(OuterVerts[1][0].X, OuterVerts[1][1].X))) + 1;
+		int32 MaxY = FMath::FloorToInt(FMath::Max(FMath::Max(OuterVerts[0][0].Y, OuterVerts[0][1].Y), FMath::Max(OuterVerts[1][0].Y, OuterVerts[1][1].Y))) + 1;
 
 		FLandscapeEditDataInterface LandscapeEdit(EdMode->CurrentToolTarget.LandscapeInfo.Get());
 
@@ -466,7 +466,7 @@ public:
 		if (bRaiseTerrain || bLowerTerrain)
 		{
 			TArray<uint16> Data;
-			Data.AddZeroed( (1+MaxY-MinY) * (1+MaxX-MinX) );
+			Data.AddZeroed((1 + MaxY - MinY) * (1 + MaxX - MinX));
 
 			int32 ValidMinX = MinX;
 			int32 ValidMinY = MinY;
@@ -488,7 +488,7 @@ public:
 			MaxY = ValidMaxY;
 
 			FTriangleRasterizer<FLandscapeRampToolHeightRasterPolicy> Rasterizer(
-				FLandscapeRampToolHeightRasterPolicy(Data, MinX, MinY, MaxX, MaxY, bRaiseTerrain, bLowerTerrain) );
+				FLandscapeRampToolHeightRasterPolicy(Data, MinX, MinY, MaxX, MaxY, bRaiseTerrain, bLowerTerrain));
 
 			// Left
 			Rasterizer.DrawTriangle(FVector2D(0, Heights[0]), FVector2D(1, Heights[0]), FVector2D(0, Heights[1]), OuterVerts[0][0], InnerVerts[0][0], OuterVerts[1][0], false);
@@ -540,9 +540,9 @@ public:
 
 void FEdModeLandscape::ApplyRampTool()
 {
-	if (CurrentToolSet->GetToolSetName() == FName("ToolSet_Ramp"))
+	if (CurrentTool->GetToolName() == FName("Ramp"))
 	{
-		FLandscapeToolRamp* RampTool = (FLandscapeToolRamp*)CurrentToolSet->GetTool();
+		FLandscapeToolRamp* RampTool = (FLandscapeToolRamp*)CurrentTool;
 		RampTool->ApplyRamp();
 		GEditor->RedrawLevelEditingViewports();
 	}
@@ -550,9 +550,9 @@ void FEdModeLandscape::ApplyRampTool()
 
 bool FEdModeLandscape::CanApplyRampTool()
 {
-	if (CurrentToolSet->GetToolSetName() == FName("ToolSet_Ramp"))
+	if (CurrentTool->GetToolName() == FName("Ramp"))
 	{
-		FLandscapeToolRamp* RampTool = (FLandscapeToolRamp*)CurrentToolSet->GetTool();
+		FLandscapeToolRamp* RampTool = (FLandscapeToolRamp*)CurrentTool;
 
 		return RampTool->CanApplyRamp();
 	}
@@ -561,9 +561,9 @@ bool FEdModeLandscape::CanApplyRampTool()
 
 void FEdModeLandscape::ResetRampTool()
 {
-	if (CurrentToolSet->GetToolSetName() == FName("ToolSet_Ramp"))
+	if (CurrentTool->GetToolName() == FName("Ramp"))
 	{
-		FLandscapeToolRamp* RampTool = (FLandscapeToolRamp*)CurrentToolSet->GetTool();
+		FLandscapeToolRamp* RampTool = (FLandscapeToolRamp*)CurrentTool;
 		RampTool->ResetRamp();
 		GEditor->RedrawLevelEditingViewports();
 	}
@@ -572,11 +572,11 @@ void FEdModeLandscape::ResetRampTool()
 //
 // Toolset initialization
 //
-void FEdModeLandscape::InitializeToolSet_Ramp()
+void FEdModeLandscape::InitializeTool_Ramp()
 {
-	FLandscapeToolSet* ToolSet_Ramp = new(LandscapeToolSets) FLandscapeToolSet(TEXT("ToolSet_Ramp"));
-	ToolSet_Ramp->AddTool(new FLandscapeToolRamp(this));
-	ToolSet_Ramp->ValidBrushes.Add("BrushSet_Dummy");
+	auto Tool_Ramp = MakeUnique<FLandscapeToolRamp>(this);
+	Tool_Ramp->ValidBrushes.Add("BrushSet_Dummy");
+	LandscapeTools.Add(MoveTemp(Tool_Ramp));
 }
 
 #undef LOCTEXT_NAMESPACE
\ No newline at end of file
diff --git a/Engine/Source/Editor/LandscapeEditor/Private/LandscapeEdModeSplineTools.cpp b/Engine/Source/Editor/LandscapeEditor/Private/LandscapeEdModeSplineTools.cpp
--- a/Engine/Source/Editor/LandscapeEditor/Private/LandscapeEdModeSplineTools.cpp
+++ b/Engine/Source/Editor/LandscapeEditor/Private/LandscapeEdModeSplineTools.cpp
@@ -27,7 +27,7 @@
 class FLandscapeToolSplines : public FLandscapeTool, public FEditorUndoClient
 {
 public:
-	FLandscapeToolSplines(class FEdModeLandscape* InEdMode)
+	FLandscapeToolSplines(FEdModeLandscape* InEdMode)
 		: EdMode(InEdMode)
 		, LandscapeInfo(NULL)
 		, SelectedSplineControlPoints()
@@ -47,7 +47,7 @@ public:
 	~FLandscapeToolSplines()
 	{
 		// GEditor is invalid at shutdown as the object system is unloaded before the landscape module.
-		if( UObjectInitialized() )
+		if (UObjectInitialized())
 		{
 			// Remove undo delegate
 			GEditor->UnregisterForUndo(this);
@@ -60,12 +60,7 @@ public:
 	virtual void SetEditRenderType() override { GLandscapeEditRenderMode = ELandscapeEditRenderMode::None | (GLandscapeEditRenderMode & ELandscapeEditRenderMode::BitMaskForMask); }
 	virtual bool SupportsMask() override { return false; }
 
-	virtual bool IsValidForTarget(const FLandscapeToolTarget& Target) override
-	{
-		return true; // applied to all...
-	}
-
-	void CreateSplineComponent( ALandscapeProxy* Landscape, FVector Scale3D ) 
+	void CreateSplineComponent(ALandscapeProxy* Landscape, FVector Scale3D)
 	{
 		Landscape->SplineComponent = ConstructObject<ULandscapeSplinesComponent>(ULandscapeSplinesComponent::StaticClass(), Landscape, NAME_None, RF_Transactional);
 		Landscape->SplineComponent->RelativeScale3D = Scale3D;
@@ -75,7 +70,7 @@ public:
 
 	void UpdatePropertiesWindows()
 	{
-		if( GLevelEditorModeTools().IsModeActive( EdMode->GetID() ) )
+		if (GLevelEditorModeTools().IsModeActive(EdMode->GetID()))
 		{
 			TArray<UObject*> Objects;
 			Objects.Reset(SelectedSplineControlPoints.Num() + SelectedSplineSegments.Num());
@@ -90,7 +85,7 @@ public:
 			}
 
 
-			FPropertyEditorModule& PropertyModule = FModuleManager::Get().LoadModuleChecked<FPropertyEditorModule>( TEXT("PropertyEditor") );
+			FPropertyEditorModule& PropertyModule = FModuleManager::Get().LoadModuleChecked<FPropertyEditorModule>(TEXT("PropertyEditor"));
 			PropertyModule.UpdatePropertyViews(Objects);
 		}
 	}
@@ -228,7 +223,7 @@ public:
 
 	void AddSegment(ULandscapeSplineControlPoint* Start, ULandscapeSplineControlPoint* End, bool bAutoRotateStart, bool bAutoRotateEnd)
 	{
-		FScopedTransaction Transaction( LOCTEXT("LandscapeSpline_AddSegment", "Add Landscape Spline Segment") );
+		FScopedTransaction Transaction(LOCTEXT("LandscapeSpline_AddSegment", "Add Landscape Spline Segment"));
 
 		if (Start == End)
 		{
@@ -280,23 +275,23 @@ public:
 		// set up other segment options
 		if (Start->ConnectedSegments.Num() > 0)
 		{
-			NewSegment->LayerName         = Start->ConnectedSegments[0].Segment->LayerName;
-			NewSegment->SplineMeshes      = Start->ConnectedSegments[0].Segment->SplineMeshes;
+			NewSegment->LayerName = Start->ConnectedSegments[0].Segment->LayerName;
+			NewSegment->SplineMeshes = Start->ConnectedSegments[0].Segment->SplineMeshes;
 			NewSegment->LDMaxDrawDistance = Start->ConnectedSegments[0].Segment->LDMaxDrawDistance;
-			NewSegment->bRaiseTerrain     = Start->ConnectedSegments[0].Segment->bRaiseTerrain;
-			NewSegment->bLowerTerrain     = Start->ConnectedSegments[0].Segment->bLowerTerrain;
-			NewSegment->bEnableCollision  = Start->ConnectedSegments[0].Segment->bEnableCollision;
-			NewSegment->bCastShadow       = Start->ConnectedSegments[0].Segment->bCastShadow;
+			NewSegment->bRaiseTerrain = Start->ConnectedSegments[0].Segment->bRaiseTerrain;
+			NewSegment->bLowerTerrain = Start->ConnectedSegments[0].Segment->bLowerTerrain;
+			NewSegment->bEnableCollision = Start->ConnectedSegments[0].Segment->bEnableCollision;
+			NewSegment->bCastShadow = Start->ConnectedSegments[0].Segment->bCastShadow;
 		}
 		else if (End->ConnectedSegments.Num() > 0)
 		{
-			NewSegment->LayerName         = End->ConnectedSegments[0].Segment->LayerName;
-			NewSegment->SplineMeshes      = End->ConnectedSegments[0].Segment->SplineMeshes;
+			NewSegment->LayerName = End->ConnectedSegments[0].Segment->LayerName;
+			NewSegment->SplineMeshes = End->ConnectedSegments[0].Segment->SplineMeshes;
 			NewSegment->LDMaxDrawDistance = End->ConnectedSegments[0].Segment->LDMaxDrawDistance;
-			NewSegment->bRaiseTerrain     = End->ConnectedSegments[0].Segment->bRaiseTerrain;
-			NewSegment->bLowerTerrain     = End->ConnectedSegments[0].Segment->bLowerTerrain;
-			NewSegment->bEnableCollision  = End->ConnectedSegments[0].Segment->bEnableCollision;
-			NewSegment->bCastShadow       = End->ConnectedSegments[0].Segment->bCastShadow;
+			NewSegment->bRaiseTerrain = End->ConnectedSegments[0].Segment->bRaiseTerrain;
+			NewSegment->bLowerTerrain = End->ConnectedSegments[0].Segment->bLowerTerrain;
+			NewSegment->bEnableCollision = End->ConnectedSegments[0].Segment->bEnableCollision;
+			NewSegment->bCastShadow = End->ConnectedSegments[0].Segment->bCastShadow;
 		}
 		else
 		{
@@ -331,7 +326,7 @@ public:
 
 	void AddControlPoint(ALandscapeProxy* Landscape, const FVector& LocalLocation)
 	{
-		FScopedTransaction Transaction( LOCTEXT("LandscapeSpline_AddControlPoint", "Add Landscape Spline Control Point") );
+		FScopedTransaction Transaction(LOCTEXT("LandscapeSpline_AddControlPoint", "Add Landscape Spline Control Point"));
 
 		Landscape->SplineComponent->Modify();
 
@@ -378,7 +373,7 @@ public:
 
 	void DeleteSegment(ULandscapeSplineSegment* ToDelete, bool bInDeleteLooseEnds)
 	{
-		FScopedTransaction Transaction( LOCTEXT("LandscapeSpline_DeleteSegment", "Delete Landscape Spline Segment") );
+		FScopedTransaction Transaction(LOCTEXT("LandscapeSpline_DeleteSegment", "Delete Landscape Spline Segment"));
 
 		ULandscapeSplinesComponent* SplinesComponent = ToDelete->GetOuterULandscapeSplinesComponent();
 		SplinesComponent->Modify();
@@ -421,7 +416,7 @@ public:
 
 	void DeleteControlPoint(ULandscapeSplineControlPoint* ToDelete, bool bInDeleteLooseEnds)
 	{
-		FScopedTransaction Transaction( LOCTEXT("LandscapeSpline_DeleteControlPoint", "Delete Landscape Spline Control Point") );
+		FScopedTransaction Transaction(LOCTEXT("LandscapeSpline_DeleteControlPoint", "Delete Landscape Spline Control Point"));
 
 		ULandscapeSplinesComponent* SplinesComponent = ToDelete->GetOuterULandscapeSplinesComponent();
 		SplinesComponent->Modify();
@@ -432,37 +427,37 @@ public:
 		if (ToDelete->ConnectedSegments.Num() == 2
 			&& ToDelete->ConnectedSegments[0].Segment != ToDelete->ConnectedSegments[1].Segment)
 		{
-			int32 Result = FMessageDialog::Open( EAppMsgType::YesNoCancel, LOCTEXT("WantToJoinControlPoint", "Control point has two segments attached, do you want to join them?"));
+			int32 Result = FMessageDialog::Open(EAppMsgType::YesNoCancel, LOCTEXT("WantToJoinControlPoint", "Control point has two segments attached, do you want to join them?"));
 			switch (Result)
 			{
 			case EAppReturnType::Yes:
-				{
-					// Copy the other end of connection 1 into the near end of connection 0, then delete connection 1
-					TArray<FLandscapeSplineConnection>& Connections = ToDelete->ConnectedSegments;
-					Connections[0].Segment->Modify();
-					Connections[1].Segment->Modify();
+			{
+				// Copy the other end of connection 1 into the near end of connection 0, then delete connection 1
+				TArray<FLandscapeSplineConnection>& Connections = ToDelete->ConnectedSegments;
+				Connections[0].Segment->Modify();
+				Connections[1].Segment->Modify();
 
-					Connections[0].GetNearConnection() = Connections[1].GetFarConnection();
-					Connections[0].Segment->UpdateSplinePoints();
+				Connections[0].GetNearConnection() = Connections[1].GetFarConnection();
+				Connections[0].Segment->UpdateSplinePoints();
 
-					Connections[1].Segment->DeleteSplinePoints();
+				Connections[1].Segment->DeleteSplinePoints();
 
-					// Get the control point at the *other* end of the segment and remove it from it
-					ULandscapeSplineControlPoint* OtherEnd = Connections[1].GetFarConnection().ControlPoint;
-					OtherEnd->Modify();
+				// Get the control point at the *other* end of the segment and remove it from it
+				ULandscapeSplineControlPoint* OtherEnd = Connections[1].GetFarConnection().ControlPoint;
+				OtherEnd->Modify();
 
-					FLandscapeSplineConnection* OtherConnection = OtherEnd->ConnectedSegments.FindByKey(FLandscapeSplineConnection(Connections[1].Segment, 1 - Connections[1].End));
-					*OtherConnection = FLandscapeSplineConnection(Connections[0].Segment, Connections[0].End);
+				FLandscapeSplineConnection* OtherConnection = OtherEnd->ConnectedSegments.FindByKey(FLandscapeSplineConnection(Connections[1].Segment, 1 - Connections[1].End));
+				*OtherConnection = FLandscapeSplineConnection(Connections[0].Segment, Connections[0].End);
 
-					SplinesComponent->Segments.Remove(Connections[1].Segment);
+				SplinesComponent->Segments.Remove(Connections[1].Segment);
 
-					ToDelete->ConnectedSegments.Empty();
+				ToDelete->ConnectedSegments.Empty();
 
-					SplinesComponent->ControlPoints.Remove(ToDelete);
-					SplinesComponent->MarkRenderStateDirty();
+				SplinesComponent->ControlPoints.Remove(ToDelete);
+				SplinesComponent->MarkRenderStateDirty();
 
-					return;
-				}
+				return;
+			}
 				break;
 			case EAppReturnType::No:
 				// Use the "delete all segments" code below
@@ -502,7 +497,7 @@ public:
 
 	void SplitSegment(ULandscapeSplineSegment* Segment, const FVector& LocalLocation)
 	{
-		FScopedTransaction Transaction( LOCTEXT("LandscapeSpline_SplitSegment", "Split Landscape Spline Segment") );
+		FScopedTransaction Transaction(LOCTEXT("LandscapeSpline_SplitSegment", "Split Landscape Spline Segment"));
 
 		ULandscapeSplinesComponent* SplinesComponent = Segment->GetOuterULandscapeSplinesComponent();
 		SplinesComponent->Modify();
@@ -533,13 +528,13 @@ public:
 		NewSegment->Connections[1].ControlPoint = Segment->Connections[1].ControlPoint;
 		NewSegment->Connections[1].TangentLen = Segment->Connections[1].TangentLen * (1 - t);
 		NewSegment->Connections[1].ControlPoint->ConnectedSegments.Add(FLandscapeSplineConnection(NewSegment, 1));
-		NewSegment->LayerName         = Segment->LayerName;
-		NewSegment->SplineMeshes      = Segment->SplineMeshes;
+		NewSegment->LayerName = Segment->LayerName;
+		NewSegment->SplineMeshes = Segment->SplineMeshes;
 		NewSegment->LDMaxDrawDistance = Segment->LDMaxDrawDistance;
-		NewSegment->bRaiseTerrain     = Segment->bRaiseTerrain;
-		NewSegment->bLowerTerrain     = Segment->bLowerTerrain;
-		NewSegment->bEnableCollision  = Segment->bEnableCollision;
-		NewSegment->bCastShadow       = Segment->bCastShadow;
+		NewSegment->bRaiseTerrain = Segment->bRaiseTerrain;
+		NewSegment->bLowerTerrain = Segment->bLowerTerrain;
+		NewSegment->bEnableCollision = Segment->bEnableCollision;
+		NewSegment->bCastShadow = Segment->bCastShadow;
 
 		Segment->Connections[0].TangentLen *= t;
 		Segment->Connections[1].ControlPoint->ConnectedSegments.Remove(FLandscapeSplineConnection(Segment, 1));
@@ -558,7 +553,7 @@ public:
 
 	void FlipSegment(ULandscapeSplineSegment* Segment)
 	{
-		FScopedTransaction Transaction( LOCTEXT("LandscapeSpline_FlipSegment", "Flip Landscape Spline Segment") );
+		FScopedTransaction Transaction(LOCTEXT("LandscapeSpline_FlipSegment", "Flip Landscape Spline Segment"));
 
 		ULandscapeSplinesComponent* SplinesComponent = Segment->GetOuterULandscapeSplinesComponent();
 		SplinesComponent->Modify();
@@ -575,7 +570,7 @@ public:
 
 	void SnapControlPointToGround(ULandscapeSplineControlPoint* ControlPoint)
 	{
-		FScopedTransaction Transaction( LOCTEXT("LandscapeSpline_SnapToGround", "Snap Landscape Spline to Ground") );
+		FScopedTransaction Transaction(LOCTEXT("LandscapeSpline_SnapToGround", "Snap Landscape Spline to Ground"));
 
 		ULandscapeSplinesComponent* SplinesComponent = ControlPoint->GetOuterULandscapeSplinesComponent();
 		SplinesComponent->Modify();
@@ -597,7 +592,7 @@ public:
 		}
 	}
 
-	void ShowSplineProperties() 
+	void ShowSplineProperties()
 	{
 		TArray<UObject*> Objects;
 		Objects.Reset(SelectedSplineControlPoints.Num() + SelectedSplineSegments.Num());
@@ -611,8 +606,8 @@ public:
 			Objects.Add(Segment);
 		}
 
-		FPropertyEditorModule& PropertyModule = FModuleManager::Get().LoadModuleChecked<FPropertyEditorModule>( TEXT("PropertyEditor") );
-		if( !PropertyModule.HasUnlockedDetailViews() )
+		FPropertyEditorModule& PropertyModule = FModuleManager::Get().LoadModuleChecked<FPropertyEditorModule>(TEXT("PropertyEditor"));
+		if (!PropertyModule.HasUnlockedDetailViews())
 		{
 			PropertyModule.CreateFloatingDetailsView(Objects, true);
 		}
@@ -622,7 +617,7 @@ public:
 		}
 	}
 
-	virtual bool BeginTool( FEditorViewportClient* ViewportClient, const FLandscapeToolTarget& InTarget, const FVector& InHitLocation) override
+	virtual bool BeginTool(FEditorViewportClient* ViewportClient, const FLandscapeToolTarget& InTarget, const FVector& InHitLocation) override
 	{
 		LandscapeInfo = InTarget.LandscapeInfo.Get();
 
@@ -651,7 +646,7 @@ public:
 		LandscapeInfo = NULL;
 	}
 
-	virtual bool MouseMove( FEditorViewportClient* ViewportClient, FViewport* Viewport, int32 x, int32 y ) override
+	virtual bool MouseMove(FEditorViewportClient* ViewportClient, FViewport* Viewport, int32 x, int32 y) override
 	{
 		FVector HitLocation;
 		if (EdMode->LandscapeMouseTrace(ViewportClient, x, y, HitLocation))
@@ -666,13 +661,13 @@ public:
 		return true;
 	}
 
-	virtual void ApplyTool( FEditorViewportClient* ViewportClient )
+	virtual void ApplyTool(FEditorViewportClient* ViewportClient)
 	{
 	}
 
-	virtual bool HandleClick( HHitProxy* HitProxy, const FViewportClick& Click ) override
+	virtual bool HandleClick(HHitProxy* HitProxy, const FViewportClick& Click) override
 	{
-		if ( (!HitProxy || !HitProxy->IsA(HWidgetAxis::StaticGetType()))
+		if ((!HitProxy || !HitProxy->IsA(HWidgetAxis::StaticGetType()))
 			&& !Click.IsShiftDown())
 		{
 			ClearSelection();
@@ -685,17 +680,17 @@ public:
 			ULandscapeSplineControlPoint* ClickedControlPoint = NULL;
 			ULandscapeSplineSegment* ClickedSplineSegment = NULL;
 
-			if (HitProxy->IsA(HLandscapeSplineProxy_ControlPoint::StaticGetType()) )
+			if (HitProxy->IsA(HLandscapeSplineProxy_ControlPoint::StaticGetType()))
 			{
 				HLandscapeSplineProxy_ControlPoint* SplineProxy = (HLandscapeSplineProxy_ControlPoint*)HitProxy;
 				ClickedControlPoint = SplineProxy->ControlPoint;
 			}
-			else if (HitProxy->IsA(HLandscapeSplineProxy_Segment::StaticGetType()) )
+			else if (HitProxy->IsA(HLandscapeSplineProxy_Segment::StaticGetType()))
 			{
 				HLandscapeSplineProxy_Segment* SplineProxy = (HLandscapeSplineProxy_Segment*)HitProxy;
 				ClickedSplineSegment = SplineProxy->SplineSegment;
 			}
-			else if (HitProxy->IsA(HActor::StaticGetType()) )
+			else if (HitProxy->IsA(HActor::StaticGetType()))
 			{
 				HActor* ActorProxy = (HActor*)HitProxy;
 				AActor* Actor = ActorProxy->Actor;
@@ -765,7 +760,7 @@ public:
 		return false;
 	}
 
-	virtual bool InputKey( FEditorViewportClient* InViewportClient, FViewport* InViewport, FKey InKey, EInputEvent InEvent ) override
+	virtual bool InputKey(FEditorViewportClient* InViewportClient, FViewport* InViewport, FKey InKey, EInputEvent InEvent) override
 	{
 		if (InKey == EKeys::F4 && InEvent == IE_Pressed)
 		{
@@ -780,7 +775,7 @@ public:
 		{
 			if (SelectedSplineControlPoints.Num() > 0 || SelectedSplineSegments.Num() > 0)
 			{
-				FScopedTransaction Transaction( LOCTEXT("LandscapeSpline_AutoRotate", "Auto-rotate Landscape Spline Control Points") );
+				FScopedTransaction Transaction(LOCTEXT("LandscapeSpline_AutoRotate", "Auto-rotate Landscape Spline Control Points"));
 
 				for (ULandscapeSplineControlPoint* ControlPoint : SelectedSplineControlPoints)
 				{
@@ -804,7 +799,7 @@ public:
 		{
 			if (SelectedSplineSegments.Num() > 0)
 			{
-				FScopedTransaction Transaction( LOCTEXT("LandscapeSpline_FlipSegments", "Flip Landscape Spline Segments") );
+				FScopedTransaction Transaction(LOCTEXT("LandscapeSpline_FlipSegments", "Flip Landscape Spline Segments"));
 
 				for (ULandscapeSplineSegment* Segment : SelectedSplineSegments)
 				{
@@ -819,7 +814,7 @@ public:
 		{
 			if (SelectedSplineControlPoints.Num() > 0 || SelectedSplineSegments.Num() > 0)
 			{
-				FScopedTransaction Transaction( LOCTEXT("LandscapeSpline_AutoFlipTangents", "Auto-flip Landscape Spline Tangents") );
+				FScopedTransaction Transaction(LOCTEXT("LandscapeSpline_AutoFlipTangents", "Auto-flip Landscape Spline Tangents"));
 
 				for (ULandscapeSplineControlPoint* ControlPoint : SelectedSplineControlPoints)
 				{
@@ -843,7 +838,7 @@ public:
 		{
 			if (SelectedSplineControlPoints.Num() > 0 || SelectedSplineSegments.Num() > 0)
 			{
-				FScopedTransaction Transaction( LOCTEXT("LandscapeSpline_SnapToGround", "Snap Landscape Spline to Ground") );
+				FScopedTransaction Transaction(LOCTEXT("LandscapeSpline_SnapToGround", "Snap Landscape Spline to Ground"));
 
 				for (ULandscapeSplineControlPoint* ControlPoint : SelectedSplineControlPoints)
 				{
@@ -887,12 +882,12 @@ public:
 				{
 					ULandscapeSplineControlPoint* ClickedControlPoint = NULL;
 
-					if (HitProxy->IsA(HLandscapeSplineProxy_ControlPoint::StaticGetType()) )
+					if (HitProxy->IsA(HLandscapeSplineProxy_ControlPoint::StaticGetType()))
 					{
 						HLandscapeSplineProxy_ControlPoint* SplineProxy = (HLandscapeSplineProxy_ControlPoint*)HitProxy;
 						ClickedControlPoint = SplineProxy->ControlPoint;
 					}
-					else if (HitProxy->IsA(HActor::StaticGetType()) )
+					else if (HitProxy->IsA(HActor::StaticGetType()))
 					{
 						HActor* ActorProxy = (HActor*)HitProxy;
 						ULandscapeSplinesComponent* SplineComponent = ActorProxy->Actor->FindComponentByClass<ULandscapeSplinesComponent>();
@@ -915,7 +910,7 @@ public:
 
 					if (ClickedControlPoint != NULL)
 					{
-						FScopedTransaction Transaction( LOCTEXT("LandscapeSpline_AddSegment", "Add Landscape Spline Segment") );
+						FScopedTransaction Transaction(LOCTEXT("LandscapeSpline_AddSegment", "Add Landscape Spline Segment"));
 
 						for (ULandscapeSplineControlPoint* ControlPoint : SelectedSplineControlPoints)
 						{
@@ -943,14 +938,14 @@ public:
 					ULandscapeSplineSegment* ClickedSplineSegment = NULL;
 					FTransform LandscapeToSpline;
 
-					if (HitProxy->IsA(HLandscapeSplineProxy_Segment::StaticGetType()) )
+					if (HitProxy->IsA(HLandscapeSplineProxy_Segment::StaticGetType()))
 					{
 						HLandscapeSplineProxy_Segment* SplineProxy = (HLandscapeSplineProxy_Segment*)HitProxy;
 						ClickedSplineSegment = SplineProxy->SplineSegment;
 
 						LandscapeToSpline = ClickedSplineSegment->GetTypedOuter<AActor>()->ActorToWorld().GetRelativeTransform(ClickedSplineSegment->GetTypedOuter<ULandscapeSplinesComponent>()->ComponentToWorld);
 					}
-					else if (HitProxy->IsA(HActor::StaticGetType()) )
+					else if (HitProxy->IsA(HActor::StaticGetType()))
 					{
 						HActor* ActorProxy = (HActor*)HitProxy;
 						const USplineMeshComponent* SplineMeshComponent = Cast<const USplineMeshComponent>(ActorProxy->PrimComponent);
@@ -1001,12 +996,12 @@ public:
 				HHitProxy*	HitProxy = InViewport->GetHitProxy(HitX, HitY);
 				if (HitProxy)
 				{
-					if (HitProxy->IsA(HWidgetAxis::StaticGetType()) )
+					if (HitProxy->IsA(HWidgetAxis::StaticGetType()))
 					{
 						checkSlow(SelectedSplineControlPoints.Num() > 0);
 						bMovingControlPoint = true;
 
-						GEditor->BeginTransaction( LOCTEXT("LandscapeSpline_ModifyControlPoint", "Modify Landscape Spline Control Point") );
+						GEditor->BeginTransaction(LOCTEXT("LandscapeSpline_ModifyControlPoint", "Modify Landscape Spline Control Point"));
 						for (ULandscapeSplineControlPoint* ControlPoint : SelectedSplineControlPoints)
 						{
 							ControlPoint->Modify();
@@ -1015,13 +1010,13 @@ public:
 
 						return false; // We're not actually handling this case ourselves, just wrapping it in a transaction
 					}
-					else if (HitProxy->IsA(HLandscapeSplineProxy_Tangent::StaticGetType()) )
+					else if (HitProxy->IsA(HLandscapeSplineProxy_Tangent::StaticGetType()))
 					{
 						HLandscapeSplineProxy_Tangent* SplineProxy = (HLandscapeSplineProxy_Tangent*)HitProxy;
 						DraggingTangent_Segment = SplineProxy->SplineSegment;
 						DraggingTangent_End = SplineProxy->End;
 
-						GEditor->BeginTransaction( LOCTEXT("LandscapeSpline_ModifyTangent", "Modify Landscape Spline Tangent") );
+						GEditor->BeginTransaction(LOCTEXT("LandscapeSpline_ModifyTangent", "Modify Landscape Spline Tangent"));
 						ULandscapeSplinesComponent* SplinesComponent = DraggingTangent_Segment->GetOuterULandscapeSplinesComponent();
 						SplinesComponent->Modify();
 						DraggingTangent_Segment->Modify();
@@ -1061,7 +1056,7 @@ public:
 		return false;
 	}
 
-	virtual bool InputDelta( FEditorViewportClient* InViewportClient, FViewport* InViewport, FVector& InDrag, FRotator& InRot, FVector& InScale ) override
+	virtual bool InputDelta(FEditorViewportClient* InViewportClient, FViewport* InViewport, FVector& InDrag, FRotator& InRot, FVector& InScale) override
 	{
 		FVector Drag = InDrag;
 
@@ -1138,7 +1133,7 @@ public:
 		SelectedSplineControlPoints.Empty();
 		SelectedSplineSegments.Empty();
 
-		if (EdMode->CurrentToolSet != NULL && EdMode->CurrentToolSet->GetTool() == this)
+		if (EdMode->CurrentTool != NULL && EdMode->CurrentTool == this)
 		{
 			for (const FLandscapeListInfo& Info : EdMode->GetLandscapeList())
 			{
@@ -1278,9 +1273,9 @@ public:
 					FVector HandlePos = SplinesComponent->ComponentToWorld.TransformPosition(StartLocation + StartRotation.Vector() * Connection.GetNearConnection().TangentLen / 2);
 					PDI->DrawLine(StartPos, HandlePos, FColor::White, SDPG_Foreground);
 
-					if (PDI->IsHitTesting()) PDI->SetHitProxy( new HLandscapeSplineProxy_Tangent(Connection.Segment, Connection.End) );
-					PDI->DrawPoint(HandlePos, FColor(255,255,255), 10.f, SDPG_Foreground);
-					if (PDI->IsHitTesting()) PDI->SetHitProxy( NULL );
+					if (PDI->IsHitTesting()) PDI->SetHitProxy(new HLandscapeSplineProxy_Tangent(Connection.Segment, Connection.End));
+					PDI->DrawPoint(HandlePos, FColor(255, 255, 255), 10.f, SDPG_Foreground);
+					if (PDI->IsHitTesting()) PDI->SetHitProxy(NULL);
 				}
 			}
 		}
@@ -1301,9 +1296,9 @@ public:
 					FVector EndHandlePos = SplinesComponent->ComponentToWorld.TransformPosition(StartLocation + StartRotation.Vector() * Connection.TangentLen / 2);
 
 					PDI->DrawLine(EndPos, EndHandlePos, FColor::White, SDPG_Foreground);
-					if (PDI->IsHitTesting()) PDI->SetHitProxy( new HLandscapeSplineProxy_Tangent(Segment, !!End) );
-					PDI->DrawPoint(EndHandlePos, FColor(255,255,255), 10.f, SDPG_Foreground);
-					if (PDI->IsHitTesting()) PDI->SetHitProxy( NULL );
+					if (PDI->IsHitTesting()) PDI->SetHitProxy(new HLandscapeSplineProxy_Tangent(Segment, !!End));
+					PDI->DrawPoint(EndHandlePos, FColor(255, 255, 255), 10.f, SDPG_Foreground);
+					if (PDI->IsHitTesting()) PDI->SetHitProxy(NULL);
 				}
 			}
 		}
@@ -1314,7 +1309,7 @@ public:
 		return true;
 	}
 
-	virtual bool IsSelectionAllowed( AActor* InActor, bool bInSelection ) const override
+	virtual bool IsSelectionAllowed(AActor* InActor, bool bInSelection) const override
 	{
 		// Only filter selection not deselection
 		if (bInSelection)
@@ -1468,7 +1463,7 @@ public:
 	{
 		if (SelectedSplineControlPoints.Num() > 0 || SelectedSplineSegments.Num() > 0)
 		{
-			FScopedTransaction Transaction( LOCTEXT("LandscapeSpline_Duplicate", "Duplicate Landscape Splines") );
+			FScopedTransaction Transaction(LOCTEXT("LandscapeSpline_Duplicate", "Duplicate Landscape Splines"));
 
 			FString Data;
 			InternalProcessEditCopy(&Data);
@@ -1480,7 +1475,7 @@ public:
 	{
 		if (SelectedSplineControlPoints.Num() > 0 || SelectedSplineSegments.Num() > 0)
 		{
-			FScopedTransaction Transaction( LOCTEXT("LandscapeSpline_Delete", "Delete Landscape Splines") );
+			FScopedTransaction Transaction(LOCTEXT("LandscapeSpline_Delete", "Delete Landscape Splines"));
 
 			for (ULandscapeSplineControlPoint* ControlPoint : SelectedSplineControlPoints)
 			{
@@ -1501,7 +1496,7 @@ public:
 	{
 		if (SelectedSplineControlPoints.Num() > 0 || SelectedSplineSegments.Num() > 0)
 		{
-			FScopedTransaction Transaction( LOCTEXT("LandscapeSpline_Cut", "Cut Landscape Splines") );
+			FScopedTransaction Transaction(LOCTEXT("LandscapeSpline_Cut", "Cut Landscape Splines"));
 
 			InternalProcessEditCopy();
 			InternalProcessEditDelete();
@@ -1552,7 +1547,7 @@ public:
 
 	void InternalProcessEditPaste(FString* InData = NULL, bool bOffset = false)
 	{
-		FScopedTransaction Transaction( LOCTEXT("LandscapeSpline_Paste", "Paste Landscape Splines") );
+		FScopedTransaction Transaction(LOCTEXT("LandscapeSpline_Paste", "Paste Landscape Splines"));
 
 		ALandscapeProxy* Landscape = EdMode->CurrentToolTarget.LandscapeInfo->GetCurrentLevelLandscapeProxy(true);
 		if (!Landscape)
@@ -1602,68 +1597,57 @@ protected:
 	// End of FEditorUndoClient
 
 protected:
-	class FEdModeLandscape* EdMode;
-	class ULandscapeInfo* LandscapeInfo;
+	FEdModeLandscape* EdMode;
+	ULandscapeInfo* LandscapeInfo;
 
 	TSet<ULandscapeSplineControlPoint*> SelectedSplineControlPoints;
 	TSet<ULandscapeSplineSegment*> SelectedSplineSegments;
 
 	ULandscapeSplineSegment* DraggingTangent_Segment;
-	uint32 DraggingTangent_End:1;
+	uint32 DraggingTangent_End : 1;
 
-	uint32 bMovingControlPoint:1;
+	uint32 bMovingControlPoint : 1;
 
-	uint32 bAutoRotateOnJoin:1;
-	uint32 bAutoChangeConnectionsOnMove:1;
-	uint32 bDeleteLooseEnds:1;
-	uint32 bCopyMeshToNewControlPoint:1;
+	uint32 bAutoRotateOnJoin : 1;
+	uint32 bAutoChangeConnectionsOnMove : 1;
+	uint32 bDeleteLooseEnds : 1;
+	uint32 bCopyMeshToNewControlPoint : 1;
 
-	friend class FEdModeLandscape;
+	friend FEdModeLandscape;
 };
 
 
 void FEdModeLandscape::ShowSplineProperties()
 {
-	if (SplinesToolSet /*&& SplinesToolSet == CurrentToolSet*/ )
+	if (SplinesTool /*&& SplinesTool == CurrentTool*/)
 	{
-		if (SplinesToolSet->SetToolForTarget(CurrentToolTarget) && SplinesToolSet->GetTool())
-		{
-			((FLandscapeToolSplines*)SplinesToolSet->GetTool())->ShowSplineProperties();
-		}
+		SplinesTool->ShowSplineProperties();
 	}
 }
 
 void FEdModeLandscape::SelectAllConnectedSplineControlPoints()
 {
-	if (SplinesToolSet /*&& SplinesToolSet == CurrentToolSet*/ )
+	if (SplinesTool /*&& SplinesTool == CurrentTool*/)
 	{
-		if (SplinesToolSet->SetToolForTarget(CurrentToolTarget) && SplinesToolSet->GetTool())
-		{
-			FLandscapeToolSplines* SplineTool = ((FLandscapeToolSplines*)SplinesToolSet->GetTool());
-			SplineTool->SelectAdjacentControlPoints();
-			SplineTool->ClearSelectedSegments();
-			SplineTool->SelectConnected();
+		SplinesTool->SelectAdjacentControlPoints();
+		SplinesTool->ClearSelectedSegments();
+		SplinesTool->SelectConnected();
 
-			SplineTool->UpdatePropertiesWindows();
-			GUnrealEd->RedrawLevelEditingViewports();
-		}
+		SplinesTool->UpdatePropertiesWindows();
+		GUnrealEd->RedrawLevelEditingViewports();
 	}
 }
 
 void FEdModeLandscape::SelectAllConnectedSplineSegments()
 {
-	if (SplinesToolSet /*&& SplinesToolSet == CurrentToolSet*/ )
+	if (SplinesTool /*&& SplinesTool == CurrentTool*/)
 	{
-		if (SplinesToolSet->SetToolForTarget(CurrentToolTarget) && SplinesToolSet->GetTool())
-		{
-			FLandscapeToolSplines* SplineTool = ((FLandscapeToolSplines*)SplinesToolSet->GetTool());
-			SplineTool->SelectAdjacentSegments();
-			SplineTool->ClearSelectedControlPoints();
-			SplineTool->SelectConnected();
+		SplinesTool->SelectAdjacentSegments();
+		SplinesTool->ClearSelectedControlPoints();
+		SplinesTool->SelectConnected();
 
-			SplineTool->UpdatePropertiesWindows();
-			GUnrealEd->RedrawLevelEditingViewports();
-		}
+		SplinesTool->UpdatePropertiesWindows();
+		GUnrealEd->RedrawLevelEditingViewports();
 	}
 }
 
@@ -1671,186 +1655,173 @@ void FEdModeLandscape::SplineMoveToCurrentLevel()
 {
 	FScopedTransaction Transaction(LOCTEXT("LandscapeSpline_AddControlPoint", "Add Landscape Spline Control Point"));
 
-	if (SplinesToolSet /*&& SplinesToolSet == CurrentToolSet*/)
+	if (SplinesTool /*&& SplinesTool == CurrentTool*/)
 	{
-		if (SplinesToolSet->SetToolForTarget(CurrentToolTarget) && SplinesToolSet->GetTool())
-		{
-			// Select all connected control points
-			FLandscapeToolSplines* SplineTool = ((FLandscapeToolSplines*)SplinesToolSet->GetTool());
-			SplineTool->SelectAdjacentSegments();
-			SplineTool->SelectAdjacentControlPoints();
-			SplineTool->SelectConnected();
+		// Select all connected control points
+		SplinesTool->SelectAdjacentSegments();
+		SplinesTool->SelectAdjacentControlPoints();
+		SplinesTool->SelectConnected();
 
-			TSet<ALandscapeProxy*> FromProxies;
-			ALandscapeProxy* Landscape = NULL;
+		TSet<ALandscapeProxy*> FromProxies;
+		ALandscapeProxy* Landscape = NULL;
 
-			for (ULandscapeSplineControlPoint* ControlPoint : SplineTool->SelectedSplineControlPoints)
+		for (ULandscapeSplineControlPoint* ControlPoint : SplinesTool->SelectedSplineControlPoints)
+		{
+			ULandscapeSplinesComponent* LandscapeSplinesComp = ControlPoint->GetOuterULandscapeSplinesComponent();
+			ALandscapeProxy* FromProxy = LandscapeSplinesComp ? Cast<ALandscapeProxy>(LandscapeSplinesComp->GetOuter()) : NULL;
+			if (FromProxy)
 			{
-				ULandscapeSplinesComponent* LandscapeSplinesComp = ControlPoint->GetOuterULandscapeSplinesComponent();
-				ALandscapeProxy* FromProxy = LandscapeSplinesComp ? Cast<ALandscapeProxy>(LandscapeSplinesComp->GetOuter()) : NULL;
-				if (FromProxy)
+				if (!Landscape)
 				{
+					ULandscapeInfo* LandscapeInfo = FromProxy->GetLandscapeInfo(false);
+					check(LandscapeInfo);
+					Landscape = LandscapeInfo->GetCurrentLevelLandscapeProxy(true);
 					if (!Landscape)
 					{
-						ULandscapeInfo* LandscapeInfo = FromProxy->GetLandscapeInfo(false);
-						check(LandscapeInfo);
-						Landscape = LandscapeInfo->GetCurrentLevelLandscapeProxy(true);
-						if (!Landscape)
-						{
-							// No Landscape Proxy, don't support for creating only for Spline now
-							return;
-						}
+						// No Landscape Proxy, don't support for creating only for Spline now
+						return;
 					}
+				}
 
-					if (Landscape != FromProxy)
+				if (Landscape != FromProxy)
+				{
+					Landscape->Modify();
+					if (Landscape->SplineComponent == NULL)
 					{
-						Landscape->Modify();
-						if (Landscape->SplineComponent == NULL)
-						{
-							SplineTool->CreateSplineComponent(Landscape, FromProxy->SplineComponent->RelativeScale3D);
-						}
-
-						const FTransform OldToNewTransform =
-							Landscape->SplineComponent->ComponentToWorld.GetRelativeTransform(FromProxy->SplineComponent->ComponentToWorld);
+						SplinesTool->CreateSplineComponent(Landscape, FromProxy->SplineComponent->RelativeScale3D);
+					}
 
-						Landscape->SplineComponent->Modify();
-						if (FromProxies.Find(FromProxy) == NULL)
-						{
-							FromProxies.Add(FromProxy);
-							FromProxy->Modify();
-							FromProxy->SplineComponent->Modify();
-							FromProxy->SplineComponent->MarkRenderStateDirty();
-						}
+					const FTransform OldToNewTransform =
+						Landscape->SplineComponent->ComponentToWorld.GetRelativeTransform(FromProxy->SplineComponent->ComponentToWorld);
 
-						// Add/Remove Control Point/Segment/MeshComponents
-						FromProxy->SplineComponent->ControlPoints.Remove(ControlPoint);
-						ControlPoint->Rename(NULL, Landscape->SplineComponent);
-						Landscape->SplineComponent->ControlPoints.Add(ControlPoint);
+					Landscape->SplineComponent->Modify();
+					if (FromProxies.Find(FromProxy) == NULL)
+					{
+						FromProxies.Add(FromProxy);
+						FromProxy->Modify();
+						FromProxy->SplineComponent->Modify();
+						FromProxy->SplineComponent->MarkRenderStateDirty();
+					}
 
-						ControlPoint->Location = OldToNewTransform.TransformPosition(ControlPoint->Location);
+					// Add/Remove Control Point/Segment/MeshComponents
+					FromProxy->SplineComponent->ControlPoints.Remove(ControlPoint);
+					ControlPoint->Rename(NULL, Landscape->SplineComponent);
+					Landscape->SplineComponent->ControlPoints.Add(ControlPoint);
 
-						if (ControlPoint->MeshComponent)
-						{
-							UControlPointMeshComponent* ControlPointMeshComp = ControlPoint->MeshComponent;
-							ControlPointMeshComp->Modify();
-							ControlPointMeshComp->UnregisterComponent();
-							ControlPointMeshComp->DetachFromParent(true);
-							ControlPointMeshComp->InvalidateLightingCache();
-							ControlPointMeshComp->Rename(NULL, Landscape);
-							ControlPointMeshComp->AttachTo(Landscape->SplineComponent, NAME_None, EAttachLocation::KeepWorldPosition);
-						}
+					ControlPoint->Location = OldToNewTransform.TransformPosition(ControlPoint->Location);
 
-						ControlPoint->UpdateSplinePoints(true, false);
+					if (ControlPoint->MeshComponent)
+					{
+						UControlPointMeshComponent* ControlPointMeshComp = ControlPoint->MeshComponent;
+						ControlPointMeshComp->Modify();
+						ControlPointMeshComp->UnregisterComponent();
+						ControlPointMeshComp->DetachFromParent(true);
+						ControlPointMeshComp->InvalidateLightingCache();
+						ControlPointMeshComp->Rename(NULL, Landscape);
+						ControlPointMeshComp->AttachTo(Landscape->SplineComponent, NAME_None, EAttachLocation::KeepWorldPosition);
 					}
+
+					ControlPoint->UpdateSplinePoints(true, false);
 				}
 			}
+		}
 
-			for (ULandscapeSplineSegment* Segment : SplineTool->SelectedSplineSegments)
+		for (ULandscapeSplineSegment* Segment : SplinesTool->SelectedSplineSegments)
+		{
+			ULandscapeSplinesComponent* LandscapeSplinesComp = Segment->GetOuterULandscapeSplinesComponent();
+			ALandscapeProxy* FromProxy = LandscapeSplinesComp ? Cast<ALandscapeProxy>(LandscapeSplinesComp->GetOuter()) : NULL;
+			if (FromProxy)
 			{
-				ULandscapeSplinesComponent* LandscapeSplinesComp = Segment->GetOuterULandscapeSplinesComponent();
-				ALandscapeProxy* FromProxy = LandscapeSplinesComp ? Cast<ALandscapeProxy>(LandscapeSplinesComp->GetOuter()) : NULL;
-				if (FromProxy)
+				if (!Landscape)
 				{
+					ULandscapeInfo* LandscapeInfo = FromProxy->GetLandscapeInfo(false);
+					check(LandscapeInfo);
+					Landscape = LandscapeInfo->GetCurrentLevelLandscapeProxy(true);
 					if (!Landscape)
 					{
-						ULandscapeInfo* LandscapeInfo = FromProxy->GetLandscapeInfo(false);
-						check(LandscapeInfo);
-						Landscape = LandscapeInfo->GetCurrentLevelLandscapeProxy(true);
-						if (!Landscape)
-						{
-							// No Landscape Proxy, don't support for creating only for Spline now
-							return;
-						}
+						// No Landscape Proxy, don't support for creating only for Spline now
+						return;
 					}
+				}
 
-					if (Landscape != FromProxy)
+				if (Landscape != FromProxy)
+				{
+					Landscape->Modify();
+					if (Landscape->SplineComponent == NULL)
 					{
-						Landscape->Modify();
-						if (Landscape->SplineComponent == NULL)
-						{
-							SplineTool->CreateSplineComponent(Landscape, FromProxy->SplineComponent->RelativeScale3D);
-						}
-
-						Landscape->SplineComponent->Modify();
-						if (FromProxies.Find(FromProxy) == NULL)
-						{
-							FromProxies.Add(FromProxy);
-							FromProxy->Modify();
-							FromProxy->SplineComponent->Modify();
-							FromProxy->SplineComponent->MarkRenderStateDirty();
-						}
+						SplinesTool->CreateSplineComponent(Landscape, FromProxy->SplineComponent->RelativeScale3D);
+					}
 
-						// Add/Remove Control Point/Segment/MeshComponents
-						FromProxy->SplineComponent->Segments.Remove(Segment);
-						Segment->Rename(NULL, Landscape->SplineComponent);
-						Landscape->SplineComponent->Segments.Add(Segment);
+					Landscape->SplineComponent->Modify();
+					if (FromProxies.Find(FromProxy) == NULL)
+					{
+						FromProxies.Add(FromProxy);
+						FromProxy->Modify();
+						FromProxy->SplineComponent->Modify();
+						FromProxy->SplineComponent->MarkRenderStateDirty();
+					}
 
-						for (USplineMeshComponent* SplineMeshComp : Segment->MeshComponents)
-						{
-							SplineMeshComp->Modify();
-							SplineMeshComp->UnregisterComponent();
-							SplineMeshComp->DetachFromParent(true);
-							SplineMeshComp->InvalidateLightingCache();
-							SplineMeshComp->Rename(NULL, Landscape);
-							SplineMeshComp->AttachTo(Landscape->SplineComponent, NAME_None, EAttachLocation::KeepWorldPosition);
-						}
+					// Add/Remove Control Point/Segment/MeshComponents
+					FromProxy->SplineComponent->Segments.Remove(Segment);
+					Segment->Rename(NULL, Landscape->SplineComponent);
+					Landscape->SplineComponent->Segments.Add(Segment);
 
-						Segment->UpdateSplinePoints();
+					for (USplineMeshComponent* SplineMeshComp : Segment->MeshComponents)
+					{
+						SplineMeshComp->Modify();
+						SplineMeshComp->UnregisterComponent();
+						SplineMeshComp->DetachFromParent(true);
+						SplineMeshComp->InvalidateLightingCache();
+						SplineMeshComp->Rename(NULL, Landscape);
+						SplineMeshComp->AttachTo(Landscape->SplineComponent, NAME_None, EAttachLocation::KeepWorldPosition);
 					}
+
+					Segment->UpdateSplinePoints();
 				}
 			}
+		}
 
-			if (Landscape && Landscape->SplineComponent)
+		if (Landscape && Landscape->SplineComponent)
+		{
+			if (!Landscape->SplineComponent->IsRegistered())
 			{
-				if (!Landscape->SplineComponent->IsRegistered())
-				{
-					Landscape->SplineComponent->RegisterComponent();
-				}
-				else
-				{
-					Landscape->SplineComponent->MarkRenderStateDirty();
-				}
+				Landscape->SplineComponent->RegisterComponent();
+			}
+			else
+			{
+				Landscape->SplineComponent->MarkRenderStateDirty();
 			}
-
-			SplineTool->ClearSelection();
-			SplineTool->UpdatePropertiesWindows();
-			GUnrealEd->RedrawLevelEditingViewports();
 		}
+
+		SplinesTool->ClearSelection();
+		SplinesTool->UpdatePropertiesWindows();
+		GUnrealEd->RedrawLevelEditingViewports();
 	}
 }
 
 void FEdModeLandscape::SetbUseAutoRotateOnJoin(bool InbAutoRotateOnJoin)
 {
-	if (SplinesToolSet /*&& SplinesToolSet == CurrentToolSet*/)
+	if (SplinesTool /*&& SplinesTool == CurrentTool*/)
 	{
-		if (SplinesToolSet->SetToolForTarget(CurrentToolTarget) && SplinesToolSet->GetTool())
-		{
-			FLandscapeToolSplines* SplineTool = ((FLandscapeToolSplines*)SplinesToolSet->GetTool());
-			SplineTool->bAutoRotateOnJoin = InbAutoRotateOnJoin;
-		}
+		SplinesTool->bAutoRotateOnJoin = InbAutoRotateOnJoin;
 	}
 }
 
 bool FEdModeLandscape::GetbUseAutoRotateOnJoin()
 {
-	if (SplinesToolSet /*&& SplinesToolSet == CurrentToolSet*/)
+	if (SplinesTool /*&& SplinesTool == CurrentTool*/)
 	{
-		if (SplinesToolSet->SetToolForTarget(CurrentToolTarget) && SplinesToolSet->GetTool())
-		{
-			FLandscapeToolSplines* SplineTool = ((FLandscapeToolSplines*)SplinesToolSet->GetTool());
-			return SplineTool->bAutoRotateOnJoin;
-		}
+		return SplinesTool->bAutoRotateOnJoin;
 	}
 	return true; // default value
 }
 
-void FEdModeLandscape::IntializeToolSet_Splines()
+void FEdModeLandscape::InitializeTool_Splines()
 {
-	FLandscapeToolSet* ToolSet_Splines = new(LandscapeToolSets) FLandscapeToolSet(TEXT("ToolSet_Splines"));
-	ToolSet_Splines->AddTool(new FLandscapeToolSplines(this));
-	SplinesToolSet = ToolSet_Splines;
-
-	ToolSet_Splines->ValidBrushes.Add("BrushSet_Splines");
+	auto Tool_Splines = MakeUnique<FLandscapeToolSplines>(this);
+	Tool_Splines->ValidBrushes.Add("BrushSet_Splines");
+	SplinesTool = Tool_Splines.Get();
+	LandscapeTools.Add(MoveTemp(Tool_Splines));
 }
 
 #undef LOCTEXT_NAMESPACE
