Remove FLandscapeToolSet

diff --git a/Engine/Source/Editor/LandscapeEditor/Private/LandscapeEdModeTools.h b/Engine/Source/Editor/LandscapeEditor/Private/LandscapeEdModeTools.h
--- a/Engine/Source/Editor/LandscapeEditor/Private/LandscapeEdModeTools.h
+++ b/Engine/Source/Editor/LandscapeEditor/Private/LandscapeEdModeTools.h
@@ -18,26 +18,26 @@ struct FNoiseParameter
 	// Constructors.
 
 	FNoiseParameter() {}
-	FNoiseParameter(float InBase,float InScale,float InAmount):
-	Base(InBase),
+	FNoiseParameter(float InBase, float InScale, float InAmount) :
+		Base(InBase),
 		NoiseScale(InScale),
 		NoiseAmount(InAmount)
 	{}
 
 	// Sample
-	float Sample(int32 X,int32 Y) const
+	float Sample(int32 X, int32 Y) const
 	{
 		float	Noise = 0.0f;
 		X = FMath::Abs(X);
 		Y = FMath::Abs(Y);
 
-		if(NoiseScale > DELTA)
+		if (NoiseScale > DELTA)
 		{
-			for(uint32 Octave = 0;Octave < 4;Octave++)
+			for (uint32 Octave = 0; Octave < 4; Octave++)
 			{
 				float	OctaveShift = 1 << Octave;
 				float	OctaveScale = OctaveShift / NoiseScale;
-				Noise += PerlinNoise2D(X * OctaveScale,Y * OctaveScale) / OctaveShift;
+				Noise += PerlinNoise2D(X * OctaveScale, Y * OctaveScale) / OctaveShift;
 			}
 		}
 
@@ -45,27 +45,27 @@ struct FNoiseParameter
 	}
 
 	// TestGreater - Returns 1 if TestValue is greater than the parameter.
-	bool TestGreater(int32 X,int32 Y,float TestValue) const
+	bool TestGreater(int32 X, int32 Y, float TestValue) const
 	{
 		float	ParameterValue = Base;
 
-		if(NoiseScale > DELTA)
+		if (NoiseScale > DELTA)
 		{
-			for(uint32 Octave = 0;Octave < 4;Octave++)
+			for (uint32 Octave = 0; Octave < 4; Octave++)
 			{
 				float	OctaveShift = 1 << Octave;
 				float	OctaveAmplitude = NoiseAmount / OctaveShift;
 
 				// Attempt to avoid calculating noise if the test value is outside of the noise amplitude.
 
-				if(TestValue > ParameterValue + OctaveAmplitude)
+				if (TestValue > ParameterValue + OctaveAmplitude)
 					return 1;
-				else if(TestValue < ParameterValue - OctaveAmplitude)
+				else if (TestValue < ParameterValue - OctaveAmplitude)
 					return 0;
 				else
 				{
 					float	OctaveScale = OctaveShift / NoiseScale;
-					ParameterValue += PerlinNoise2D(X * OctaveScale,Y * OctaveScale) * OctaveAmplitude;
+					ParameterValue += PerlinNoise2D(X * OctaveScale, Y * OctaveScale) * OctaveAmplitude;
 				}
 			}
 		}
@@ -74,7 +74,7 @@ struct FNoiseParameter
 	}
 
 	// TestLess
-	bool TestLess(int32 X,int32 Y,float TestValue) const { return !TestGreater(X,Y,TestValue); }
+	bool TestLess(int32 X, int32 Y, float TestValue) const { return !TestGreater(X, Y, TestValue); }
 
 private:
 	static const int32 Permutations[256];
@@ -105,7 +105,7 @@ private:
 	}
 
 
-	float Grad(int32 Hash,float X,float Y) const
+	float Grad(int32 Hash, float X, float Y) const
 	{
 		int32		H = Hash & 15;
 		float	U = H < 8 || H == 12 || H == 13 ? X : Y,
@@ -113,7 +113,7 @@ private:
 		return ((H & 1) == 0 ? U : -U) + ((H & 2) == 0 ? V : -V);
 	}
 
-	float PerlinNoise2D(float X,float Y) const
+	float PerlinNoise2D(float X, float Y) const
 	{
 		int32		TruncX = FMath::TruncToInt(X),
 			TruncY = FMath::TruncToInt(Y),
@@ -132,10 +132,10 @@ private:
 			BA = Permutations[B & 255],
 			BB = Permutations[(B + 1) & 255];
 
-		return	FMath::Lerp(	FMath::Lerp(	Grad(Permutations[AA],			FracX,	FracY	),
-			Grad(Permutations[BA],			FracX-1,FracY	),	U),
-			FMath::Lerp(	Grad(Permutations[AB],			FracX,	FracY-1	),
-			Grad(Permutations[BB],			FracX-1,FracY-1	),	U),	V);
+		return	FMath::Lerp(FMath::Lerp(Grad(Permutations[AA], FracX, FracY),
+			Grad(Permutations[BA], FracX - 1, FracY), U),
+			FMath::Lerp(Grad(Permutations[AB], FracX, FracY - 1),
+			Grad(Permutations[BB], FracX - 1, FracY - 1), U), V);
 	}
 };
 
@@ -150,30 +150,30 @@ inline void LowPassFilter(int32 X1, int32 Y1, int32 X2, int32 Y2, TMap<FIntPoint
 {
 #if WITH_KISSFFT
 	// Low-pass filter
-	int32 FFTWidth = X2-X1-1;
-	int32 FFTHeight = Y2-Y1-1;
+	int32 FFTWidth = X2 - X1 - 1;
+	int32 FFTHeight = Y2 - Y1 - 1;
 
 	const int NDims = 2;
-	const int32 Dims[NDims] = {FFTHeight-FFTHeight%2, FFTWidth-FFTWidth%2};
+	const int32 Dims[NDims] = { FFTHeight - FFTHeight % 2, FFTWidth - FFTWidth % 2 };
 	kiss_fftnd_cfg stf = kiss_fftnd_alloc(Dims, NDims, 0, NULL, NULL),
 		sti = kiss_fftnd_alloc(Dims, NDims, 1, NULL, NULL);
 
 	kiss_fft_cpx *buf = (kiss_fft_cpx *)KISS_FFT_MALLOC(sizeof(kiss_fft_cpx) * Dims[0] * Dims[1]);
 	kiss_fft_cpx *out = (kiss_fft_cpx *)KISS_FFT_MALLOC(sizeof(kiss_fft_cpx) * Dims[0] * Dims[1]);
 
-	for (int X = X1+1; X <= X2-1-FFTWidth%2; X++)
+	for (int X = X1 + 1; X <= X2 - 1 - FFTWidth % 2; X++)
 	{
-		for (int Y = Y1+1; Y <= Y2-1-FFTHeight%2; Y++)
+		for (int Y = Y1 + 1; Y <= Y2 - 1 - FFTHeight % 2; Y++)
 		{
-			buf[(X-X1-1) + (Y-Y1-1)*(Dims[1])].r = Data[(X-X1) + (Y-Y1)*(1+X2-X1)];
-			buf[(X-X1-1) + (Y-Y1-1)*(Dims[1])].i = 0;
+			buf[(X - X1 - 1) + (Y - Y1 - 1)*(Dims[1])].r = Data[(X - X1) + (Y - Y1)*(1 + X2 - X1)];
+			buf[(X - X1 - 1) + (Y - Y1 - 1)*(Dims[1])].i = 0;
 		}
 	}
 
 	// Forward FFT
 	kiss_fftnd(stf, buf, out);
 
-	int32 CenterPos[2] = {Dims[0]>>1, Dims[1]>>1};
+	int32 CenterPos[2] = { Dims[0] >> 1, Dims[1] >> 1 };
 	for (int Y = 0; Y < Dims[0]; Y++)
 	{
 		float DistFromCenter = 0.f;
@@ -189,7 +189,7 @@ inline void LowPassFilter(int32 X1, int32 Y1, int32 X2, int32 Y2, TMap<FIntPoint
 				else
 				{
 					// 2
-					DistFromCenter = (X-Dims[1])*(X-Dims[1]) + Y*Y;
+					DistFromCenter = (X - Dims[1])*(X - Dims[1]) + Y*Y;
 				}
 			}
 			else
@@ -197,20 +197,20 @@ inline void LowPassFilter(int32 X1, int32 Y1, int32 X2, int32 Y2, TMap<FIntPoint
 				if (X < CenterPos[1])
 				{
 					// 3
-					DistFromCenter = X*X + (Y-Dims[0])*(Y-Dims[0]);
+					DistFromCenter = X*X + (Y - Dims[0])*(Y - Dims[0]);
 				}
 				else
 				{
 					// 4
-					DistFromCenter = (X-Dims[1])*(X-Dims[1]) + (Y-Dims[0])*(Y-Dims[0]);
+					DistFromCenter = (X - Dims[1])*(X - Dims[1]) + (Y - Dims[0])*(Y - Dims[0]);
 				}
 			}
 			// High frequency removal
 			float Ratio = 1.f - DetailScale;
-			float Dist = FMath::Min<float>((Dims[0]*Ratio)*(Dims[0]*Ratio), (Dims[1]*Ratio)*(Dims[1]*Ratio));
-			float Filter = 1.0 / (1.0 + DistFromCenter/Dist);
-			out[X+Y*Dims[1]].r *= Filter;
-			out[X+Y*Dims[1]].i *= Filter;
+			float Dist = FMath::Min<float>((Dims[0] * Ratio)*(Dims[0] * Ratio), (Dims[1] * Ratio)*(Dims[1] * Ratio));
+			float Filter = 1.0 / (1.0 + DistFromCenter / Dist);
+			out[X + Y*Dims[1]].r *= Filter;
+			out[X + Y*Dims[1]].i *= Filter;
 		}
 	}
 
@@ -218,14 +218,14 @@ inline void LowPassFilter(int32 X1, int32 Y1, int32 X2, int32 Y2, TMap<FIntPoint
 	kiss_fftnd(sti, out, buf);
 
 	float Scale = Dims[0] * Dims[1];
-	for( auto It = BrushInfo.CreateConstIterator(); It; ++It )
+	for (auto It = BrushInfo.CreateConstIterator(); It; ++It)
 	{
 		int32 X, Y;
 		ALandscape::UnpackKey(It.Key(), X, Y);
 
 		if (It.Value() > 0.f)
 		{
-			Data[(X-X1) + (Y-Y1)*(1+X2-X1)] = FMath::Lerp((float)Data[(X-X1) + (Y-Y1)*(1+X2-X1)], buf[(X-X1-1) + (Y-Y1-1)*(Dims[1])].r / Scale, It.Value() * ApplyRatio);
+			Data[(X - X1) + (Y - Y1)*(1 + X2 - X1)] = FMath::Lerp((float)Data[(X - X1) + (Y - Y1)*(1 + X2 - X1)], buf[(X - X1 - 1) + (Y - Y1 - 1)*(Dims[1])].r / Scale, It.Value() * ApplyRatio);
 			//buf[(X-X1-1) + (Y-Y1-1)*(Dims[1])].r / Scale;
 		}
 	}
@@ -247,14 +247,14 @@ template<class Accessor, typename AccessorType>
 struct TLandscapeEditCache
 {
 	TLandscapeEditCache(Accessor& InDataAccess)
-		:	DataAccess(InDataAccess)
-		,	Valid(false)
+		: DataAccess(InDataAccess)
+		, Valid(false)
 	{
 	}
 
-	void CacheData( int32 X1, int32 Y1, int32 X2, int32 Y2 )
+	void CacheData(int32 X1, int32 Y1, int32 X2, int32 Y2)
 	{
-		if( !Valid )
+		if (!Valid)
 		{
 			if (Accessor::bUseInterp)
 			{
@@ -263,7 +263,7 @@ struct TLandscapeEditCache
 				ValidX2 = CachedX2 = X2;
 				ValidY2 = CachedY2 = Y2;
 
-				DataAccess.GetData( ValidX1, ValidY1, ValidX2, ValidY2, CachedData );
+				DataAccess.GetData(ValidX1, ValidY1, ValidX2, ValidY2, CachedData);
 				if (!ensureMsgf(ValidX1 <= ValidX2 && ValidY1 <= ValidY2, TEXT("Invalid cache area: X(%d-%d), Y(%d-%d) from region X(%d-%d), Y(%d-%d)"), ValidX1, ValidX2, ValidY1, ValidY2, X1, X2, Y1, Y2))
 				{
 					Valid = false;
@@ -277,7 +277,7 @@ struct TLandscapeEditCache
 				CachedX2 = X2;
 				CachedY2 = Y2;
 
-				DataAccess.GetDataFast( CachedX1, CachedY1, CachedX2, CachedY2, CachedData );
+				DataAccess.GetDataFast(CachedX1, CachedY1, CachedX2, CachedY2, CachedData);
 			}
 
 			OriginalData = CachedData;
@@ -287,48 +287,48 @@ struct TLandscapeEditCache
 		else
 		{
 			// Extend the cache area if needed
-			if( X1 < CachedX1 )
+			if (X1 < CachedX1)
 			{
 				if (Accessor::bUseInterp)
 				{
 					int32 x1 = X1;
 					int32 x2 = ValidX1;
-					int32 y1 = FMath::Min<int32>(Y1,CachedY1);
-					int32 y2 = FMath::Max<int32>(Y2,CachedY2);
+					int32 y1 = FMath::Min<int32>(Y1, CachedY1);
+					int32 y2 = FMath::Max<int32>(Y2, CachedY2);
 
-					DataAccess.GetData( x1, y1, x2, y2, CachedData );
-					ValidX1 = FMath::Min<int32>(x1,ValidX1);
+					DataAccess.GetData(x1, y1, x2, y2, CachedData);
+					ValidX1 = FMath::Min<int32>(x1, ValidX1);
 				}
 				else
 				{
-					DataAccess.GetDataFast( X1, CachedY1, CachedX1-1, CachedY2, CachedData );
+					DataAccess.GetDataFast(X1, CachedY1, CachedX1 - 1, CachedY2, CachedData);
 				}
 
-				CacheOriginalData( X1, CachedY1, CachedX1-1, CachedY2 );
+				CacheOriginalData(X1, CachedY1, CachedX1 - 1, CachedY2);
 				CachedX1 = X1;
 			}
 
-			if( X2 > CachedX2 )
+			if (X2 > CachedX2)
 			{
 				if (Accessor::bUseInterp)
 				{
 					int32 x1 = ValidX2;
 					int32 x2 = X2;
-					int32 y1 = FMath::Min<int32>(Y1,CachedY1);
-					int32 y2 = FMath::Max<int32>(Y2,CachedY2);
+					int32 y1 = FMath::Min<int32>(Y1, CachedY1);
+					int32 y2 = FMath::Max<int32>(Y2, CachedY2);
 
-					DataAccess.GetData( x1, y1, x2, y2, CachedData );
-					ValidX2 = FMath::Max<int32>(x2,ValidX2);
+					DataAccess.GetData(x1, y1, x2, y2, CachedData);
+					ValidX2 = FMath::Max<int32>(x2, ValidX2);
 				}
 				else
 				{
-					DataAccess.GetDataFast( CachedX2+1, CachedY1, X2, CachedY2, CachedData );
+					DataAccess.GetDataFast(CachedX2 + 1, CachedY1, X2, CachedY2, CachedData);
 				}
-				CacheOriginalData( CachedX2+1, CachedY1, X2, CachedY2 );			
+				CacheOriginalData(CachedX2 + 1, CachedY1, X2, CachedY2);
 				CachedX2 = X2;
-			}			
+			}
 
-			if( Y1 < CachedY1 )
+			if (Y1 < CachedY1)
 			{
 				if (Accessor::bUseInterp)
 				{
@@ -337,18 +337,18 @@ struct TLandscapeEditCache
 					int32 y1 = Y1;
 					int32 y2 = ValidY1;
 
-					DataAccess.GetData( x1, y1, x2, y2, CachedData );
-					ValidY1 = FMath::Min<int32>(y1,ValidY1);
+					DataAccess.GetData(x1, y1, x2, y2, CachedData);
+					ValidY1 = FMath::Min<int32>(y1, ValidY1);
 				}
 				else
 				{
-					DataAccess.GetDataFast( CachedX1, Y1, CachedX2, CachedY1-1, CachedData );
+					DataAccess.GetDataFast(CachedX1, Y1, CachedX2, CachedY1 - 1, CachedData);
 				}
-				CacheOriginalData( CachedX1, Y1, CachedX2, CachedY1-1 );			
+				CacheOriginalData(CachedX1, Y1, CachedX2, CachedY1 - 1);
 				CachedY1 = Y1;
 			}
 
-			if( Y2 > CachedY2 )
+			if (Y2 > CachedY2)
 			{
 				if (Accessor::bUseInterp)
 				{
@@ -357,15 +357,15 @@ struct TLandscapeEditCache
 					int32 y1 = ValidY2;
 					int32 y2 = Y2;
 
-					DataAccess.GetData( x1, y1, x2, y2, CachedData );
-					ValidY2 = FMath::Max<int32>(y2,ValidY2);
+					DataAccess.GetData(x1, y1, x2, y2, CachedData);
+					ValidY2 = FMath::Max<int32>(y2, ValidY2);
 				}
 				else
 				{
-					DataAccess.GetDataFast( CachedX1, CachedY2+1, CachedX2, Y2, CachedData );
+					DataAccess.GetDataFast(CachedX1, CachedY2 + 1, CachedX2, Y2, CachedData);
 				}
 
-				CacheOriginalData( CachedX1, CachedY2+1, CachedX2, Y2 );			
+				CacheOriginalData(CachedX1, CachedY2 + 1, CachedX2, Y2);
 				CachedY2 = Y2;
 			}
 		}
@@ -373,7 +373,7 @@ struct TLandscapeEditCache
 
 	AccessorType* GetValueRef(int32 LandscapeX, int32 LandscapeY)
 	{
-		return CachedData.Find(ALandscape::MakeKey(LandscapeX,LandscapeY));
+		return CachedData.Find(ALandscape::MakeKey(LandscapeX, LandscapeY));
 	}
 
 	float GetValue(float LandscapeX, float LandscapeY)
@@ -381,15 +381,15 @@ struct TLandscapeEditCache
 		int32 X = FMath::FloorToInt(LandscapeX);
 		int32 Y = FMath::FloorToInt(LandscapeY);
 		AccessorType* P00 = CachedData.Find(ALandscape::MakeKey(X, Y));
-		AccessorType* P10 = CachedData.Find(ALandscape::MakeKey(X+1, Y));
-		AccessorType* P01 = CachedData.Find(ALandscape::MakeKey(X, Y+1));
-		AccessorType* P11 = CachedData.Find(ALandscape::MakeKey(X+1, Y+1));
+		AccessorType* P10 = CachedData.Find(ALandscape::MakeKey(X + 1, Y));
+		AccessorType* P01 = CachedData.Find(ALandscape::MakeKey(X, Y + 1));
+		AccessorType* P11 = CachedData.Find(ALandscape::MakeKey(X + 1, Y + 1));
 
 		// Search for nearest value if missing data
-		float V00 = P00 ? *P00 : (P10 ? *P10 : (P01 ? *P01 : (P11 ? *P11 : 0.f) ));
-		float V10 = P10 ? *P10 : (P00 ? *P00 : (P11 ? *P11 : (P01 ? *P01 : 0.f) ));
-		float V01 = P01 ? *P01 : (P00 ? *P00 : (P11 ? *P11 : (P10 ? *P10 : 0.f) ));
-		float V11 = P11 ? *P11 : (P10 ? *P10 : (P01 ? *P01 : (P00 ? *P00 : 0.f) ));
+		float V00 = P00 ? *P00 : (P10 ? *P10 : (P01 ? *P01 : (P11 ? *P11 : 0.f)));
+		float V10 = P10 ? *P10 : (P00 ? *P00 : (P11 ? *P11 : (P01 ? *P01 : 0.f)));
+		float V01 = P01 ? *P01 : (P00 ? *P00 : (P11 ? *P11 : (P10 ? *P10 : 0.f)));
+		float V11 = P11 ? *P11 : (P10 ? *P10 : (P01 ? *P01 : (P00 ? *P00 : 0.f)));
 
 		return FMath::Lerp(
 			FMath::Lerp(V00, V10, LandscapeX - X),
@@ -400,45 +400,45 @@ struct TLandscapeEditCache
 	FVector GetNormal(int32 X, int32 Y)
 	{
 		AccessorType* P00 = CachedData.Find(ALandscape::MakeKey(X, Y));
-		AccessorType* P10 = CachedData.Find(ALandscape::MakeKey(X+1, Y));
-		AccessorType* P01 = CachedData.Find(ALandscape::MakeKey(X, Y+1));
-		AccessorType* P11 = CachedData.Find(ALandscape::MakeKey(X+1, Y+1));
+		AccessorType* P10 = CachedData.Find(ALandscape::MakeKey(X + 1, Y));
+		AccessorType* P01 = CachedData.Find(ALandscape::MakeKey(X, Y + 1));
+		AccessorType* P11 = CachedData.Find(ALandscape::MakeKey(X + 1, Y + 1));
 
 		// Search for nearest value if missing data
-		float V00 = P00 ? *P00 : (P10 ? *P10 : (P01 ? *P01 : (P11 ? *P11 : 0.f) ));
-		float V10 = P10 ? *P10 : (P00 ? *P00 : (P11 ? *P11 : (P01 ? *P01 : 0.f) ));
-		float V01 = P01 ? *P01 : (P00 ? *P00 : (P11 ? *P11 : (P10 ? *P10 : 0.f) ));
-		float V11 = P11 ? *P11 : (P10 ? *P10 : (P01 ? *P01 : (P00 ? *P00 : 0.f) ));
-
-		FVector Vert00 = FVector(0.f,0.f, V00);
-		FVector Vert01 = FVector(0.f,1.f, V01);
-		FVector Vert10 = FVector(1.f,0.f, V10);
-		FVector Vert11 = FVector(1.f,1.f, V11);
-
-		FVector FaceNormal1 = ((Vert00-Vert10) ^ (Vert10-Vert11)).SafeNormal();
-		FVector FaceNormal2 = ((Vert11-Vert01) ^ (Vert01-Vert00)).SafeNormal();
+		float V00 = P00 ? *P00 : (P10 ? *P10 : (P01 ? *P01 : (P11 ? *P11 : 0.f)));
+		float V10 = P10 ? *P10 : (P00 ? *P00 : (P11 ? *P11 : (P01 ? *P01 : 0.f)));
+		float V01 = P01 ? *P01 : (P00 ? *P00 : (P11 ? *P11 : (P10 ? *P10 : 0.f)));
+		float V11 = P11 ? *P11 : (P10 ? *P10 : (P01 ? *P01 : (P00 ? *P00 : 0.f)));
+
+		FVector Vert00 = FVector(0.f, 0.f, V00);
+		FVector Vert01 = FVector(0.f, 1.f, V01);
+		FVector Vert10 = FVector(1.f, 0.f, V10);
+		FVector Vert11 = FVector(1.f, 1.f, V11);
+
+		FVector FaceNormal1 = ((Vert00 - Vert10) ^ (Vert10 - Vert11)).SafeNormal();
+		FVector FaceNormal2 = ((Vert11 - Vert01) ^ (Vert01 - Vert00)).SafeNormal();
 		return (FaceNormal1 + FaceNormal2).SafeNormal();
 	}
 
 	void SetValue(int32 LandscapeX, int32 LandscapeY, AccessorType Value)
 	{
-		CachedData.Add(ALandscape::MakeKey(LandscapeX,LandscapeY), Value);
+		CachedData.Add(ALandscape::MakeKey(LandscapeX, LandscapeY), Value);
 	}
 
 	void GetCachedData(int32 X1, int32 Y1, int32 X2, int32 Y2, TArray<AccessorType>& OutData)
 	{
-		int32 NumSamples = (1+X2-X1)*(1+Y2-Y1);
+		int32 NumSamples = (1 + X2 - X1)*(1 + Y2 - Y1);
 		OutData.Empty(NumSamples);
 		OutData.AddUninitialized(NumSamples);
 
-		for( int32 Y=Y1;Y<=Y2;Y++ )
+		for (int32 Y = Y1; Y <= Y2; Y++)
 		{
-			for( int32 X=X1;X<=X2;X++ )
+			for (int32 X = X1; X <= X2; X++)
 			{
-				AccessorType* Ptr = GetValueRef(X,Y);
-				if( Ptr )
+				AccessorType* Ptr = GetValueRef(X, Y);
+				if (Ptr)
 				{
-					OutData[(X-X1) + (Y-Y1)*(1+X2-X1)] = *Ptr;
+					OutData[(X - X1) + (Y - Y1)*(1 + X2 - X1)] = *Ptr;
 				}
 			}
 		}
@@ -447,33 +447,33 @@ struct TLandscapeEditCache
 	void SetCachedData(int32 X1, int32 Y1, int32 X2, int32 Y2, TArray<AccessorType>& Data, ELandscapeLayerPaintingRestriction::Type PaintingRestriction = ELandscapeLayerPaintingRestriction::None)
 	{
 		// Update cache
-		for( int32 Y=Y1;Y<=Y2;Y++ )
+		for (int32 Y = Y1; Y <= Y2; Y++)
 		{
-			for( int32 X=X1;X<=X2;X++ )
+			for (int32 X = X1; X <= X2; X++)
 			{
-				SetValue( X, Y, Data[(X-X1) + (Y-Y1)*(1+X2-X1)] );
+				SetValue(X, Y, Data[(X - X1) + (Y - Y1)*(1 + X2 - X1)]);
 			}
 		}
 
 		// Update real data
-		DataAccess.SetData( X1, Y1, X2, Y2, Data.GetTypedData(), PaintingRestriction );
+		DataAccess.SetData(X1, Y1, X2, Y2, Data.GetTypedData(), PaintingRestriction);
 	}
 
 	// Get the original data before we made any changes with the SetCachedData interface.
 	void GetOriginalData(int32 X1, int32 Y1, int32 X2, int32 Y2, TArray<AccessorType>& OutOriginalData)
 	{
-		int32 NumSamples = (1+X2-X1)*(1+Y2-Y1);
+		int32 NumSamples = (1 + X2 - X1)*(1 + Y2 - Y1);
 		OutOriginalData.Empty(NumSamples);
 		OutOriginalData.AddUninitialized(NumSamples);
 
-		for( int32 Y=Y1;Y<=Y2;Y++ )
+		for (int32 Y = Y1; Y <= Y2; Y++)
 		{
-			for( int32 X=X1;X<=X2;X++ )
+			for (int32 X = X1; X <= X2; X++)
 			{
-				AccessorType* Ptr = OriginalData.Find(ALandscape::MakeKey(X,Y));
-				if( Ptr )
+				AccessorType* Ptr = OriginalData.Find(ALandscape::MakeKey(X, Y));
+				if (Ptr)
 				{
-					OutOriginalData[(X-X1) + (Y-Y1)*(1+X2-X1)] = *Ptr;
+					OutOriginalData[(X - X1) + (Y - Y1)*(1 + X2 - X1)] = *Ptr;
 				}
 			}
 		}
@@ -487,17 +487,17 @@ struct TLandscapeEditCache
 protected:
 	Accessor& DataAccess;
 private:
-	void CacheOriginalData(int32 X1, int32 Y1, int32 X2, int32 Y2 )
+	void CacheOriginalData(int32 X1, int32 Y1, int32 X2, int32 Y2)
 	{
-		for( int32 Y=Y1;Y<=Y2;Y++ )
+		for (int32 Y = Y1; Y <= Y2; Y++)
 		{
-			for( int32 X=X1;X<=X2;X++ )
+			for (int32 X = X1; X <= X2; X++)
 			{
-				FIntPoint Key = ALandscape::MakeKey(X,Y);
+				FIntPoint Key = ALandscape::MakeKey(X, Y);
 				AccessorType* Ptr = CachedData.Find(Key);
-				if( Ptr )
+				if (Ptr)
 				{
-					check( OriginalData.Find(Key) == NULL );
+					check(OriginalData.Find(Key) == NULL);
 					OriginalData.Add(Key, *Ptr);
 				}
 			}
@@ -525,7 +525,7 @@ template<bool bInUseInterp>
 struct FHeightmapAccessor
 {
 	enum { bUseInterp = bInUseInterp };
-	FHeightmapAccessor( ULandscapeInfo* InLandscapeInfo )
+	FHeightmapAccessor(ULandscapeInfo* InLandscapeInfo)
 	{
 		LandscapeInfo = InLandscapeInfo;
 		LandscapeEdit = new FLandscapeEditDataInterface(InLandscapeInfo);
@@ -534,12 +534,12 @@ struct FHeightmapAccessor
 	// accessors
 	void GetData(int32& X1, int32& Y1, int32& X2, int32& Y2, TMap<FIntPoint, uint16>& Data)
 	{
-		LandscapeEdit->GetHeightData( X1, Y1, X2, Y2, Data);
+		LandscapeEdit->GetHeightData(X1, Y1, X2, Y2, Data);
 	}
 
 	void GetDataFast(int32 X1, int32 Y1, int32 X2, int32 Y2, TMap<FIntPoint, uint16>& Data)
 	{
-		LandscapeEdit->GetHeightDataFast( X1, Y1, X2, Y2, Data);
+		LandscapeEdit->GetHeightDataFast(X1, Y1, X2, Y2, Data);
 	}
 
 	void SetData(int32 X1, int32 Y1, int32 X2, int32 Y2, const uint16* Data, ELandscapeLayerPaintingRestriction::Type PaintingRestriction = ELandscapeLayerPaintingRestriction::None)
@@ -598,7 +598,7 @@ struct FHeightmapAccessor
 			else
 			{
 				// No foliage, just update landscape.
-				LandscapeEdit->SetHeightData( X1, Y1, X2, Y2, Data, 0, true);
+				LandscapeEdit->SetHeightData(X1, Y1, X2, Y2, Data, 0, true);
 			}
 		}
 		else
@@ -618,11 +618,11 @@ struct FHeightmapAccessor
 		LandscapeEdit = NULL;
 
 		// Update the bounds and navmesh for the components we edited
-		for(TSet<ULandscapeComponent*>::TConstIterator It(ChangedComponents);It;++It)
+		for (TSet<ULandscapeComponent*>::TConstIterator It(ChangedComponents); It; ++It)
 		{
 			(*It)->UpdateCachedBounds();
 			(*It)->UpdateComponentToWorld();
-			
+
 			// Recreate collision for modified components to update the physical materials
 			ULandscapeHeightfieldCollisionComponent* CollisionComponent = (*It)->CollisionComponent.Get();
 			if (CollisionComponent)
@@ -643,16 +643,16 @@ private:
 	TSet<ULandscapeComponent*> ChangedComponents;
 };
 
-struct FLandscapeHeightCache : public TLandscapeEditCache<FHeightmapAccessor<true>,uint16>
+struct FLandscapeHeightCache : public TLandscapeEditCache < FHeightmapAccessor<true>, uint16 >
 {
 	typedef uint16 DataType;
-	static uint16 ClampValue( int32 Value ) { return FMath::Clamp(Value, 0, LandscapeDataAccess::MaxValue); }
+	static uint16 ClampValue(int32 Value) { return FMath::Clamp(Value, 0, LandscapeDataAccess::MaxValue); }
 
 	FHeightmapAccessor<true> HeightmapAccessor;
 
 	FLandscapeHeightCache(const FLandscapeToolTarget& InTarget)
-		:	HeightmapAccessor(InTarget.LandscapeInfo.Get())
-		,	TLandscapeEditCache(HeightmapAccessor)
+		: HeightmapAccessor(InTarget.LandscapeInfo.Get())
+		, TLandscapeEditCache(HeightmapAccessor)
 	{
 	}
 };
@@ -664,7 +664,7 @@ template<bool bInUseInterp>
 struct FXYOffsetmapAccessor
 {
 	enum { bUseInterp = bInUseInterp };
-	FXYOffsetmapAccessor( ULandscapeInfo* InLandscapeInfo )
+	FXYOffsetmapAccessor(ULandscapeInfo* InLandscapeInfo)
 	{
 		LandscapeInfo = InLandscapeInfo;
 		LandscapeEdit = new FLandscapeEditDataInterface(InLandscapeInfo);
@@ -676,15 +676,15 @@ struct FXYOffsetmapAccessor
 		LandscapeEdit->GetXYOffsetData(X1, Y1, X2, Y2, Data);
 
 		TMap<FIntPoint, uint16> NewHeights;
-		LandscapeEdit->GetHeightData( X1, Y1, X2, Y2, NewHeights);
+		LandscapeEdit->GetHeightData(X1, Y1, X2, Y2, NewHeights);
 		for (int32 Y = Y1; Y <= Y2; ++Y)
 		{
 			for (int32 X = X1; X <= X2; ++X)
 			{
-				FVector* Value = Data.Find( ALandscape::MakeKey(X, Y) );
+				FVector* Value = Data.Find(ALandscape::MakeKey(X, Y));
 				if (Value)
 				{
-					Value->Z = ( (float)NewHeights.FindRef( ALandscape::MakeKey(X, Y) ) - 32768.f) * LANDSCAPE_ZSCALE;
+					Value->Z = ((float)NewHeights.FindRef(ALandscape::MakeKey(X, Y)) - 32768.f) * LANDSCAPE_ZSCALE;
 				}
 			}
 		}
@@ -695,15 +695,15 @@ struct FXYOffsetmapAccessor
 		LandscapeEdit->GetXYOffsetData(X1, Y1, X2, Y2, Data);
 
 		TMap<FIntPoint, uint16> NewHeights;
-		LandscapeEdit->GetHeightDataFast( X1, Y1, X2, Y2, NewHeights);
+		LandscapeEdit->GetHeightDataFast(X1, Y1, X2, Y2, NewHeights);
 		for (int32 Y = Y1; Y <= Y2; ++Y)
 		{
 			for (int32 X = X1; X <= X2; ++X)
 			{
-				FVector* Value = Data.Find( ALandscape::MakeKey(X, Y) );
+				FVector* Value = Data.Find(ALandscape::MakeKey(X, Y));
 				if (Value)
 				{
-					Value->Z = ( (float)NewHeights.FindRef( ALandscape::MakeKey(X, Y) ) - 32768.f) * LANDSCAPE_ZSCALE;
+					Value->Z = ((float)NewHeights.FindRef(ALandscape::MakeKey(X, Y)) - 32768.f) * LANDSCAPE_ZSCALE;
 				}
 			}
 		}
@@ -719,12 +719,12 @@ struct FXYOffsetmapAccessor
 
 			// Convert Height to uint16
 			TArray<uint16> NewHeights;
-			NewHeights.AddZeroed( (Y2-Y1+1) * (X2-X1+1) );
+			NewHeights.AddZeroed((Y2 - Y1 + 1) * (X2 - X1 + 1));
 			for (int32 Y = Y1; Y <= Y2; ++Y)
 			{
 				for (int32 X = X1; X <= X2; ++X)
 				{
-					NewHeights[X - X1 + (Y - Y1) * (X2-X1+1) ] = FMath::Clamp<uint16>(Data[(X - X1 + (Y - Y1) * (X2-X1+1) )].Z * LANDSCAPE_INV_ZSCALE + 32768.f, 0, 65535);
+					NewHeights[X - X1 + (Y - Y1) * (X2 - X1 + 1)] = FMath::Clamp<uint16>(Data[(X - X1 + (Y - Y1) * (X2 - X1 + 1))].Z * LANDSCAPE_INV_ZSCALE + 32768.f, 0, 65535);
 				}
 			}
 
@@ -755,8 +755,8 @@ struct FXYOffsetmapAccessor
 				}
 
 				// Update landscape.
-				LandscapeEdit->SetXYOffsetData( X1, Y1, X2, Y2, Data, 0 ); // XY Offset always need to be update before the height update
-				LandscapeEdit->SetHeightData( X1, Y1, X2, Y2, NewHeights.GetTypedData(), 0, true);
+				LandscapeEdit->SetXYOffsetData(X1, Y1, X2, Y2, Data, 0); // XY Offset always need to be update before the height update
+				LandscapeEdit->SetHeightData(X1, Y1, X2, Y2, NewHeights.GetTypedData(), 0, true);
 
 				// Snap foliage for each component.
 				for (int32 Index = 0; Index < CollisionComponents.Num(); ++Index)
@@ -769,8 +769,8 @@ struct FXYOffsetmapAccessor
 			else
 			{
 				// No foliage, just update landscape.
-				LandscapeEdit->SetXYOffsetData( X1, Y1, X2, Y2, Data, 0 ); // XY Offset always need to be update before the height update
-				LandscapeEdit->SetHeightData( X1, Y1, X2, Y2, NewHeights.GetTypedData(), 0, true);
+				LandscapeEdit->SetXYOffsetData(X1, Y1, X2, Y2, Data, 0); // XY Offset always need to be update before the height update
+				LandscapeEdit->SetHeightData(X1, Y1, X2, Y2, NewHeights.GetTypedData(), 0, true);
 			}
 		}
 		else
@@ -790,7 +790,7 @@ struct FXYOffsetmapAccessor
 		LandscapeEdit = NULL;
 
 		// Update the bounds for the components we edited
-		for(TSet<ULandscapeComponent*>::TConstIterator It(ChangedComponents);It;++It)
+		for (TSet<ULandscapeComponent*>::TConstIterator It(ChangedComponents); It; ++It)
 		{
 			(*It)->UpdateCachedBounds();
 			(*It)->UpdateComponentToWorld();
@@ -803,15 +803,15 @@ private:
 	TSet<ULandscapeComponent*> ChangedComponents;
 };
 
-struct FLandscapeXYOffsetCache : public TLandscapeEditCache<FXYOffsetmapAccessor<true>,FVector>
+struct FLandscapeXYOffsetCache : public TLandscapeEditCache < FXYOffsetmapAccessor<true>, FVector >
 {
 	typedef FVector DataType;
 
 	FXYOffsetmapAccessor<true> XYOffsetmapAccessor;
 
 	FLandscapeXYOffsetCache(const FLandscapeToolTarget& InTarget)
-		:	XYOffsetmapAccessor(InTarget.LandscapeInfo.Get())
-		,	TLandscapeEditCache(XYOffsetmapAccessor)
+		: XYOffsetmapAccessor(InTarget.LandscapeInfo.Get())
+		, TLandscapeEditCache(XYOffsetmapAccessor)
 	{
 	}
 };
@@ -824,16 +824,16 @@ struct FAlphamapAccessor
 {
 	enum { bUseInterp = bInUseInterp };
 	enum { bUseTotalNormalize = bInUseTotalNormalize };
-	FAlphamapAccessor( ULandscapeInfo* InLandscapeInfo, ULandscapeLayerInfoObject* InLayerInfo )
-		:	LandscapeInfo(InLandscapeInfo)
-		,	LandscapeEdit(InLandscapeInfo)
-		,	LayerInfo(InLayerInfo)
-		,	bBlendWeight(true)
+	FAlphamapAccessor(ULandscapeInfo* InLandscapeInfo, ULandscapeLayerInfoObject* InLayerInfo)
+		: LandscapeInfo(InLandscapeInfo)
+		, LandscapeEdit(InLandscapeInfo)
+		, LayerInfo(InLayerInfo)
+		, bBlendWeight(true)
 	{
 		// should be no Layer change during FAlphamapAccessor lifetime...
 		if (InLandscapeInfo && InLayerInfo)
 		{
-			if (LayerInfo == ALandscapeProxy::DataLayer)
+			if (LayerInfo == ALandscapeProxy::VisibilityLayer)
 			{
 				bBlendWeight = false;
 			}
@@ -850,12 +850,12 @@ struct FAlphamapAccessor
 		for (auto It = ModifiedComponents.CreateConstIterator(); It; ++It)
 		{
 			ULandscapeHeightfieldCollisionComponent* CollisionComponent = (*It)->CollisionComponent.Get();
-			if( CollisionComponent )
+			if (CollisionComponent)
 			{
 				CollisionComponent->RecreateCollision(false);
-				
+
 				// We need to trigger navigation mesh build, in case user have painted holes on a landscape
-				if (LayerInfo == ALandscapeProxy::DataLayer)
+				if (LayerInfo == ALandscapeProxy::VisibilityLayer)
 				{
 					UNavigationSystem* NavSys = UNavigationSystem::GetCurrent(*It);
 					if (NavSys)
@@ -900,30 +900,30 @@ private:
 	bool bBlendWeight;
 };
 
-struct FLandscapeAlphaCache : public TLandscapeEditCache<FAlphamapAccessor<true, false>,uint8>
+struct FLandscapeAlphaCache : public TLandscapeEditCache < FAlphamapAccessor<true, false>, uint8 >
 {
 	typedef uint8 DataType;
-	static uint8 ClampValue( int32 Value ) { return FMath::Clamp(Value, 0, 255); }
+	static uint8 ClampValue(int32 Value) { return FMath::Clamp(Value, 0, 255); }
 
 	FAlphamapAccessor<true, false> AlphamapAccessor;
 
 	FLandscapeAlphaCache(const FLandscapeToolTarget& InTarget)
-		:	AlphamapAccessor(InTarget.LandscapeInfo.Get(), InTarget.LayerInfo.Get())
-		,	TLandscapeEditCache(AlphamapAccessor)
+		: AlphamapAccessor(InTarget.LandscapeInfo.Get(), InTarget.LayerInfo.Get())
+		, TLandscapeEditCache(AlphamapAccessor)
 	{
 	}
 };
 
-struct FLandscapeVisCache : public TLandscapeEditCache<FAlphamapAccessor<false, false>,uint8>
+struct FLandscapeVisCache : public TLandscapeEditCache < FAlphamapAccessor<false, false>, uint8 >
 {
 	typedef uint8 DataType;
-	static uint8 ClampValue( int32 Value ) { return FMath::Clamp(Value, 0, 255); }
+	static uint8 ClampValue(int32 Value) { return FMath::Clamp(Value, 0, 255); }
 
 	FAlphamapAccessor<false, false> AlphamapAccessor;
 
 	FLandscapeVisCache(const FLandscapeToolTarget& InTarget)
-		:	AlphamapAccessor(InTarget.LandscapeInfo.Get(), ALandscapeProxy::DataLayer)
-		,	TLandscapeEditCache(AlphamapAccessor)
+		: AlphamapAccessor(InTarget.LandscapeInfo.Get(), ALandscapeProxy::VisibilityLayer)
+		, TLandscapeEditCache(AlphamapAccessor)
 	{
 	}
 };
@@ -935,8 +935,8 @@ template<bool bInUseInterp>
 struct FFullWeightmapAccessor
 {
 	enum { bUseInterp = bInUseInterp };
-	FFullWeightmapAccessor( ULandscapeInfo* InLandscapeInfo)
-		:	LandscapeEdit(InLandscapeInfo)
+	FFullWeightmapAccessor(ULandscapeInfo* InLandscapeInfo)
+		: LandscapeEdit(InLandscapeInfo)
 	{
 	}
 	void GetData(int32& X1, int32& Y1, int32& X2, int32& Y2, TMap<FIntPoint, TArray<uint8>>& Data)
@@ -970,35 +970,35 @@ private:
 	FLandscapeEditDataInterface LandscapeEdit;
 };
 
-struct FLandscapeFullWeightCache : public TLandscapeEditCache<FFullWeightmapAccessor<false>,TArray<uint8>>
+struct FLandscapeFullWeightCache : public TLandscapeEditCache < FFullWeightmapAccessor<false>, TArray<uint8> >
 {
 	typedef TArray<uint8> DataType;
 
 	FFullWeightmapAccessor<false> WeightmapAccessor;
 
 	FLandscapeFullWeightCache(const FLandscapeToolTarget& InTarget)
-		:	WeightmapAccessor(InTarget.LandscapeInfo.Get())
-		,	TLandscapeEditCache(WeightmapAccessor)
+		: WeightmapAccessor(InTarget.LandscapeInfo.Get())
+		, TLandscapeEditCache(WeightmapAccessor)
 	{
 	}
 
 	// Only for all weight case... the accessor type should be TArray<uint8>
 	void GetCachedData(int32 X1, int32 Y1, int32 X2, int32 Y2, TArray<uint8>& OutData, int32 ArraySize)
 	{
-		int32 NumSamples = (1+X2-X1)*(1+Y2-Y1) * ArraySize;
+		int32 NumSamples = (1 + X2 - X1)*(1 + Y2 - Y1) * ArraySize;
 		OutData.Empty(NumSamples);
 		OutData.AddUninitialized(NumSamples);
 
-		for( int32 Y=Y1;Y<=Y2;Y++ )
+		for (int32 Y = Y1; Y <= Y2; Y++)
 		{
-			for( int32 X=X1;X<=X2;X++ )
+			for (int32 X = X1; X <= X2; X++)
 			{
-				TArray<uint8>* Ptr = GetValueRef(X,Y);
-				if( Ptr )
+				TArray<uint8>* Ptr = GetValueRef(X, Y);
+				if (Ptr)
 				{
 					for (int32 Z = 0; Z < ArraySize; Z++)
 					{
-						OutData[( (X-X1) + (Y-Y1)*(1+X2-X1)) * ArraySize + Z] = (*Ptr)[Z];
+						OutData[((X - X1) + (Y - Y1)*(1 + X2 - X1)) * ArraySize + Z] = (*Ptr)[Z];
 					}
 				}
 			}
@@ -1009,18 +1009,18 @@ struct FLandscapeFullWeightCache : public TLandscapeEditCache<FFullWeightmapAcce
 	void SetCachedData(int32 X1, int32 Y1, int32 X2, int32 Y2, TArray<uint8>& Data, int32 ArraySize, ELandscapeLayerPaintingRestriction::Type PaintingRestriction)
 	{
 		// Update cache
-		for( int32 Y=Y1;Y<=Y2;Y++ )
+		for (int32 Y = Y1; Y <= Y2; Y++)
 		{
-			for( int32 X=X1;X<=X2;X++ )
+			for (int32 X = X1; X <= X2; X++)
 			{
 				TArray<uint8> Value;
 				Value.Empty(ArraySize);
 				Value.AddUninitialized(ArraySize);
-				for ( int32 Z=0; Z < ArraySize; Z++)
+				for (int32 Z = 0; Z < ArraySize; Z++)
 				{
-					Value[Z] = Data[ ((X-X1) + (Y-Y1)*(1+X2-X1)) * ArraySize + Z];
+					Value[Z] = Data[((X - X1) + (Y - Y1)*(1 + X2 - X1)) * ArraySize + Z];
 				}
-				SetValue( X, Y, Value );
+				SetValue(X, Y, Value);
 			}
 		}
 
@@ -1041,8 +1041,8 @@ template<bool bInUseInterp>
 struct FDatamapAccessor
 {
 	enum { bUseInterp = bInUseInterp };
-	FDatamapAccessor( ULandscapeInfo* InLandscapeInfo )
-		:	LandscapeEdit(InLandscapeInfo)
+	FDatamapAccessor(ULandscapeInfo* InLandscapeInfo)
+		: LandscapeEdit(InLandscapeInfo)
 	{
 	}
 
@@ -1056,7 +1056,7 @@ struct FDatamapAccessor
 		LandscapeEdit.GetSelectData(X1, Y1, X2, Y2, Data);
 	}
 
-	void SetData(int32 X1, int32 Y1, int32 X2, int32 Y2, const uint8* Data, ELandscapeLayerPaintingRestriction::Type PaintingRestriction = ELandscapeLayerPaintingRestriction::None )
+	void SetData(int32 X1, int32 Y1, int32 X2, int32 Y2, const uint8* Data, ELandscapeLayerPaintingRestriction::Type PaintingRestriction = ELandscapeLayerPaintingRestriction::None)
 	{
 		if (LandscapeEdit.GetComponentsInRegion(X1, Y1, X2, Y2))
 		{
@@ -1073,16 +1073,16 @@ private:
 	FLandscapeEditDataInterface LandscapeEdit;
 };
 
-struct FLandscapeDataCache : public TLandscapeEditCache<FDatamapAccessor<false>,uint8>
+struct FLandscapeDataCache : public TLandscapeEditCache < FDatamapAccessor<false>, uint8 >
 {
 	typedef uint8 DataType;
-	static uint8 ClampValue( int32 Value ) { return FMath::Clamp(Value, 0, 255); }
+	static uint8 ClampValue(int32 Value) { return FMath::Clamp(Value, 0, 255); }
 
 	FDatamapAccessor<false> DataAccessor;
 
 	FLandscapeDataCache(const FLandscapeToolTarget& InTarget)
-		:	DataAccessor(InTarget.LandscapeInfo.Get())
-		,	TLandscapeEditCache(DataAccessor)
+		: DataAccessor(InTarget.LandscapeInfo.Get())
+		, TLandscapeEditCache(DataAccessor)
 	{
 	}
 };
@@ -1108,15 +1108,15 @@ struct FHeightmapToolTarget
 
 	static FMatrix ToWorldMatrix(ULandscapeInfo* LandscapeInfo)
 	{
-		FMatrix Result = FTranslationMatrix(FVector(0,0,-32768.f));
-		Result *= FScaleMatrix( FVector(1.f,1.f,LANDSCAPE_ZSCALE) * LandscapeInfo->DrawScale );
+		FMatrix Result = FTranslationMatrix(FVector(0, 0, -32768.f));
+		Result *= FScaleMatrix(FVector(1.f, 1.f, LANDSCAPE_ZSCALE) * LandscapeInfo->DrawScale);
 		return Result;
 	}
 
 	static FMatrix FromWorldMatrix(ULandscapeInfo* LandscapeInfo)
 	{
-		FMatrix Result = FScaleMatrix( FVector(1.f,1.f,LANDSCAPE_INV_ZSCALE) / (LandscapeInfo->DrawScale) );
-		Result *= FTranslationMatrix(FVector(0,0,32768.f));
+		FMatrix Result = FScaleMatrix(FVector(1.f, 1.f, LANDSCAPE_INV_ZSCALE) / (LandscapeInfo->DrawScale));
+		Result *= FTranslationMatrix(FVector(0, 0, 32768.f));
 		return Result;
 	}
 };
@@ -1155,14 +1155,14 @@ class FLandscapeToolBase : public FLandscapeTool
 {
 public:
 	FLandscapeToolBase(FEdModeLandscape* InEdMode)
-	:	EdMode(InEdMode)
-	,	bToolActive(false)
-	,	ToolStroke(NULL)
+		: EdMode(InEdMode)
+		, bToolActive(false)
+		, ToolStroke(NULL)
 	{}
 
-	virtual bool BeginTool( FEditorViewportClient* ViewportClient, const FLandscapeToolTarget& InTarget, const FVector& InHitLocation ) override
+	virtual bool BeginTool(FEditorViewportClient* ViewportClient, const FLandscapeToolTarget& InTarget, const FVector& InHitLocation) override
 	{
-		if( !ensure(MousePositions.Num()==0) ) 
+		if (!ensure(MousePositions.Num() == 0))
 		{
 			MousePositions.Empty();
 		}
@@ -1172,24 +1172,24 @@ public:
 
 		EdMode->CurrentBrush->BeginStroke(InHitLocation.X, InHitLocation.Y, this);
 
-		new(MousePositions) FLandscapeToolMousePosition(InHitLocation.X, InHitLocation.Y, IsShiftDown(ViewportClient->Viewport));
+		new(MousePositions)FLandscapeToolMousePosition(InHitLocation.X, InHitLocation.Y, IsShiftDown(ViewportClient->Viewport));
 		ToolStroke->Apply(ViewportClient, EdMode->CurrentBrush, EdMode->UISettings, MousePositions);
 		MousePositions.Empty();
 		return true;
 	}
 
-	virtual void Tick(FEditorViewportClient* ViewportClient,float DeltaTime) override 
+	virtual void Tick(FEditorViewportClient* ViewportClient, float DeltaTime) override
 	{
-		if( bToolActive && MousePositions.Num() )
+		if (bToolActive && MousePositions.Num())
 		{
 			ToolStroke->Apply(ViewportClient, EdMode->CurrentBrush, EdMode->UISettings, MousePositions);
 			MousePositions.Empty();
 		}
 	}
 
-	virtual void EndTool(FEditorViewportClient* ViewportClient) override 
+	virtual void EndTool(FEditorViewportClient* ViewportClient) override
 	{
-		if( bToolActive && MousePositions.Num() )
+		if (bToolActive && MousePositions.Num())
 		{
 			ToolStroke->Apply(ViewportClient, EdMode->CurrentBrush, EdMode->UISettings, MousePositions);
 			MousePositions.Empty();
@@ -1201,30 +1201,30 @@ public:
 		EdMode->CurrentBrush->EndStroke();
 	}
 
-	virtual bool MouseMove( FEditorViewportClient* ViewportClient, FViewport* Viewport, int32 x, int32 y ) override 
+	virtual bool MouseMove(FEditorViewportClient* ViewportClient, FViewport* Viewport, int32 x, int32 y) override
 	{
 		FVector HitLocation;
-		if( EdMode->LandscapeMouseTrace(ViewportClient, x, y, HitLocation)  )
+		if (EdMode->LandscapeMouseTrace(ViewportClient, x, y, HitLocation))
 		{
-			if( EdMode->CurrentBrush )
+			if (EdMode->CurrentBrush)
 			{
 				// Inform the brush of the current location, to update the cursor
 				EdMode->CurrentBrush->MouseMove(HitLocation.X, HitLocation.Y);
 			}
 
-			if( bToolActive )
+			if (bToolActive)
 			{
 				// Save the mouse position
-				new(MousePositions) FLandscapeToolMousePosition(HitLocation.X, HitLocation.Y, IsShiftDown(ViewportClient->Viewport));
+				new(MousePositions)FLandscapeToolMousePosition(HitLocation.X, HitLocation.Y, IsShiftDown(ViewportClient->Viewport));
 			}
 		}
 
 		return true;
-	}	
+	}
 
 protected:
 	TArray<FLandscapeToolMousePosition> MousePositions;
-	class FEdModeLandscape* EdMode;
+	FEdModeLandscape* EdMode;
 	bool bToolActive;
 	TStrokeClass* ToolStroke;
 };
diff --git a/Engine/Source/Editor/LandscapeEditor/Private/LandscapeEdModeXYOffset.cpp b/Engine/Source/Editor/LandscapeEditor/Private/LandscapeEdModeXYOffset.cpp
--- a/Engine/Source/Editor/LandscapeEditor/Private/LandscapeEdModeXYOffset.cpp
+++ b/Engine/Source/Editor/LandscapeEditor/Private/LandscapeEdModeXYOffset.cpp
@@ -17,29 +17,29 @@ namespace
 {
 	FORCEINLINE FVector4 GetWorldPos(const FMatrix& LocalToWorld, FVector2D LocalXY, uint16 Height, FVector2D XYOffset)
 	{
-		return LocalToWorld.TransformPosition( FVector( LocalXY.X + XYOffset.X, LocalXY.Y + XYOffset.Y, ((float)Height - 32768.f) * LANDSCAPE_ZSCALE ) );
+		return LocalToWorld.TransformPosition(FVector(LocalXY.X + XYOffset.X, LocalXY.Y + XYOffset.Y, ((float)Height - 32768.f) * LANDSCAPE_ZSCALE));
 	}
 
 	FORCEINLINE FVector4 GetWorldPos(const FMatrix& LocalToWorld, FVector2D LocalXY, FVector XYOffsetVector)
 	{
-		return LocalToWorld.TransformPosition( FVector( LocalXY.X + XYOffsetVector.X, LocalXY.Y + XYOffsetVector.Y, XYOffsetVector.Z ) );
+		return LocalToWorld.TransformPosition(FVector(LocalXY.X + XYOffsetVector.X, LocalXY.Y + XYOffsetVector.Y, XYOffsetVector.Z));
 	}
 
 	const int32 XOffsets[4] = { 0, 1, 0, 1 };
 	const int32 YOffsets[4] = { 0, 0, 1, 1 };
 
-	float GetHeight(int32 X, int32 Y, int32 MinX, int32 MinY, int32 MaxX, int32 MaxY, const FVector& XYOffset, const TArray<FVector>& XYOffsetVectorData )
+	float GetHeight(int32 X, int32 Y, int32 MinX, int32 MinY, int32 MaxX, int32 MaxY, const FVector& XYOffset, const TArray<FVector>& XYOffsetVectorData)
 	{
 		float Height[4];
 		for (int32 Idx = 0; Idx < 4; ++Idx)
 		{
-			int32 XX = FMath::Clamp(FMath::FloorToInt(X+XYOffset.X + XOffsets[Idx]), MinX, MaxX);
-			int32 YY = FMath::Clamp(FMath::FloorToInt(Y+XYOffset.Y + YOffsets[Idx]), MinY, MaxY);
-			Height[Idx] = XYOffsetVectorData[ XX - MinX + (YY - MinY) * (MaxX-MinX+1) ].Z;
+			int32 XX = FMath::Clamp(FMath::FloorToInt(X + XYOffset.X + XOffsets[Idx]), MinX, MaxX);
+			int32 YY = FMath::Clamp(FMath::FloorToInt(Y + XYOffset.Y + YOffsets[Idx]), MinY, MaxY);
+			Height[Idx] = XYOffsetVectorData[XX - MinX + (YY - MinY) * (MaxX - MinX + 1)].Z;
 		}
-		float FracX = FMath::Fractional(X+XYOffset.X);
-		float FracY = FMath::Fractional(Y+XYOffset.Y);
-		return FMath::Lerp(	FMath::Lerp(Height[0], Height[1], FracX),
+		float FracX = FMath::Fractional(X + XYOffset.X);
+		float FracY = FMath::Fractional(Y + XYOffset.Y);
+		return FMath::Lerp(FMath::Lerp(Height[0], Height[1], FracX),
 			FMath::Lerp(Height[2], Height[3], FracX),
 			FracY);
 	}
@@ -54,13 +54,13 @@ class FLandscapeToolStrokeRetopologize
 {
 public:
 	FLandscapeToolStrokeRetopologize(FEdModeLandscape* InEdMode, const FLandscapeToolTarget& InTarget)
-		:	LandscapeInfo(InTarget.LandscapeInfo.Get())
-		,	Cache(InTarget)
+		: LandscapeInfo(InTarget.LandscapeInfo.Get())
+		, Cache(InTarget)
 	{}
 
 	virtual void Apply(FEditorViewportClient* ViewportClient, FLandscapeBrush* Brush, const ULandscapeEditorObject* UISettings, const TArray<FLandscapeToolMousePosition>& MousePositions)
 	{
-		if (!LandscapeInfo) 
+		if (!LandscapeInfo)
 		{
 			return;
 		}
@@ -81,12 +81,12 @@ public:
 		float Pressure = ViewportClient->Viewport->IsPenActive() ? ViewportClient->Viewport->GetTabletPressure() : 1.f;
 
 		// expand the area by one vertex in each direction to ensure normals are calculated correctly
-/*
-		X1 -= 1;
-		Y1 -= 1;
-		X2 += 1;
-		Y2 += 1;
-*/
+		/*
+				X1 -= 1;
+				Y1 -= 1;
+				X2 += 1;
+				Y2 += 1;
+				*/
 		{
 			int32 ValidX1, ValidX2, ValidY1, ValidY2;
 			ValidX1 = ValidY1 = INT_MAX;
@@ -95,19 +95,19 @@ public:
 			int32 ComponentSizeQuads = LandscapeInfo->ComponentSizeQuads;
 			ALandscape::CalcComponentIndicesOverlap(X1, Y1, X2, Y2, ComponentSizeQuads, ComponentIndexX1, ComponentIndexY1, ComponentIndexX2, ComponentIndexY2);
 
-			for( int32 ComponentIndexY=ComponentIndexY1;ComponentIndexY<=ComponentIndexY2;ComponentIndexY++ )
+			for (int32 ComponentIndexY = ComponentIndexY1; ComponentIndexY <= ComponentIndexY2; ComponentIndexY++)
 			{
-				for( int32 ComponentIndexX=ComponentIndexX1;ComponentIndexX<=ComponentIndexX2;ComponentIndexX++ )
-				{		
-					ULandscapeComponent* Component = LandscapeInfo->XYtoComponentMap.FindRef(FIntPoint(ComponentIndexX,ComponentIndexY));
+				for (int32 ComponentIndexX = ComponentIndexX1; ComponentIndexX <= ComponentIndexX2; ComponentIndexX++)
+				{
+					ULandscapeComponent* Component = LandscapeInfo->XYtoComponentMap.FindRef(FIntPoint(ComponentIndexX, ComponentIndexY));
 
 					if (Component)
 					{
 						// Update valid region
 						ValidX1 = FMath::Min<int32>(Component->GetSectionBase().X, ValidX1);
-						ValidX2 = FMath::Max<int32>(Component->GetSectionBase().X+ComponentSizeQuads, ValidX2);
+						ValidX2 = FMath::Max<int32>(Component->GetSectionBase().X + ComponentSizeQuads, ValidX2);
 						ValidY1 = FMath::Min<int32>(Component->GetSectionBase().Y, ValidY1);
-						ValidY2 = FMath::Max<int32>(Component->GetSectionBase().Y+ComponentSizeQuads, ValidY2);
+						ValidY2 = FMath::Max<int32>(Component->GetSectionBase().Y + ComponentSizeQuads, ValidY2);
 					}
 				}
 			}
@@ -125,10 +125,10 @@ public:
 
 		const float AreaResolution = LANDSCAPE_XYOFFSET_SCALE; //1.f/256.f;
 
-		Cache.CacheData(X1,Y1,X2,Y2);
+		Cache.CacheData(X1, Y1, X2, Y2);
 
 		TArray<FVector> XYOffsetVectorData;
-		Cache.GetCachedData(X1,Y1,X2,Y2,XYOffsetVectorData);
+		Cache.GetCachedData(X1, Y1, X2, Y2, XYOffsetVectorData);
 		TArray<FVector> NewXYOffset;
 		NewXYOffset = XYOffsetVectorData;
 
@@ -140,13 +140,13 @@ public:
 			const int32 MaxIterNum = 300;
 
 			TArray<int32> QuadX, QuadY, MinX, MaxX, MinY, MaxY;
-			QuadX.AddZeroed( X2 - X1 );
-			QuadY.AddZeroed( Y2 - Y1 );
+			QuadX.AddZeroed(X2 - X1);
+			QuadY.AddZeroed(Y2 - Y1);
 
-			MinX.Empty( Y2 - Y1 + 1 );
-			MaxX.Empty( Y2 - Y1 + 1 );
-			MinY.Empty( X2 - X1 + 1 );
-			MaxY.Empty( X2 - X1 + 1 );
+			MinX.Empty(Y2 - Y1 + 1);
+			MaxX.Empty(Y2 - Y1 + 1);
+			MinY.Empty(X2 - X1 + 1);
+			MaxY.Empty(X2 - X1 + 1);
 			for (int32 X = X1; X <= X2; ++X)
 			{
 				MinY.Add(INT_MAX);
@@ -160,33 +160,33 @@ public:
 
 			// Calculate Average...
 			TArray<ULandscapeComponent*> ComponentArray; // Ptr to component
-			ComponentArray.AddZeroed( (X2-X1+1)*(Y2-Y1+1) );
+			ComponentArray.AddZeroed((X2 - X1 + 1)*(Y2 - Y1 + 1));
 			int32 ComponentIndexX1, ComponentIndexY1, ComponentIndexX2, ComponentIndexY2;
 			int32 ComponentSizeQuads = LandscapeInfo->ComponentSizeQuads;
 			ALandscape::CalcComponentIndicesOverlap(X1, Y1, X2, Y2, ComponentSizeQuads, ComponentIndexX1, ComponentIndexY1, ComponentIndexX2, ComponentIndexY2);
 
-			for( int32 ComponentIndexY=ComponentIndexY1;ComponentIndexY<=ComponentIndexY2;ComponentIndexY++ )
+			for (int32 ComponentIndexY = ComponentIndexY1; ComponentIndexY <= ComponentIndexY2; ComponentIndexY++)
 			{
-				for( int32 ComponentIndexX=ComponentIndexX1;ComponentIndexX<=ComponentIndexX2;ComponentIndexX++ )
-				{		
-					ULandscapeComponent* Comp = LandscapeInfo->XYtoComponentMap.FindRef(FIntPoint(ComponentIndexX,ComponentIndexY));
+				for (int32 ComponentIndexX = ComponentIndexX1; ComponentIndexX <= ComponentIndexX2; ComponentIndexX++)
+				{
+					ULandscapeComponent* Comp = LandscapeInfo->XYtoComponentMap.FindRef(FIntPoint(ComponentIndexX, ComponentIndexY));
 
 					if (Comp)
 					{
 						const FMatrix LocalToWorld = Comp->GetRenderMatrix();
 
 						// Find coordinates of box that lies inside component
-						int32 ComponentX1 = FMath::Clamp<int32>(X1-ComponentIndexX*ComponentSizeQuads, 0, ComponentSizeQuads);
-						int32 ComponentY1 = FMath::Clamp<int32>(Y1-ComponentIndexY*ComponentSizeQuads, 0, ComponentSizeQuads);
-						int32 ComponentX2 = FMath::Clamp<int32>(X2-ComponentIndexX*ComponentSizeQuads, 0, ComponentSizeQuads);
-						int32 ComponentY2 = FMath::Clamp<int32>(Y2-ComponentIndexY*ComponentSizeQuads, 0, ComponentSizeQuads);
+						int32 ComponentX1 = FMath::Clamp<int32>(X1 - ComponentIndexX*ComponentSizeQuads, 0, ComponentSizeQuads);
+						int32 ComponentY1 = FMath::Clamp<int32>(Y1 - ComponentIndexY*ComponentSizeQuads, 0, ComponentSizeQuads);
+						int32 ComponentX2 = FMath::Clamp<int32>(X2 - ComponentIndexX*ComponentSizeQuads, 0, ComponentSizeQuads);
+						int32 ComponentY2 = FMath::Clamp<int32>(Y2 - ComponentIndexY*ComponentSizeQuads, 0, ComponentSizeQuads);
 
 						// World space area calculation
-						for (int32 Y = ComponentY1; Y <= ComponentY2; ++Y )
+						for (int32 Y = ComponentY1; Y <= ComponentY2; ++Y)
 						{
-							for (int32 X = ComponentX1; X <= ComponentX2; ++X )
+							for (int32 X = ComponentX1; X <= ComponentX2; ++X)
 							{
-								if ( X < ComponentX2 && Y < ComponentY2)
+								if (X < ComponentX2 && Y < ComponentY2)
 								{
 									// Need to read XY Offset value from XYOffsetTexture before this
 									FVector P[4];
@@ -194,24 +194,24 @@ public:
 									{
 										int32 XX = X + XOffsets[Idx];
 										int32 YY = Y + YOffsets[Idx];
-										P[Idx] = FVector(GetWorldPos(LocalToWorld, FVector2D(XX, YY), 
-											XYOffsetVectorData[ (ComponentIndexX*ComponentSizeQuads + XX - X1) + (ComponentIndexY*ComponentSizeQuads + YY - Y1)*(X2-X1+1) ]
+										P[Idx] = FVector(GetWorldPos(LocalToWorld, FVector2D(XX, YY),
+											XYOffsetVectorData[(ComponentIndexX*ComponentSizeQuads + XX - X1) + (ComponentIndexY*ComponentSizeQuads + YY - Y1)*(X2 - X1 + 1)]
 											));
 									}
 
 									TotalArea += (((P[3] - P[0]) ^ (P[1] - P[0])).Size() + ((P[3] - P[0]) ^ (P[2] - P[0])).Size()) * 0.5f;
 									QuadNum++;
-									QuadX[ ComponentIndexX*ComponentSizeQuads + X - X1 ]++;
-									QuadY[ ComponentIndexY*ComponentSizeQuads + Y - Y1 ]++;
+									QuadX[ComponentIndexX*ComponentSizeQuads + X - X1]++;
+									QuadY[ComponentIndexY*ComponentSizeQuads + Y - Y1]++;
 
 									// Mark valid quad position
-									ComponentArray[ (ComponentIndexX*ComponentSizeQuads) + X - X1 + (ComponentIndexY*ComponentSizeQuads + Y - Y1)*(X2-X1+1) ] = Comp;
+									ComponentArray[(ComponentIndexX*ComponentSizeQuads) + X - X1 + (ComponentIndexY*ComponentSizeQuads + Y - Y1)*(X2 - X1 + 1)] = Comp;
 								}
 
-								MinX[ComponentIndexY*ComponentSizeQuads+Y-Y1] = FMath::Min<int32>(MinX[ComponentIndexY*ComponentSizeQuads+Y-Y1], ComponentIndexX*ComponentSizeQuads + X);
-								MaxX[ComponentIndexY*ComponentSizeQuads+Y-Y1] = FMath::Max<int32>(MaxX[ComponentIndexY*ComponentSizeQuads+Y-Y1], ComponentIndexX*ComponentSizeQuads + X);
-								MinY[ComponentIndexX*ComponentSizeQuads+X-X1] = FMath::Min<int32>(MinY[ComponentIndexX*ComponentSizeQuads+X-X1], ComponentIndexY*ComponentSizeQuads + Y);
-								MaxY[ComponentIndexX*ComponentSizeQuads+X-X1] = FMath::Max<int32>(MaxY[ComponentIndexX*ComponentSizeQuads+X-X1], ComponentIndexY*ComponentSizeQuads + Y);
+								MinX[ComponentIndexY*ComponentSizeQuads + Y - Y1] = FMath::Min<int32>(MinX[ComponentIndexY*ComponentSizeQuads + Y - Y1], ComponentIndexX*ComponentSizeQuads + X);
+								MaxX[ComponentIndexY*ComponentSizeQuads + Y - Y1] = FMath::Max<int32>(MaxX[ComponentIndexY*ComponentSizeQuads + Y - Y1], ComponentIndexX*ComponentSizeQuads + X);
+								MinY[ComponentIndexX*ComponentSizeQuads + X - X1] = FMath::Min<int32>(MinY[ComponentIndexX*ComponentSizeQuads + X - X1], ComponentIndexY*ComponentSizeQuads + Y);
+								MaxY[ComponentIndexX*ComponentSizeQuads + X - X1] = FMath::Max<int32>(MaxY[ComponentIndexX*ComponentSizeQuads + X - X1], ComponentIndexY*ComponentSizeQuads + Y);
 							}
 						}
 					}
@@ -224,11 +224,11 @@ public:
 			float RemainArea = TotalArea;
 			int32 RemainQuads = QuadNum;
 
-			for (int32 Y = Y1; Y < Y2-1; ++Y)
+			for (int32 Y = Y1; Y < Y2 - 1; ++Y)
 			{
 				// Like rasterization style
 				// Search for Y offset
-				if (MinX[Y-Y1] > MaxX[Y-Y1])
+				if (MinX[Y - Y1] > MaxX[Y - Y1])
 				{
 					continue;
 				}
@@ -237,15 +237,15 @@ public:
 				float AreaBaseError = AverageArea * 0.5f;
 				float TotalLineArea = 0.f;
 				float TargetLineArea = AverageArea * QuadY[Y - Y1];
-				float YOffset = Y+1, PreYOffset = Y+1; // Need to be bigger than previous Y
+				float YOffset = Y + 1, PreYOffset = Y + 1; // Need to be bigger than previous Y
 				float StepSize = FPlatformMath::Sqrt(2) * 0.25f;
 				float LineAreaDiff = FLT_MAX; //Abs(TotalLineArea - TargetLineArea);
 				int32 IterNum = 0;
 				float TotalHeightError = 0.f;
 
-				while( NewXYOffset[ (Y - Y1) * (X2-X1+1) ].Y + Y > XYOffsetVectorData[(FPlatformMath::FloorToInt(YOffset) - Y1) * (X2-X1+1) ].Y + FPlatformMath::FloorToInt(YOffset) )
+				while (NewXYOffset[(Y - Y1) * (X2 - X1 + 1)].Y + Y > XYOffsetVectorData[(FPlatformMath::FloorToInt(YOffset) - Y1) * (X2 - X1 + 1)].Y + FPlatformMath::FloorToInt(YOffset))
 				{
-					YOffset = YOffset+1.f;
+					YOffset = YOffset + 1.f;
 					if (YOffset >= Y2)
 					{
 						YOffset = Y2;
@@ -260,9 +260,9 @@ public:
 					TotalLineArea = 0.f;
 					TotalHeightError = 0.f;
 					//for (int32 X = X1; X < X2; ++X)
-					for (int32 X = MinX[Y-Y1]; X < MaxX[Y-Y1]; ++X)
+					for (int32 X = MinX[Y - Y1]; X < MaxX[Y - Y1]; ++X)
 					{
-						ULandscapeComponent* Comp = ComponentArray[X - X1 + (Y - Y1)*(X2-X1+1) ];
+						ULandscapeComponent* Comp = ComponentArray[X - X1 + (Y - Y1)*(X2 - X1 + 1)];
 						if (Comp != NULL) // valid
 						{
 							const FMatrix LocalToWorld = Comp->GetRenderMatrix();
@@ -270,23 +270,23 @@ public:
 							for (int32 Idx = 0; Idx < 2; ++Idx)
 							{
 								int32 XX = FMath::Clamp<int32>(X + XOffsets[Idx], X1, X2);
-								P[Idx] = FVector( GetWorldPos( LocalToWorld, FVector2D(XX - Comp->GetSectionBase().X, Y - Comp->GetSectionBase().Y), NewXYOffset[ XX - X1 + (Y - Y1) * (X2-X1+1) ] ) );
+								P[Idx] = FVector(GetWorldPos(LocalToWorld, FVector2D(XX - Comp->GetSectionBase().X, Y - Comp->GetSectionBase().Y), NewXYOffset[XX - X1 + (Y - Y1) * (X2 - X1 + 1)]));
 							}
 
-							int32 YY0 = FMath::Clamp<int32>(FMath::FloorToInt(YOffset-1), Y1, Y2);
+							int32 YY0 = FMath::Clamp<int32>(FMath::FloorToInt(YOffset - 1), Y1, Y2);
 							int32 YY1 = FMath::Clamp<int32>(FMath::FloorToInt(YOffset), Y1, Y2);
-							int32 YY2 = FMath::Clamp<int32>(FMath::FloorToInt(1+YOffset), Y1, Y2);
+							int32 YY2 = FMath::Clamp<int32>(FMath::FloorToInt(1 + YOffset), Y1, Y2);
 							// Search for valid YOffset...
 							for (int32 Idx = 2; Idx < 4; ++Idx)
 							{
 								int32 XX = FMath::Clamp<int32>(X + XOffsets[Idx], X1, X2);
-								FVector P1( GetWorldPos( LocalToWorld, FVector2D(XX - Comp->GetSectionBase().X, YY1 - Comp->GetSectionBase().Y), XYOffsetVectorData[ XX - X1 + (YY1 - Y1) * (X2-X1+1) ] ) );
-								FVector P2( GetWorldPos( LocalToWorld, FVector2D(XX - Comp->GetSectionBase().X, YY2 - Comp->GetSectionBase().Y), XYOffsetVectorData[ XX - X1 + (YY2 - Y1) * (X2-X1+1) ] ) );
+								FVector P1(GetWorldPos(LocalToWorld, FVector2D(XX - Comp->GetSectionBase().X, YY1 - Comp->GetSectionBase().Y), XYOffsetVectorData[XX - X1 + (YY1 - Y1) * (X2 - X1 + 1)]));
+								FVector P2(GetWorldPos(LocalToWorld, FVector2D(XX - Comp->GetSectionBase().X, YY2 - Comp->GetSectionBase().Y), XYOffsetVectorData[XX - X1 + (YY2 - Y1) * (X2 - X1 + 1)]));
 								P[Idx] = FMath::Lerp(P1, P2, FMath::Fractional(YOffset));
 								if (Idx == 2)
 								{
-									FVector P0( GetWorldPos( LocalToWorld, FVector2D(XX - Comp->GetSectionBase().X, YY0 - Comp->GetSectionBase().Y), XYOffsetVectorData[ XX - X1 + (YY0 - Y1) * (X2-X1+1) ] ) );
-									TotalHeightError += FMath::Abs( ((P[2] - P0) ^ (P2 - P[2])).Size() - ((P1 - P0) ^ (P2 - P1)).Size() );
+									FVector P0(GetWorldPos(LocalToWorld, FVector2D(XX - Comp->GetSectionBase().X, YY0 - Comp->GetSectionBase().Y), XYOffsetVectorData[XX - X1 + (YY0 - Y1) * (X2 - X1 + 1)]));
+									TotalHeightError += FMath::Abs(((P[2] - P0) ^ (P2 - P[2])).Size() - ((P1 - P0) ^ (P2 - P1)).Size());
 								}
 							}
 
@@ -300,9 +300,9 @@ public:
 						break;
 					}
 
-					if (MaxX[Y-Y1] - MinX[Y-Y1] > 0)
+					if (MaxX[Y - Y1] - MinX[Y - Y1] > 0)
 					{
-						TotalHeightError /= (MaxX[Y-Y1] - MinX[Y-Y1]);
+						TotalHeightError /= (MaxX[Y - Y1] - MinX[Y - Y1]);
 					}
 
 					float NewLineAreaDiff = FMath::Abs(TotalLineArea - TargetLineArea);
@@ -325,12 +325,12 @@ public:
 							YOffset += StepSize;
 						}
 						// clamp
-						if ( YOffset < Y1 )
+						if (YOffset < Y1)
 						{
 							YOffset = Y1;
 							break;
 						}
-						if ( YOffset >= Y2 )
+						if (YOffset >= Y2)
 						{
 							YOffset = Y2;
 							break;
@@ -348,20 +348,20 @@ public:
 				{
 					RemainArea -= TotalLineArea;
 					RemainQuads -= QuadY[Y - Y1];
-					for (int32 X = MinX[Y-Y1]; X < MaxX[Y-Y1]; ++X)
+					for (int32 X = MinX[Y - Y1]; X < MaxX[Y - Y1]; ++X)
 					{
 						int32 YY1 = FMath::Clamp<int32>(FMath::FloorToInt(YOffset), Y1, Y2);
-						int32 YY2 = FMath::Clamp<int32>(FMath::FloorToInt(1+YOffset), Y1, Y2);
-						FVector P1 = XYOffsetVectorData[ X - X1 + (YY1 - Y1) * (X2-X1+1) ];
+						int32 YY2 = FMath::Clamp<int32>(FMath::FloorToInt(1 + YOffset), Y1, Y2);
+						FVector P1 = XYOffsetVectorData[X - X1 + (YY1 - Y1) * (X2 - X1 + 1)];
 						//P1.X = X+P1.X;
-						P1.Y = YY1+P1.Y;
-						FVector P2 = XYOffsetVectorData[ X - X1 + (YY2 - Y1) * (X2-X1+1) ];
+						P1.Y = YY1 + P1.Y;
+						FVector P2 = XYOffsetVectorData[X - X1 + (YY2 - Y1) * (X2 - X1 + 1)];
 						//P2.X = X+P2.X;
-						P2.Y = YY2+P2.Y;
-						FVector& XYOffset =	NewXYOffset[ X - X1 + (Y+1 - Y1) * (X2-X1+1) ];
+						P2.Y = YY2 + P2.Y;
+						FVector& XYOffset = NewXYOffset[X - X1 + (Y + 1 - Y1) * (X2 - X1 + 1)];
 						XYOffset = FMath::Lerp(P1, P2, FMath::Fractional(YOffset));
 						//XYOffset.X -= X;
-						XYOffset.Y -= Y+1;
+						XYOffset.Y -= Y + 1;
 					}
 				}
 			}
@@ -371,11 +371,11 @@ public:
 			RemainArea = TotalArea;
 			RemainQuads = QuadNum;
 
-			for (int32 X = X1; X < X2-1; ++X)
+			for (int32 X = X1; X < X2 - 1; ++X)
 			{
 				// Like rasterization style
 				// Search for X offset
-				if (MinY[X-X1] > MaxY[X-X1])
+				if (MinY[X - X1] > MaxY[X - X1])
 				{
 					continue;
 				}
@@ -384,15 +384,15 @@ public:
 				float AreaBaseError = AverageArea * 0.5f;
 				float TotalLineArea = 0.f;
 				float TargetLineArea = AverageArea * QuadX[X - X1];
-				float XOffset = X+1, PreXOffset = X+1; // Need to be bigger than previous Y
+				float XOffset = X + 1, PreXOffset = X + 1; // Need to be bigger than previous Y
 				float StepSize = FMath::Sqrt(2) * 0.25f;
 				float LineAreaDiff = FLT_MAX; // Abs(TotalLineArea - TargetLineArea);
 				int32 IterNum = 0;
 				float TotalHeightError = 0.f;
 
-				while( NewXYOffset[X - X1].X + X > NewYOffsets[FMath::FloorToInt(XOffset) - X1].X + FMath::FloorToFloat(XOffset) )
+				while (NewXYOffset[X - X1].X + X > NewYOffsets[FMath::FloorToInt(XOffset) - X1].X + FMath::FloorToFloat(XOffset))
 				{
-					XOffset = XOffset+1.f;
+					XOffset = XOffset + 1.f;
 					if (XOffset >= X2)
 					{
 						XOffset = X2;
@@ -405,34 +405,34 @@ public:
 				{
 					TotalLineArea = 0.f;
 					IterNum++;
-					for (int32 Y = MinY[X-X1]; Y < MaxY[X-X1]; ++Y)
+					for (int32 Y = MinY[X - X1]; Y < MaxY[X - X1]; ++Y)
 					{
-						ULandscapeComponent* Comp = ComponentArray[X - X1 + (Y - Y1)*(X2-X1+1) ];
+						ULandscapeComponent* Comp = ComponentArray[X - X1 + (Y - Y1)*(X2 - X1 + 1)];
 						if (Comp != NULL) // valid
 						{
 							const FMatrix LocalToWorld = Comp->GetRenderMatrix();
 							FVector P[4];
-							for (int32 Idx = 0; Idx < 4; Idx+=2)
+							for (int32 Idx = 0; Idx < 4; Idx += 2)
 							{
 								int32 YY = FMath::Clamp<int32>(Y + YOffsets[Idx], Y1, Y2);
-								P[Idx] = FVector( GetWorldPos( LocalToWorld, FVector2D(X - Comp->GetSectionBase().X, YY - Comp->GetSectionBase().Y), NewXYOffset[ X - X1 + (YY - Y1) * (X2-X1+1) ] ) );
+								P[Idx] = FVector(GetWorldPos(LocalToWorld, FVector2D(X - Comp->GetSectionBase().X, YY - Comp->GetSectionBase().Y), NewXYOffset[X - X1 + (YY - Y1) * (X2 - X1 + 1)]));
 							}
 
-							int32 XX0 = FMath::Clamp<int32>(FMath::FloorToInt(XOffset-1), X1, X2);
+							int32 XX0 = FMath::Clamp<int32>(FMath::FloorToInt(XOffset - 1), X1, X2);
 							int32 XX1 = FMath::Clamp<int32>(FMath::FloorToInt(XOffset), X1, X2);
-							int32 XX2 = FMath::Clamp<int32>(FMath::FloorToInt(1+XOffset), X1, X2);
+							int32 XX2 = FMath::Clamp<int32>(FMath::FloorToInt(1 + XOffset), X1, X2);
 
 							// Search for valid YOffset...
-							for (int32 Idx = 1; Idx < 4; Idx+=2)
+							for (int32 Idx = 1; Idx < 4; Idx += 2)
 							{
 								int32 YY = FMath::Clamp<int32>(Y + YOffsets[Idx], Y1, Y2);
-								FVector P1( GetWorldPos( LocalToWorld, FVector2D(XX1 - Comp->GetSectionBase().X, YY - Comp->GetSectionBase().Y), NewYOffsets[ XX1 - X1 + (YY - Y1) * (X2-X1+1) ] ) );
-								FVector P2( GetWorldPos( LocalToWorld, FVector2D(XX2 - Comp->GetSectionBase().X, YY - Comp->GetSectionBase().Y), NewYOffsets[ XX2 - X1 + (YY - Y1) * (X2-X1+1) ] ) );
+								FVector P1(GetWorldPos(LocalToWorld, FVector2D(XX1 - Comp->GetSectionBase().X, YY - Comp->GetSectionBase().Y), NewYOffsets[XX1 - X1 + (YY - Y1) * (X2 - X1 + 1)]));
+								FVector P2(GetWorldPos(LocalToWorld, FVector2D(XX2 - Comp->GetSectionBase().X, YY - Comp->GetSectionBase().Y), NewYOffsets[XX2 - X1 + (YY - Y1) * (X2 - X1 + 1)]));
 								P[Idx] = FMath::Lerp(P1, P2, FMath::Fractional(XOffset));
 								if (Idx == 1)
 								{
-									FVector P0( GetWorldPos( LocalToWorld, FVector2D(XX0 - Comp->GetSectionBase().X, YY - Comp->GetSectionBase().Y), NewYOffsets[ XX0 - X1 + (YY - Y1) * (X2-X1+1) ] ) );
-									TotalHeightError += FMath::Abs( ((P[1] - P0) ^ (P2 - P[1])).Size() - ((P1 - P0) ^ (P2 - P1)).Size() );
+									FVector P0(GetWorldPos(LocalToWorld, FVector2D(XX0 - Comp->GetSectionBase().X, YY - Comp->GetSectionBase().Y), NewYOffsets[XX0 - X1 + (YY - Y1) * (X2 - X1 + 1)]));
+									TotalHeightError += FMath::Abs(((P[1] - P0) ^ (P2 - P[1])).Size() - ((P1 - P0) ^ (P2 - P1)).Size());
 								}
 							}
 
@@ -441,14 +441,14 @@ public:
 						}
 					}
 
-					if (TotalLineArea < AreaErrorThreshold || IterNum > MaxIterNum) 
+					if (TotalLineArea < AreaErrorThreshold || IterNum > MaxIterNum)
 					{
 						break;
 					}
 
-					if (MaxY[X-X1] - MinY[X-X1] > 0)
+					if (MaxY[X - X1] - MinY[X - X1] > 0)
 					{
-						TotalHeightError /= (MaxY[X-X1] - MinY[X-X1]);
+						TotalHeightError /= (MaxY[X - X1] - MinY[X - X1]);
 					}
 
 					float NewLineAreaDiff = FMath::Abs(TotalLineArea - TargetLineArea);
@@ -471,12 +471,12 @@ public:
 							XOffset += StepSize;
 						}
 						// clamp
-						if ( XOffset <= X1 )
+						if (XOffset <= X1)
 						{
 							XOffset = X1;
 							break;
 						}
-						else if ( XOffset >= X2 )
+						else if (XOffset >= X2)
 						{
 							XOffset = X2;
 							break;
@@ -495,17 +495,17 @@ public:
 					RemainArea -= TotalLineArea;
 					RemainQuads -= QuadX[X - X1];
 
-					for (int32 Y = MinY[X-X1]; Y < MaxY[X-X1]; ++Y)
+					for (int32 Y = MinY[X - X1]; Y < MaxY[X - X1]; ++Y)
 					{
 						int32 XX1 = FMath::Clamp<int32>(FMath::FloorToInt(XOffset), X1, X2);
-						int32 XX2 = FMath::Clamp<int32>(FMath::FloorToInt(1+XOffset), X1, X2);
-						FVector P1 = NewYOffsets[ XX1 - X1 + (Y - Y1) * (X2-X1+1) ];
-						P1.X = XX1+P1.X;
-						FVector P2 = NewYOffsets[ XX2 - X1 + (Y - Y1) * (X2-X1+1) ];
-						P2.X = XX2+P2.X;
-						FVector& XYOffset =	NewXYOffset[ X+1 - X1 + (Y - Y1) * (X2-X1+1) ];
+						int32 XX2 = FMath::Clamp<int32>(FMath::FloorToInt(1 + XOffset), X1, X2);
+						FVector P1 = NewYOffsets[XX1 - X1 + (Y - Y1) * (X2 - X1 + 1)];
+						P1.X = XX1 + P1.X;
+						FVector P2 = NewYOffsets[XX2 - X1 + (Y - Y1) * (X2 - X1 + 1)];
+						P2.X = XX2 + P2.X;
+						FVector& XYOffset = NewXYOffset[X + 1 - X1 + (Y - Y1) * (X2 - X1 + 1)];
 						XYOffset = FMath::Lerp(P1, P2, FMath::Fractional(XOffset));
-						XYOffset.X -= X+1;
+						XYOffset.X -= X + 1;
 					}
 				}
 			}
@@ -518,20 +518,20 @@ public:
 		float H = Y2 - Y1 + 1;
 		float FalloffRadius = W * 0.5f * UISettings->BrushFalloff;
 		float SquareRadius = W * 0.5f - FalloffRadius;
-		for (int32 Y = 0; Y <= Y2-Y1; ++Y)
+		for (int32 Y = 0; Y <= Y2 - Y1; ++Y)
 		{
-			for (int32 X = 0; X <= X2-X1; ++X)
+			for (int32 X = 0; X <= X2 - X1; ++X)
 			{
-				int32 Index = X + Y * (X2-X1+1);
+				int32 Index = X + Y * (X2 - X1 + 1);
 				FVector2D TransformedLocal(FMath::Abs(X - W * 0.5f), FMath::Abs(Y - H * 0.5f) * (W / H));
 				float Cos = FMath::Abs(TransformedLocal.X) / TransformedLocal.Size();
 				float Sin = FMath::Abs(TransformedLocal.Y) / TransformedLocal.Size();
 				float RatioX = FalloffRadius > 0.f ? 1.f - FMath::Clamp((FMath::Abs(TransformedLocal.X) - Cos*SquareRadius) / FalloffRadius, 0.f, 1.f) : 1.f;
 				float RatioY = FalloffRadius > 0.f ? 1.f - FMath::Clamp((FMath::Abs(TransformedLocal.Y) - Sin*SquareRadius) / FalloffRadius, 0.f, 1.f) : 1.f;
 				float Ratio = TransformedLocal.Size() > SquareRadius ? RatioX * RatioY : 1.f; //TransformedLocal.X / LW * TransformedLocal.Y / LW;
-				float PaintAmount = Ratio*Ratio*(3-2*Ratio);
+				float PaintAmount = Ratio*Ratio*(3 - 2 * Ratio);
 
-				XYOffsetVectorData[Index] = FMath::Lerp( XYOffsetVectorData[Index], NewXYOffset[Index], PaintAmount );
+				XYOffsetVectorData[Index] = FMath::Lerp(XYOffsetVectorData[Index], NewXYOffset[Index], PaintAmount);
 				//XYOffsetVectorData(Index) = NewXYOffset(Index);
 			}
 		}
@@ -542,32 +542,32 @@ public:
 	}
 
 protected:
-	class ULandscapeInfo* LandscapeInfo;
+	ULandscapeInfo* LandscapeInfo;
 	FLandscapeXYOffsetCache Cache;
 };
 
 class FLandscapeToolRetopologize : public FLandscapeToolBase<FLandscapeToolStrokeRetopologize>
 {
 public:
-	FLandscapeToolRetopologize(class FEdModeLandscape* InEdMode)
+	FLandscapeToolRetopologize(FEdModeLandscape* InEdMode)
 		: FLandscapeToolBase<FLandscapeToolStrokeRetopologize>(InEdMode)
 	{}
 
 	virtual const TCHAR* GetToolName() override { return TEXT("Retopologize"); }
 	virtual FText GetDisplayName() override { return NSLOCTEXT("UnrealEd", "LandscapeMode_Retopologize", "Retopologize"); }
 
-	virtual bool IsValidForTarget(const FLandscapeToolTarget& Target)
+	virtual ELandscapeToolTargetTypeMask::Type GetSupportedTargetTypes() override
 	{
-		return true; // erosion applied to all...
+		// technically not entirely accurate, also modifies the XYOffset map
+		return ELandscapeToolTargetTypeMask::Heightmap;
 	}
 };
 
-void FEdModeLandscape::IntializeToolSet_Retopologize()
+void FEdModeLandscape::InitializeTool_Retopologize()
 {
-	FLandscapeToolSet* ToolSet_Retopologize = new(LandscapeToolSets) FLandscapeToolSet(TEXT("ToolSet_Retopologize"));
-	ToolSet_Retopologize->AddTool(new FLandscapeToolRetopologize(this));
-
-	ToolSet_Retopologize->ValidBrushes.Add("BrushSet_Circle");
-	ToolSet_Retopologize->ValidBrushes.Add("BrushSet_Alpha");
-	ToolSet_Retopologize->ValidBrushes.Add("BrushSet_Pattern");
+	auto Tool_Retopologize = MakeUnique<FLandscapeToolRetopologize>(this);
+	Tool_Retopologize->ValidBrushes.Add("BrushSet_Circle");
+	Tool_Retopologize->ValidBrushes.Add("BrushSet_Alpha");
+	Tool_Retopologize->ValidBrushes.Add("BrushSet_Pattern");
+	LandscapeTools.Add(MoveTemp(Tool_Retopologize));
 }
diff --git a/Engine/Source/Editor/LandscapeEditor/Private/LandscapeEditorCommands.cpp b/Engine/Source/Editor/LandscapeEditor/Private/LandscapeEditorCommands.cpp
--- a/Engine/Source/Editor/LandscapeEditor/Private/LandscapeEditorCommands.cpp
+++ b/Engine/Source/Editor/LandscapeEditor/Private/LandscapeEditorCommands.cpp
@@ -6,80 +6,80 @@
 
 void FLandscapeEditorCommands::RegisterCommands()
 {
-	UI_COMMAND( ManageMode, "Mode - Manage", "", EUserInterfaceActionType::RadioButton, FInputGesture() );
+	UI_COMMAND(ManageMode, "Mode - Manage", "", EUserInterfaceActionType::RadioButton, FInputGesture());
 	NameToCommandMap.Add("ToolMode_Manage", ManageMode);
-	UI_COMMAND( SculptMode, "Mode - Sculpt", "", EUserInterfaceActionType::RadioButton, FInputGesture() );
+	UI_COMMAND(SculptMode, "Mode - Sculpt", "", EUserInterfaceActionType::RadioButton, FInputGesture());
 	NameToCommandMap.Add("ToolMode_Sculpt", SculptMode);
-	UI_COMMAND( PaintMode,  "Mode - Paint",  "", EUserInterfaceActionType::RadioButton, FInputGesture() );
+	UI_COMMAND(PaintMode, "Mode - Paint", "", EUserInterfaceActionType::RadioButton, FInputGesture());
 	NameToCommandMap.Add("ToolMode_Paint", PaintMode);
 
-	UI_COMMAND( NewLandscape, "Tool - New Landscape", "", EUserInterfaceActionType::RadioButton, FInputGesture() );
-	NameToCommandMap.Add("ToolSet_NewLandscape", NewLandscape);
+	UI_COMMAND(NewLandscape, "Tool - New Landscape", "", EUserInterfaceActionType::RadioButton, FInputGesture());
+	NameToCommandMap.Add("Tool_NewLandscape", NewLandscape);
 
-	UI_COMMAND( ResizeLandscape, "Tool - Change Component Size", "", EUserInterfaceActionType::RadioButton, FInputGesture() );
-	NameToCommandMap.Add("ToolSet_ResizeLandscape", ResizeLandscape);
+	UI_COMMAND(ResizeLandscape, "Tool - Change Component Size", "", EUserInterfaceActionType::RadioButton, FInputGesture());
+	NameToCommandMap.Add("Tool_ResizeLandscape", ResizeLandscape);
 
-	UI_COMMAND( SculptTool, "Tool - Sculpt", "", EUserInterfaceActionType::RadioButton, FInputGesture() );
-	NameToCommandMap.Add("ToolSet_Sculpt", SculptTool);
-	UI_COMMAND( PaintTool, "Tool - Paint", "", EUserInterfaceActionType::RadioButton, FInputGesture() );
-	NameToCommandMap.Add("ToolSet_Paint", PaintTool);
-	UI_COMMAND( SmoothTool, "Tool - Smooth", "", EUserInterfaceActionType::RadioButton, FInputGesture() );
-	NameToCommandMap.Add("ToolSet_Smooth", SmoothTool);
-	UI_COMMAND( FlattenTool, "Tool - Flatten", "", EUserInterfaceActionType::RadioButton, FInputGesture() );
-	NameToCommandMap.Add("ToolSet_Flatten", FlattenTool);
-	UI_COMMAND( RampTool, "Tool - Ramp", "", EUserInterfaceActionType::RadioButton, FInputGesture() );
-	NameToCommandMap.Add("ToolSet_Ramp", RampTool);
-	UI_COMMAND( ErosionTool, "Tool - Erosion", "", EUserInterfaceActionType::RadioButton, FInputGesture() );
-	NameToCommandMap.Add("ToolSet_Erosion", ErosionTool);
-	UI_COMMAND( HydroErosionTool, "Tool - Hydraulic Erosion", "", EUserInterfaceActionType::RadioButton, FInputGesture() );
-	NameToCommandMap.Add("ToolSet_HydraErosion", HydroErosionTool);
-	UI_COMMAND( NoiseTool, "Tool - Noise", "", EUserInterfaceActionType::RadioButton, FInputGesture() );
-	NameToCommandMap.Add("ToolSet_Noise", NoiseTool);
-	UI_COMMAND( RetopologizeTool, "Tool - Retopologize", "", EUserInterfaceActionType::RadioButton, FInputGesture() );
-	NameToCommandMap.Add("ToolSet_Retopologize", RetopologizeTool);
-	UI_COMMAND( VisibilityTool, "Tool - Visibility", "", EUserInterfaceActionType::RadioButton, FInputGesture() );
-	NameToCommandMap.Add("ToolSet_Visibility", VisibilityTool);
+	UI_COMMAND(SculptTool, "Tool - Sculpt", "", EUserInterfaceActionType::RadioButton, FInputGesture());
+	NameToCommandMap.Add("Tool_Sculpt", SculptTool);
+	UI_COMMAND(PaintTool, "Tool - Paint", "", EUserInterfaceActionType::RadioButton, FInputGesture());
+	NameToCommandMap.Add("Tool_Paint", PaintTool);
+	UI_COMMAND(SmoothTool, "Tool - Smooth", "", EUserInterfaceActionType::RadioButton, FInputGesture());
+	NameToCommandMap.Add("Tool_Smooth", SmoothTool);
+	UI_COMMAND(FlattenTool, "Tool - Flatten", "", EUserInterfaceActionType::RadioButton, FInputGesture());
+	NameToCommandMap.Add("Tool_Flatten", FlattenTool);
+	UI_COMMAND(RampTool, "Tool - Ramp", "", EUserInterfaceActionType::RadioButton, FInputGesture());
+	NameToCommandMap.Add("Tool_Ramp", RampTool);
+	UI_COMMAND(ErosionTool, "Tool - Erosion", "", EUserInterfaceActionType::RadioButton, FInputGesture());
+	NameToCommandMap.Add("Tool_Erosion", ErosionTool);
+	UI_COMMAND(HydroErosionTool, "Tool - Hydraulic Erosion", "", EUserInterfaceActionType::RadioButton, FInputGesture());
+	NameToCommandMap.Add("Tool_HydraErosion", HydroErosionTool);
+	UI_COMMAND(NoiseTool, "Tool - Noise", "", EUserInterfaceActionType::RadioButton, FInputGesture());
+	NameToCommandMap.Add("Tool_Noise", NoiseTool);
+	UI_COMMAND(RetopologizeTool, "Tool - Retopologize", "", EUserInterfaceActionType::RadioButton, FInputGesture());
+	NameToCommandMap.Add("Tool_Retopologize", RetopologizeTool);
+	UI_COMMAND(VisibilityTool, "Tool - Visibility", "", EUserInterfaceActionType::RadioButton, FInputGesture());
+	NameToCommandMap.Add("Tool_Visibility", VisibilityTool);
 
-	UI_COMMAND( SelectComponentTool, "Tool - Component Selection", "", EUserInterfaceActionType::RadioButton, FInputGesture() );
-	NameToCommandMap.Add("ToolSet_Select", SelectComponentTool);
-	UI_COMMAND( AddComponentTool, "Tool - Add Components", "", EUserInterfaceActionType::RadioButton, FInputGesture() );
-	NameToCommandMap.Add("ToolSet_AddComponent", AddComponentTool);
-	UI_COMMAND( DeleteComponentTool, "Tool - Delete Components", "", EUserInterfaceActionType::RadioButton, FInputGesture() );
-	NameToCommandMap.Add("ToolSet_DeleteComponent", DeleteComponentTool);
-	UI_COMMAND( MoveToLevelTool, "Tool - Move to Level", "Moves the selected landscape components to the current streaming level", EUserInterfaceActionType::RadioButton, FInputGesture() );
-	NameToCommandMap.Add("ToolSet_MoveToLevel", MoveToLevelTool);
+	UI_COMMAND(SelectComponentTool, "Tool - Component Selection", "", EUserInterfaceActionType::RadioButton, FInputGesture());
+	NameToCommandMap.Add("Tool_Select", SelectComponentTool);
+	UI_COMMAND(AddComponentTool, "Tool - Add Components", "", EUserInterfaceActionType::RadioButton, FInputGesture());
+	NameToCommandMap.Add("Tool_AddComponent", AddComponentTool);
+	UI_COMMAND(DeleteComponentTool, "Tool - Delete Components", "", EUserInterfaceActionType::RadioButton, FInputGesture());
+	NameToCommandMap.Add("Tool_DeleteComponent", DeleteComponentTool);
+	UI_COMMAND(MoveToLevelTool, "Tool - Move to Level", "Moves the selected landscape components to the current streaming level", EUserInterfaceActionType::RadioButton, FInputGesture());
+	NameToCommandMap.Add("Tool_MoveToLevel", MoveToLevelTool);
 
-	UI_COMMAND( RegionSelectTool, "Tool - Region Selection", "", EUserInterfaceActionType::RadioButton, FInputGesture() );
-	NameToCommandMap.Add("ToolSet_Mask", RegionSelectTool);
-	UI_COMMAND( RegionCopyPasteTool, "Tool - Copy/Paste", "", EUserInterfaceActionType::RadioButton, FInputGesture() );
-	NameToCommandMap.Add("ToolSet_CopyPaste", RegionCopyPasteTool);
+	UI_COMMAND(RegionSelectTool, "Tool - Region Selection", "", EUserInterfaceActionType::RadioButton, FInputGesture());
+	NameToCommandMap.Add("Tool_Mask", RegionSelectTool);
+	UI_COMMAND(RegionCopyPasteTool, "Tool - Copy/Paste", "", EUserInterfaceActionType::RadioButton, FInputGesture());
+	NameToCommandMap.Add("Tool_CopyPaste", RegionCopyPasteTool);
 
-	UI_COMMAND( SplineTool, "Tool - Edit Splines", "", EUserInterfaceActionType::RadioButton, FInputGesture() );
-	NameToCommandMap.Add("ToolSet_Splines", SplineTool);
+	UI_COMMAND(SplineTool, "Tool - Edit Splines", "", EUserInterfaceActionType::RadioButton, FInputGesture());
+	NameToCommandMap.Add("Tool_Splines", SplineTool);
 
-	UI_COMMAND( CircleBrush, "Brush - Circle", "", EUserInterfaceActionType::RadioButton, FInputGesture() );
+	UI_COMMAND(CircleBrush, "Brush - Circle", "", EUserInterfaceActionType::RadioButton, FInputGesture());
 	NameToCommandMap.Add("BrushSet_Circle", CircleBrush);
-	UI_COMMAND( AlphaBrush, "Brush - Alpha", "", EUserInterfaceActionType::RadioButton, FInputGesture() );
+	UI_COMMAND(AlphaBrush, "Brush - Alpha", "", EUserInterfaceActionType::RadioButton, FInputGesture());
 	NameToCommandMap.Add("BrushSet_Alpha", AlphaBrush);
-	UI_COMMAND( AlphaBrush_Pattern, "Brush - Pattern", "", EUserInterfaceActionType::RadioButton, FInputGesture() );
+	UI_COMMAND(AlphaBrush_Pattern, "Brush - Pattern", "", EUserInterfaceActionType::RadioButton, FInputGesture());
 	NameToCommandMap.Add("BrushSet_Pattern", AlphaBrush_Pattern);
-	UI_COMMAND( ComponentBrush, "Brush - Component", "", EUserInterfaceActionType::RadioButton, FInputGesture() );
+	UI_COMMAND(ComponentBrush, "Brush - Component", "", EUserInterfaceActionType::RadioButton, FInputGesture());
 	NameToCommandMap.Add("BrushSet_Component", ComponentBrush);
-	UI_COMMAND( GizmoBrush, "Brush - Gizmo", "", EUserInterfaceActionType::RadioButton, FInputGesture() );
+	UI_COMMAND(GizmoBrush, "Brush - Gizmo", "", EUserInterfaceActionType::RadioButton, FInputGesture());
 	NameToCommandMap.Add("BrushSet_Gizmo", GizmoBrush);
 	NameToCommandMap.Add("BrushSet_Splines", SplineTool);
 
-	UI_COMMAND( CircleBrush_Smooth, "Circle Brush - Smooth Falloff", "", EUserInterfaceActionType::RadioButton, FInputGesture() );
+	UI_COMMAND(CircleBrush_Smooth, "Circle Brush - Smooth Falloff", "", EUserInterfaceActionType::RadioButton, FInputGesture());
 	NameToCommandMap.Add("Circle_Smooth", CircleBrush_Smooth);
-	UI_COMMAND( CircleBrush_Linear, "Circle Brush - Linear Falloff", "", EUserInterfaceActionType::RadioButton, FInputGesture() );
+	UI_COMMAND(CircleBrush_Linear, "Circle Brush - Linear Falloff", "", EUserInterfaceActionType::RadioButton, FInputGesture());
 	NameToCommandMap.Add("Circle_Linear", CircleBrush_Linear);
-	UI_COMMAND( CircleBrush_Spherical, "Circle Brush - Spherical Falloff", "", EUserInterfaceActionType::RadioButton, FInputGesture() );
+	UI_COMMAND(CircleBrush_Spherical, "Circle Brush - Spherical Falloff", "", EUserInterfaceActionType::RadioButton, FInputGesture());
 	NameToCommandMap.Add("Circle_Spherical", CircleBrush_Spherical);
-	UI_COMMAND( CircleBrush_Tip, "Circle Brush - Tip Falloff", "", EUserInterfaceActionType::RadioButton, FInputGesture() );
+	UI_COMMAND(CircleBrush_Tip, "Circle Brush - Tip Falloff", "", EUserInterfaceActionType::RadioButton, FInputGesture());
 	NameToCommandMap.Add("Circle_Tip", CircleBrush_Tip);
 
-	UI_COMMAND( ViewModeNormal, "Normal", "", EUserInterfaceActionType::RadioButton, FInputGesture() );
-	UI_COMMAND( ViewModeLOD, "LOD", "", EUserInterfaceActionType::RadioButton, FInputGesture() );
-	UI_COMMAND( ViewModeLayerDensity, "Layer Density", "", EUserInterfaceActionType::RadioButton, FInputGesture() );
-	UI_COMMAND( ViewModeLayerDebug, "Layer Debug", "", EUserInterfaceActionType::RadioButton, FInputGesture() );
+	UI_COMMAND(ViewModeNormal, "Normal", "", EUserInterfaceActionType::RadioButton, FInputGesture());
+	UI_COMMAND(ViewModeLOD, "LOD", "", EUserInterfaceActionType::RadioButton, FInputGesture());
+	UI_COMMAND(ViewModeLayerDensity, "Layer Density", "", EUserInterfaceActionType::RadioButton, FInputGesture());
+	UI_COMMAND(ViewModeLayerDebug, "Layer Debug", "", EUserInterfaceActionType::RadioButton, FInputGesture());
 }
diff --git a/Engine/Source/Editor/LandscapeEditor/Private/LandscapeEditorDetailCustomization_Base.cpp b/Engine/Source/Editor/LandscapeEditor/Private/LandscapeEditorDetailCustomization_Base.cpp
--- a/Engine/Source/Editor/LandscapeEditor/Private/LandscapeEditorDetailCustomization_Base.cpp
+++ b/Engine/Source/Editor/LandscapeEditor/Private/LandscapeEditorDetailCustomization_Base.cpp
@@ -10,18 +10,18 @@
 #include "PropertyHandle.h"
 
 
-class FEdModeLandscape* FLandscapeEditorDetailCustomization_Base::GetEditorMode()
+FEdModeLandscape* FLandscapeEditorDetailCustomization_Base::GetEditorMode()
 {
 	return (FEdModeLandscape*)GLevelEditorModeTools().GetActiveMode(FBuiltinEditorModes::EM_Landscape);
 }
 
-bool FLandscapeEditorDetailCustomization_Base::IsToolActive(FName ToolSetName)
+bool FLandscapeEditorDetailCustomization_Base::IsToolActive(FName ToolName)
 {
 	FEdModeLandscape* LandscapeEdMode = GetEditorMode();
-	if (LandscapeEdMode != NULL && LandscapeEdMode->CurrentToolSet != NULL)
+	if (LandscapeEdMode != NULL && LandscapeEdMode->CurrentTool != NULL)
 	{
-		const FName CurrentToolSetName = LandscapeEdMode->CurrentToolSet->GetToolSetName();
-		return CurrentToolSetName == ToolSetName;
+		const FName CurrentToolName = LandscapeEdMode->CurrentTool->GetToolName();
+		return CurrentToolName == ToolName;
 	}
 
 	return false;
@@ -78,7 +78,7 @@ void FLandscapeEditorDetailCustomization_Base::OnValueCommitted(int32 NewValue,
 
 //////////////////////////////////////////////////////////////////////////
 
-class FEdModeLandscape* FLandscapeEditorStructCustomization_Base::GetEditorMode()
+FEdModeLandscape* FLandscapeEditorStructCustomization_Base::GetEditorMode()
 {
 	return (FEdModeLandscape*)GLevelEditorModeTools().GetActiveMode(FBuiltinEditorModes::EM_Landscape);
 }
diff --git a/Engine/Source/Editor/LandscapeEditor/Private/LandscapeEditorDetailCustomization_Base.h b/Engine/Source/Editor/LandscapeEditor/Private/LandscapeEditorDetailCustomization_Base.h
--- a/Engine/Source/Editor/LandscapeEditor/Private/LandscapeEditorDetailCustomization_Base.h
+++ b/Engine/Source/Editor/LandscapeEditor/Private/LandscapeEditorDetailCustomization_Base.h
@@ -15,8 +15,8 @@ public:
 	virtual void CustomizeDetails(IDetailLayoutBuilder& DetailBuilder) override = 0;
 
 protected:
-	static class FEdModeLandscape* GetEditorMode();
-	static bool IsToolActive(FName ToolSetName);
+	static FEdModeLandscape* GetEditorMode();
+	static bool IsToolActive(FName ToolName);
 	static bool IsBrushSetActive(FName BrushSetName);
 
 	static TOptional<int32> OnGetValue(TSharedRef<IPropertyHandle> PropertyHandle);
@@ -97,5 +97,5 @@ void FLandscapeEditorDetailCustomization_Base::SetPropertyValue(type NewValue, E
 class FLandscapeEditorStructCustomization_Base : public IPropertyTypeCustomization
 {
 protected:
-	static class FEdModeLandscape* GetEditorMode();
+	static FEdModeLandscape* GetEditorMode();
 };
diff --git a/Engine/Source/Editor/LandscapeEditor/Private/LandscapeEditorDetailCustomization_CopyPaste.cpp b/Engine/Source/Editor/LandscapeEditor/Private/LandscapeEditorDetailCustomization_CopyPaste.cpp
--- a/Engine/Source/Editor/LandscapeEditor/Private/LandscapeEditorDetailCustomization_CopyPaste.cpp
+++ b/Engine/Source/Editor/LandscapeEditor/Private/LandscapeEditorDetailCustomization_CopyPaste.cpp
@@ -28,7 +28,7 @@ TSharedRef<IDetailCustomization> FLandscapeEditorDetailCustomization_CopyPaste::
 BEGIN_SLATE_FUNCTION_BUILD_OPTIMIZATION
 void FLandscapeEditorDetailCustomization_CopyPaste::CustomizeDetails(IDetailLayoutBuilder& DetailBuilder)
 {
-	if (!IsToolActive("ToolSet_CopyPaste"))
+	if (!IsToolActive("CopyPaste"))
 	{
 		return;
 	}
@@ -330,7 +330,7 @@ FReply FLandscapeEditorDetailCustomization_CopyPaste::OnGizmoImportButtonClicked
 				if (LayerName == TEXT(""))
 				{
 					FMessageDialog::Open(EAppMsgType::Ok,
-						FText::Format( NSLOCTEXT("UnrealEd", "LandscapeImport_BadLayerName", "You must enter a name for the layer being imported from {0}."), FText::FromString(Layer.LayerFilename) ));
+						FText::Format(NSLOCTEXT("UnrealEd", "LandscapeImport_BadLayerName", "You must enter a name for the layer being imported from {0}."), FText::FromString(Layer.LayerFilename)));
 					return FReply::Handled();
 				}
 
@@ -342,7 +342,7 @@ FReply FLandscapeEditorDetailCustomization_CopyPaste::OnGizmoImportButtonClicked
 					if (LayerData->Num() != (LandscapeEdMode->UISettings->GizmoImportSize.X * LandscapeEdMode->UISettings->GizmoImportSize.Y))
 					{
 						FMessageDialog::Open(EAppMsgType::Ok,
-							FText::Format( NSLOCTEXT("UnrealEd", "LandscapeImport_BadLayerSize", "Layer {0} file size does not match the heightmap resolution."), FText::FromString(Layer.LayerFilename) ));
+							FText::Format(NSLOCTEXT("UnrealEd", "LandscapeImport_BadLayerSize", "Layer {0} file size does not match the heightmap resolution."), FText::FromString(Layer.LayerFilename)));
 						return FReply::Handled();
 					}
 
@@ -462,11 +462,11 @@ TSharedRef<IPropertyTypeCustomization> FLandscapeEditorStructCustomization_FGizm
 	return MakeShareable(new FLandscapeEditorStructCustomization_FGizmoImportLayer);
 }
 
-void FLandscapeEditorStructCustomization_FGizmoImportLayer::CustomizeHeader(TSharedRef<IPropertyHandle> StructPropertyHandle, class FDetailWidgetRow& HeaderRow, IPropertyTypeCustomizationUtils& StructCustomizationUtils)
+void FLandscapeEditorStructCustomization_FGizmoImportLayer::CustomizeHeader(TSharedRef<IPropertyHandle> StructPropertyHandle, FDetailWidgetRow& HeaderRow, IPropertyTypeCustomizationUtils& StructCustomizationUtils)
 {
 }
 
-void FLandscapeEditorStructCustomization_FGizmoImportLayer::CustomizeChildren(TSharedRef<IPropertyHandle> StructPropertyHandle, class IDetailChildrenBuilder& ChildBuilder, IPropertyTypeCustomizationUtils& StructCustomizationUtils)
+void FLandscapeEditorStructCustomization_FGizmoImportLayer::CustomizeChildren(TSharedRef<IPropertyHandle> StructPropertyHandle, IDetailChildrenBuilder& ChildBuilder, IPropertyTypeCustomizationUtils& StructCustomizationUtils)
 {
 	TSharedRef<IPropertyHandle> PropertyHandle_LayerFilename = StructPropertyHandle->GetChildHandle(GET_MEMBER_NAME_CHECKED(FGizmoImportLayer, LayerFilename)).ToSharedRef();
 	ChildBuilder.AddChildProperty(PropertyHandle_LayerFilename)
diff --git a/Engine/Source/Editor/LandscapeEditor/Private/LandscapeEditorDetailCustomization_MiscTools.cpp b/Engine/Source/Editor/LandscapeEditor/Private/LandscapeEditorDetailCustomization_MiscTools.cpp
--- a/Engine/Source/Editor/LandscapeEditor/Private/LandscapeEditorDetailCustomization_MiscTools.cpp
+++ b/Engine/Source/Editor/LandscapeEditor/Private/LandscapeEditorDetailCustomization_MiscTools.cpp
@@ -39,7 +39,7 @@ void FLandscapeEditorDetailCustomization_MiscTools::CustomizeDetails(IDetailLayo
 		];
 	}
 
-	//if (IsToolActive("ToolSet_Mask"))
+	//if (IsToolActive("Mask"))
 	{
 		ToolsCategory.AddCustomRow("Clear Region Selection")
 		.Visibility(TAttribute<EVisibility>::Create(TAttribute<EVisibility>::FGetter::CreateStatic(&FLandscapeEditorDetailCustomization_MiscTools::GetClearRegionSelectionVisibility)))
@@ -51,7 +51,7 @@ void FLandscapeEditorDetailCustomization_MiscTools::CustomizeDetails(IDetailLayo
 		];
 	}
 
-	if (IsToolActive("ToolSet_Splines"))
+	if (IsToolActive("Splines"))
 	{
 		ToolsCategory.AddCustomRow("Apply Splines")
 		[
@@ -103,7 +103,7 @@ void FLandscapeEditorDetailCustomization_MiscTools::CustomizeDetails(IDetailLayo
 	}
 
 
-	if (IsToolActive("ToolSet_Ramp"))
+	if (IsToolActive("Ramp"))
 	{
 		ToolsCategory.AddCustomRow("Ramp")
 		[
@@ -149,10 +149,10 @@ END_SLATE_FUNCTION_BUILD_OPTIMIZATION
 EVisibility FLandscapeEditorDetailCustomization_MiscTools::GetClearComponentSelectionVisibility()
 {
 	FEdModeLandscape* LandscapeEdMode = GetEditorMode();
-	if (LandscapeEdMode && LandscapeEdMode->CurrentToolSet)
+	if (LandscapeEdMode && LandscapeEdMode->CurrentTool)
 	{
-		const FName CurrentToolSetName = LandscapeEdMode->CurrentToolSet->GetToolSetName();
-		if (CurrentToolSetName == FName("ToolSet_Select"))
+		const FName CurrentToolName = LandscapeEdMode->CurrentTool->GetToolName();
+		if (CurrentToolName == FName("Select"))
 		{
 			return EVisibility::Visible;
 		}
@@ -185,14 +185,14 @@ FReply FLandscapeEditorDetailCustomization_MiscTools::OnClearComponentSelectionB
 EVisibility FLandscapeEditorDetailCustomization_MiscTools::GetClearRegionSelectionVisibility()
 {
 	FEdModeLandscape* LandscapeEdMode = GetEditorMode();
-	if (LandscapeEdMode && LandscapeEdMode->CurrentToolSet)
+	if (LandscapeEdMode && LandscapeEdMode->CurrentTool)
 	{
-		const FName CurrentToolSetName = LandscapeEdMode->CurrentToolSet->GetToolSetName();
-		if (CurrentToolSetName == FName("ToolSet_Mask"))
+		const FName CurrentToolName = LandscapeEdMode->CurrentTool->GetToolName();
+		if (CurrentToolName == FName("Mask"))
 		{
 			return EVisibility::Visible;
 		}
-		else if (LandscapeEdMode->CurrentToolSet->GetTool() && LandscapeEdMode->CurrentToolSet->GetTool()->SupportsMask() &&
+		else if (LandscapeEdMode->CurrentTool && LandscapeEdMode->CurrentTool->SupportsMask() &&
 			LandscapeEdMode->CurrentToolTarget.LandscapeInfo.IsValid() && LandscapeEdMode->CurrentToolTarget.LandscapeInfo->SelectedRegion.Num() > 0)
 		{
 			return EVisibility::Visible;
@@ -208,7 +208,7 @@ FReply FLandscapeEditorDetailCustomization_MiscTools::OnClearRegionSelectionButt
 	if (LandscapeEdMode)
 	{
 		ULandscapeInfo* LandscapeInfo = LandscapeEdMode->CurrentToolTarget.LandscapeInfo.Get();
-		if( LandscapeInfo )
+		if (LandscapeInfo)
 		{
 			LandscapeInfo->ClearSelectedRegion(false);
 		}
@@ -261,7 +261,7 @@ ESlateCheckBoxState::Type FLandscapeEditorDetailCustomization_MiscTools::GetbUse
 FReply FLandscapeEditorDetailCustomization_MiscTools::OnApplyRampButtonClicked()
 {
 	FEdModeLandscape* LandscapeEdMode = GetEditorMode();
-	if (LandscapeEdMode != NULL && IsToolActive(FName("ToolSet_Ramp")))
+	if (LandscapeEdMode != NULL && IsToolActive(FName("Ramp")))
 	{
 		LandscapeEdMode->ApplyRampTool();
 	}
@@ -272,7 +272,7 @@ FReply FLandscapeEditorDetailCustomization_MiscTools::OnApplyRampButtonClicked()
 bool FLandscapeEditorDetailCustomization_MiscTools::GetApplyRampButtonIsEnabled()
 {
 	FEdModeLandscape* LandscapeEdMode = GetEditorMode();
-	if (LandscapeEdMode != NULL && IsToolActive(FName("ToolSet_Ramp")))
+	if (LandscapeEdMode != NULL && IsToolActive(FName("Ramp")))
 	{
 		return LandscapeEdMode->CanApplyRampTool();
 	}
@@ -283,7 +283,7 @@ bool FLandscapeEditorDetailCustomization_MiscTools::GetApplyRampButtonIsEnabled(
 FReply FLandscapeEditorDetailCustomization_MiscTools::OnResetRampButtonClicked()
 {
 	FEdModeLandscape* LandscapeEdMode = GetEditorMode();
-	if (LandscapeEdMode != NULL && IsToolActive(FName("ToolSet_Ramp")))
+	if (LandscapeEdMode != NULL && IsToolActive(FName("Ramp")))
 	{
 		LandscapeEdMode->ResetRampTool();
 	}
diff --git a/Engine/Source/Editor/LandscapeEditor/Private/LandscapeEditorDetailCustomization_NewLandscape.cpp b/Engine/Source/Editor/LandscapeEditor/Private/LandscapeEditorDetailCustomization_NewLandscape.cpp
--- a/Engine/Source/Editor/LandscapeEditor/Private/LandscapeEditorDetailCustomization_NewLandscape.cpp
+++ b/Engine/Source/Editor/LandscapeEditor/Private/LandscapeEditorDetailCustomization_NewLandscape.cpp
@@ -42,7 +42,7 @@ TSharedRef<IDetailCustomization> FLandscapeEditorDetailCustomization_NewLandscap
 BEGIN_SLATE_FUNCTION_BUILD_OPTIMIZATION
 void FLandscapeEditorDetailCustomization_NewLandscape::CustomizeDetails(IDetailLayoutBuilder& DetailBuilder)
 {
-	if (!IsToolActive("ToolSet_NewLandscape"))
+	if (!IsToolActive("NewLandscape"))
 	{
 		return;
 	}
@@ -513,7 +513,7 @@ TSharedRef<SWidget> FLandscapeEditorDetailCustomization_NewLandscape::GetSection
 
 	for (int32 i = 0; i < ARRAY_COUNT(SectionSizes); i++)
 	{
-		MenuBuilder.AddMenuEntry( FText::Format( LOCTEXT("NxNQuads", "{0}x{0} Quads"), FText::AsNumber( SectionSizes[i]) ), FText::GetEmpty(), FSlateIcon(), FExecuteAction::CreateStatic(&OnChangeSectionSize, PropertyHandle, SectionSizes[i]));
+		MenuBuilder.AddMenuEntry(FText::Format(LOCTEXT("NxNQuads", "{0}x{0} Quads"), FText::AsNumber(SectionSizes[i])), FText::GetEmpty(), FSlateIcon(), FExecuteAction::CreateStatic(&OnChangeSectionSize, PropertyHandle, SectionSizes[i]));
 	}
 
 	return MenuBuilder.MakeWidget();
@@ -545,9 +545,9 @@ TSharedRef<SWidget> FLandscapeEditorDetailCustomization_NewLandscape::GetSection
 	for (int32 i = 0; i < ARRAY_COUNT(NumSections); i++)
 	{
 		FFormatNamedArguments Args;
-		Args.Add( TEXT("Width"), NumSections[i] );
-		Args.Add( TEXT("Height"), NumSections[i] );
-		MenuBuilder.AddMenuEntry(FText::Format(NumSections[i] == 1 ? LOCTEXT("1x1Section", "{Width}\u00D7{Height} Section") : LOCTEXT("NxNSections", "{Width}\u00D7{Height} Sections"), Args ), FText::GetEmpty(), FSlateIcon(), FExecuteAction::CreateStatic(&OnChangeSectionsPerComponent, PropertyHandle, NumSections[i]));
+		Args.Add(TEXT("Width"), NumSections[i]);
+		Args.Add(TEXT("Height"), NumSections[i]);
+		MenuBuilder.AddMenuEntry(FText::Format(NumSections[i] == 1 ? LOCTEXT("1x1Section", "{Width}\u00D7{Height} Section") : LOCTEXT("NxNSections", "{Width}\u00D7{Height} Sections"), Args), FText::GetEmpty(), FSlateIcon(), FExecuteAction::CreateStatic(&OnChangeSectionsPerComponent, PropertyHandle, NumSections[i]));
 	}
 
 	return MenuBuilder.MakeWidget();
@@ -570,8 +570,8 @@ FText FLandscapeEditorDetailCustomization_NewLandscape::GetSectionsPerComponent(
 	}
 
 	FFormatNamedArguments Args;
-	Args.Add( TEXT("Width"), SectionsPerComponent );
-	Args.Add( TEXT("Height"), SectionsPerComponent );
+	Args.Add(TEXT("Width"), SectionsPerComponent);
+	Args.Add(TEXT("Height"), SectionsPerComponent);
 	return FText::Format(SectionsPerComponent == 1 ? LOCTEXT("1x1Section", "{Width}\u00D7{Height} Section") : LOCTEXT("NxNSections", "{Width}\u00D7{Height} Sections"), Args);
 }
 
@@ -750,7 +750,7 @@ FReply FLandscapeEditorDetailCustomization_NewLandscape::OnCreateButtonClicked()
 
 			const TArray<FLandscapeImportLayer>& ImportLandscapeLayersList = LandscapeEdMode->UISettings->ImportLandscape_Layers;
 			LayerInfos.Reserve(ImportLandscapeLayersList.Num());
-		
+
 			// Fill in LayerInfos array and allocate data
 			for (int32 LayerIdx = 0; LayerIdx < ImportLandscapeLayersList.Num(); LayerIdx++)
 			{
@@ -830,7 +830,7 @@ FReply FLandscapeEditorDetailCustomization_NewLandscape::OnCreateButtonClicked()
 		// >=2048x2048 -> LOD1
 		// >= 4096x4096 -> LOD2
 		// >= 8192x8192 -> LOD3
-		Landscape->StaticLightingLOD = FMath::DivideAndRoundUp(FMath::CeilLogTwo((SizeX * SizeY) / (2048*2048) + 1), (uint32)2);
+		Landscape->StaticLightingLOD = FMath::DivideAndRoundUp(FMath::CeilLogTwo((SizeX * SizeY) / (2048 * 2048) + 1), (uint32)2);
 
 		if (LandscapeEdMode->NewLandscapePreviewMode == ENewLandscapePreviewMode::ImportLandscape)
 		{
@@ -870,8 +870,8 @@ FReply FLandscapeEditorDetailCustomization_NewLandscape::OnCreateButtonClicked()
 		LandscapeEdMode->CurrentToolTarget.LayerName = NAME_None;
 		LandscapeEdMode->UpdateTargetList();
 
-		LandscapeEdMode->SetCurrentTool("ToolSet_Select"); // change tool so switching back to the manage mode doesn't give "New Landscape" again
-		LandscapeEdMode->SetCurrentTool("ToolSet_Sculpt"); // change to sculpting mode and tool
+		LandscapeEdMode->SetCurrentTool("Select"); // change tool so switching back to the manage mode doesn't give "New Landscape" again
+		LandscapeEdMode->SetCurrentTool("Sculpt"); // change to sculpting mode and tool
 	}
 
 	return FReply::Handled();
@@ -1030,12 +1030,12 @@ void FLandscapeEditorDetailCustomization_NewLandscape::OnImportHeightmapFilename
 				{
 					IImageWrapperModule& ImageWrapperModule = FModuleManager::LoadModuleChecked<IImageWrapperModule>("ImageWrapper");
 					IImageWrapperPtr ImageWrapper = ImageWrapperModule.CreateImageWrapper(EImageFormat::PNG);
-				
+
 					if (!ImageWrapper->SetCompressed(ImportData.GetData(), ImportData.Num()))
 					{
 						LandscapeEdMode->UISettings->ImportLandscape_HeightmapError = ELandscapeImportHeightmapError::CorruptFile;
 					}
-					else if(ImageWrapper->GetWidth() <= 0 || ImageWrapper->GetHeight() <= 0)
+					else if (ImageWrapper->GetWidth() <= 0 || ImageWrapper->GetHeight() <= 0)
 					{
 						LandscapeEdMode->UISettings->ImportLandscape_HeightmapError = ELandscapeImportHeightmapError::InvalidSize;
 					}
@@ -1169,9 +1169,9 @@ TSharedRef<SWidget> FLandscapeEditorDetailCustomization_NewLandscape::GetImportL
 	for (int32 i = 0; i < ImportResolutions.Num(); i++)
 	{
 		FFormatNamedArguments Args;
-		Args.Add( TEXT("Width"), ImportResolutions[i].Width );
-		Args.Add( TEXT("Height"), ImportResolutions[i].Height );
-		MenuBuilder.AddMenuEntry(FText::Format( LOCTEXT("ImportResolution_Format", "{Width}\u00D7{Height}"), Args ), FText(), FSlateIcon(), FExecuteAction::CreateSP(this, &FLandscapeEditorDetailCustomization_NewLandscape::OnChangeImportLandscapeResolution, i));
+		Args.Add(TEXT("Width"), ImportResolutions[i].Width);
+		Args.Add(TEXT("Height"), ImportResolutions[i].Height);
+		MenuBuilder.AddMenuEntry(FText::Format(LOCTEXT("ImportResolution_Format", "{Width}\u00D7{Height}"), Args), FText(), FSlateIcon(), FExecuteAction::CreateSP(this, &FLandscapeEditorDetailCustomization_NewLandscape::OnChangeImportLandscapeResolution, i));
 	}
 
 	return MenuBuilder.MakeWidget();
@@ -1199,9 +1199,9 @@ FText FLandscapeEditorDetailCustomization_NewLandscape::GetImportLandscapeResolu
 		if (Width != 0 && Height != 0)
 		{
 			FFormatNamedArguments Args;
-			Args.Add( TEXT("Width"), Width );
-			Args.Add( TEXT("Height"), Height );
-			return FText::Format( LOCTEXT("ImportResolution_Format", "{Width}\u00D7{Height}"), Args );
+			Args.Add(TEXT("Width"), Width);
+			Args.Add(TEXT("Height"), Height);
+			return FText::Format(LOCTEXT("ImportResolution_Format", "{Width}\u00D7{Height}"), Args);
 		}
 		else
 		{
@@ -1221,21 +1221,21 @@ void FLandscapeEditorDetailCustomization_NewLandscape::ChooseBestComponentSizeFo
 	if (Width > 0 && Height > 0)
 	{
 		// Try to find a section size and number of sections that matches the dimensions of the heightfield
-		for (int32 SectionSizesIdx=ARRAY_COUNT(SectionSizes)-1; SectionSizesIdx >= 0; SectionSizesIdx--)
+		for (int32 SectionSizesIdx = ARRAY_COUNT(SectionSizes) - 1; SectionSizesIdx >= 0; SectionSizesIdx--)
 		{
-			for (int32 NumSectionsIdx=ARRAY_COUNT(NumSections)-1; NumSectionsIdx >= 0; NumSectionsIdx--)
+			for (int32 NumSectionsIdx = ARRAY_COUNT(NumSections) - 1; NumSectionsIdx >= 0; NumSectionsIdx--)
 			{
 				int32 ss = SectionSizes[SectionSizesIdx];
 				int32 ns = NumSections[NumSectionsIdx];
 
-				if (((Width-1) % (ss * ns)) == 0 && ((Width-1) / (ss * ns)) <= 32 &&
-					((Height-1) % (ss * ns)) == 0 && ((Height-1) / (ss * ns)) <= 32)
+				if (((Width - 1) % (ss * ns)) == 0 && ((Width - 1) / (ss * ns)) <= 32 &&
+					((Height - 1) % (ss * ns)) == 0 && ((Height - 1) / (ss * ns)) <= 32)
 				{
 					bFoundMatch = true;
 					LandscapeEdMode->UISettings->NewLandscape_QuadsPerSection = ss;
 					LandscapeEdMode->UISettings->NewLandscape_SectionsPerComponent = ns;
-					LandscapeEdMode->UISettings->NewLandscape_ComponentCount.X = (Width-1) / (ss * ns);
-					LandscapeEdMode->UISettings->NewLandscape_ComponentCount.Y = (Height-1) / (ss * ns);
+					LandscapeEdMode->UISettings->NewLandscape_ComponentCount.X = (Width - 1) / (ss * ns);
+					LandscapeEdMode->UISettings->NewLandscape_ComponentCount.Y = (Height - 1) / (ss * ns);
 					LandscapeEdMode->UISettings->NewLandscape_ClampSize();
 					break;
 				}
@@ -1258,8 +1258,8 @@ void FLandscapeEditorDetailCustomization_NewLandscape::ChooseBestComponentSizeFo
 					continue;
 				}
 
-				const int32 ComponentsX = FMath::DivideAndRoundUp((Width-1), SectionSizes[SectionSizesIdx] * CurrentNumSections);
-				const int32 ComponentsY = FMath::DivideAndRoundUp((Height-1), SectionSizes[SectionSizesIdx] * CurrentNumSections);
+				const int32 ComponentsX = FMath::DivideAndRoundUp((Width - 1), SectionSizes[SectionSizesIdx] * CurrentNumSections);
+				const int32 ComponentsY = FMath::DivideAndRoundUp((Height - 1), SectionSizes[SectionSizesIdx] * CurrentNumSections);
 				if (ComponentsX <= 32 && ComponentsY <= 32)
 				{
 					bFoundMatch = true;
@@ -1275,9 +1275,9 @@ void FLandscapeEditorDetailCustomization_NewLandscape::ChooseBestComponentSizeFo
 			if (!bFoundMatch)
 			{
 				// if the heightmap is very large, fall back to using the largest values we support
-				const int32 MaxSectionSize = SectionSizes[ARRAY_COUNT(SectionSizes)-1];
-				const int32 ComponentsX = FMath::DivideAndRoundUp((Width-1), MaxSectionSize * CurrentNumSections);
-				const int32 ComponentsY = FMath::DivideAndRoundUp((Height-1), MaxSectionSize * CurrentNumSections);
+				const int32 MaxSectionSize = SectionSizes[ARRAY_COUNT(SectionSizes) - 1];
+				const int32 ComponentsX = FMath::DivideAndRoundUp((Width - 1), MaxSectionSize * CurrentNumSections);
+				const int32 ComponentsY = FMath::DivideAndRoundUp((Height - 1), MaxSectionSize * CurrentNumSections);
 
 				bFoundMatch = true;
 				LandscapeEdMode->UISettings->NewLandscape_QuadsPerSection = SectionSizes[MaxSectionSize];
@@ -1312,11 +1312,11 @@ TSharedRef<IPropertyTypeCustomization> FLandscapeEditorStructCustomization_FLand
 	return MakeShareable(new FLandscapeEditorStructCustomization_FLandscapeImportLayer);
 }
 
-void FLandscapeEditorStructCustomization_FLandscapeImportLayer::CustomizeHeader(TSharedRef<IPropertyHandle> StructPropertyHandle, class FDetailWidgetRow& HeaderRow, IPropertyTypeCustomizationUtils& StructCustomizationUtils)
+void FLandscapeEditorStructCustomization_FLandscapeImportLayer::CustomizeHeader(TSharedRef<IPropertyHandle> StructPropertyHandle, FDetailWidgetRow& HeaderRow, IPropertyTypeCustomizationUtils& StructCustomizationUtils)
 {
 }
 
-void FLandscapeEditorStructCustomization_FLandscapeImportLayer::CustomizeChildren(TSharedRef<IPropertyHandle> StructPropertyHandle, class IDetailChildrenBuilder& ChildBuilder, IPropertyTypeCustomizationUtils& StructCustomizationUtils)
+void FLandscapeEditorStructCustomization_FLandscapeImportLayer::CustomizeChildren(TSharedRef<IPropertyHandle> StructPropertyHandle, IDetailChildrenBuilder& ChildBuilder, IPropertyTypeCustomizationUtils& StructCustomizationUtils)
 {
 	TSharedRef<IPropertyHandle> PropertyHandle_LayerName = StructPropertyHandle->GetChildHandle(GET_MEMBER_NAME_CHECKED(FLandscapeImportLayer, LayerName)).ToSharedRef();
 	TSharedRef<IPropertyHandle> PropertyHandle_LayerInfo = StructPropertyHandle->GetChildHandle(GET_MEMBER_NAME_CHECKED(FLandscapeImportLayer, LayerInfo)).ToSharedRef();
@@ -1478,7 +1478,7 @@ FReply FLandscapeEditorStructCustomization_FLandscapeImportLayer::OnLayerFilenam
 	return FReply::Handled();
 }
 
-bool FLandscapeEditorStructCustomization_FLandscapeImportLayer::ShouldFilterLayerInfo(const class FAssetData& AssetData, FName LayerName)
+bool FLandscapeEditorStructCustomization_FLandscapeImportLayer::ShouldFilterLayerInfo(const FAssetData& AssetData, FName LayerName)
 {
 	ULandscapeLayerInfoObject* LayerInfo = CastChecked<ULandscapeLayerInfoObject>(AssetData.GetAsset());
 	if (LayerInfo->LayerName != LayerName)
@@ -1531,7 +1531,7 @@ void FLandscapeEditorStructCustomization_FLandscapeImportLayer::OnImportLayerCre
 		}
 		FString PackageName = Path + LayerObjectName.ToString();
 
-		TSharedRef<SDlgPickAssetPath> NewLayerDlg = 
+		TSharedRef<SDlgPickAssetPath> NewLayerDlg =
 			SNew(SDlgPickAssetPath)
 			.Title(LOCTEXT("CreateNewLayerInfo", "Create New Landscape Layer Info Object"))
 			.DefaultAssetPath(FText::FromString(PackageName));
