Fixing CIS (template instatiation errors).

diff --git a/Engine/Source/Editor/LandscapeEditor/Private/LandscapeEdModeComponentTools.cpp b/Engine/Source/Editor/LandscapeEditor/Private/LandscapeEdModeComponentTools.cpp
--- a/Engine/Source/Editor/LandscapeEditor/Private/LandscapeEdModeComponentTools.cpp
+++ b/Engine/Source/Editor/LandscapeEditor/Private/LandscapeEdModeComponentTools.cpp
@@ -1802,16 +1802,16 @@ public:
 	// Copy tool doesn't use any view information, so just do it as one function
 	void Copy()
 	{
-		CopyTool->BeginTool(nullptr, EdMode->CurrentToolTarget, FVector::ZeroVector);
+		CopyTool->BeginTool(nullptr, this->EdMode->CurrentToolTarget, FVector::ZeroVector);
 		CopyTool->EndTool(nullptr);
 	}
 
 	void Paste()
 	{
-		SetGizmoMode(true);
-		BeginTool(nullptr, EdMode->CurrentToolTarget, FVector::ZeroVector);
-		EndTool(nullptr);
-		SetGizmoMode(false);
+		this->SetGizmoMode(true);
+		this->BeginTool(nullptr, this->EdMode->CurrentToolTarget, FVector::ZeroVector);
+		this->EndTool(nullptr);
+		this->SetGizmoMode(false);
 	}
 
 protected:
