Added FLandscapeTool constructor to initialize its PreviousBrushIndex

diff --git a/Engine/Source/Editor/LandscapeEditor/Public/LandscapeEdMode.h b/Engine/Source/Editor/LandscapeEditor/Public/LandscapeEdMode.h
--- a/Engine/Source/Editor/LandscapeEditor/Public/LandscapeEdMode.h
+++ b/Engine/Source/Editor/LandscapeEditor/Public/LandscapeEdMode.h
@@ -180,6 +180,7 @@ public:
 	virtual bool HandleClick(HHitProxy* HitProxy, const FViewportClick& Click) { return false; }
 	virtual bool InputKey(FEditorViewportClient* InViewportClient, FViewport* InViewport, FKey InKey, EInputEvent InEvent) { return false; }
 	virtual bool InputDelta(FEditorViewportClient* InViewportClient, FViewport* InViewport, FVector& InDrag, FRotator& InRot, FVector& InScale) { return false; }
+	FLandscapeTool() : PreviousBrushIndex(-1) {}
 	virtual ~FLandscapeTool() {}
 	virtual const TCHAR* GetToolName() = 0;
 	virtual FText GetDisplayName() = 0;
