Add manual rotation option to landscape editor alpha brush
Fix bug in alpha brush that could read off the +X/+Y edges of the supplied texture

diff --git a/Engine/Source/Editor/LandscapeEditor/Classes/LandscapeEditorObject.h b/Engine/Source/Editor/LandscapeEditor/Classes/LandscapeEditorObject.h
--- a/Engine/Source/Editor/LandscapeEditor/Classes/LandscapeEditorObject.h
+++ b/Engine/Source/Editor/LandscapeEditor/Classes/LandscapeEditorObject.h
@@ -447,8 +447,12 @@ class ULandscapeEditorObject : public UObject
 	UPROPERTY(Category="Brush Settings", EditAnywhere, meta=(DisplayName="Texture Scale", ShowForBrushes="BrushSet_Pattern", ClampMin="0.005", ClampMax="5", SliderExponent="3"))
 	float AlphaBrushScale;
 
+	// Rotate brush to follow mouse
+	UPROPERTY(Category = "Brush Settings", EditAnywhere, meta = (DisplayName = "Auto-Rotate", ShowForBrushes = "BrushSet_Alpha"))
+	bool bAlphaBrushAutoRotate;
+
 	// Rotates the brush mask texture
-	UPROPERTY(Category="Brush Settings", EditAnywhere, meta=(DisplayName="Texture Rotation", ShowForBrushes="BrushSet_Pattern", ClampMin="-360", ClampMax="360", UIMin="-180", UIMax="180"))
+	UPROPERTY(Category="Brush Settings", EditAnywhere, meta=(DisplayName="Texture Rotation", ShowForBrushes="BrushSet_Alpha,BrushSet_Pattern", ClampMin="-360", ClampMax="360", UIMin="-180", UIMax="180"))
 	float AlphaBrushRotation;
 
 	// Horizontally offsets the brush mask texture
diff --git a/Engine/Source/Editor/LandscapeEditor/Private/LandscapeEdModeBrushes.cpp b/Engine/Source/Editor/LandscapeEditor/Private/LandscapeEdModeBrushes.cpp
--- a/Engine/Source/Editor/LandscapeEditor/Private/LandscapeEdModeBrushes.cpp
+++ b/Engine/Source/Editor/LandscapeEditor/Private/LandscapeEdModeBrushes.cpp
@@ -1081,7 +1081,7 @@ public:
 	virtual bool ApplyBrush(const TArray<FLandscapeToolMousePosition>& MousePositions, TMap<FIntPoint, float>& OutBrush, int32& X1, int32& Y1, int32& X2, int32& Y2) override
 	{
 		ULandscapeInfo* LandscapeInfo = EdMode->CurrentToolTarget.LandscapeInfo.Get();
-		if (OldMousePosition.IsZero())
+		if (EdMode->UISettings->bAlphaBrushAutoRotate && OldMousePosition.IsZero())
 		{
 			X1 = FMath::FloorToInt(LastMousePosition.X);
 			Y1 = FMath::FloorToInt(LastMousePosition.Y);
@@ -1099,6 +1099,7 @@ public:
 			int32 SizeY = EdMode->UISettings->AlphaTextureSizeY;
 			float MaxSize = 2.f * FMath::Sqrt(FMath::Square(Radius) / 2.f);
 			float AlphaBrushScale = MaxSize / (float)FMath::Max<int32>(SizeX, SizeY);
+			const float BrushAngle = EdMode->UISettings->bAlphaBrushAutoRotate ? LastMouseAngle : FMath::DegreesToRadians(EdMode->UISettings->AlphaBrushRotation);
 
 			X1 = FMath::FloorToInt(LastMousePosition.X - Radius);
 			Y1 = FMath::FloorToInt(LastMousePosition.Y - Radius);
@@ -1114,14 +1115,14 @@ public:
 					float ScaleSampleY = ((float)Y - LastMousePosition.Y) / AlphaBrushScale;
 
 					// Rotate around center to match angle
-					float SampleX = ScaleSampleX * FMath::Cos(LastMouseAngle) - ScaleSampleY * FMath::Sin(LastMouseAngle);
-					float SampleY = ScaleSampleY * FMath::Cos(LastMouseAngle) + ScaleSampleX * FMath::Sin(LastMouseAngle);
+					float SampleX = ScaleSampleX * FMath::Cos(BrushAngle) - ScaleSampleY * FMath::Sin(BrushAngle);
+					float SampleY = ScaleSampleY * FMath::Cos(BrushAngle) + ScaleSampleX * FMath::Sin(BrushAngle);
 
 					SampleX += (float)SizeX * 0.5f;
 					SampleY += (float)SizeY * 0.5f;
 
-					if (SampleX >= 0.f && SampleX < (float)SizeX &&
-						SampleY >= 0.f && SampleY < (float)SizeY)
+					if (SampleX >= 0 && SampleX <= (SizeX - 1) &&
+						SampleY >= 0 && SampleY <= (SizeY - 1))
 					{
 						// Sample the alpha texture
 						float Alpha = GetAlphaSample(SampleX, SampleY);
@@ -1156,18 +1157,21 @@ public:
 	{
 		FLandscapeBrushAlphaBase::MouseMove(LandscapeX, LandscapeY);
 
-		// don't do anything with the angle unless we move at least 0.1 units.
-		FVector2D MouseDelta = LastMousePosition - OldMousePosition;
-		if (MouseDelta.SizeSquared() >= FMath::Square(0.5f))
+		if (EdMode->UISettings->bAlphaBrushAutoRotate)
 		{
-			double SampleTime = FPlatformTime::Seconds();
-			float DeltaTime = (float)(SampleTime - LastMouseSampleTime);
-			FVector2D MouseDirection = MouseDelta.SafeNormal();
-			float MouseAngle = FMath::Lerp(LastMouseAngle, FMath::Atan2(-MouseDirection.Y, MouseDirection.X), FMath::Min<float>(10.f * DeltaTime, 1.f));		// lerp over 100ms
-			LastMouseAngle = MouseAngle;
-			LastMouseSampleTime = SampleTime;
-			OldMousePosition = LastMousePosition;
-			// UE_LOG(LogLandscape, Log, TEXT("(%f,%f) delta (%f,%f) angle %f"), LandscapeX, LandscapeY, MouseDirection.X, MouseDirection.Y, MouseAngle);
+			// don't do anything with the angle unless we move at least 0.1 units.
+			FVector2D MouseDelta = LastMousePosition - OldMousePosition;
+			if (MouseDelta.SizeSquared() >= FMath::Square(0.5f))
+			{
+				double SampleTime = FPlatformTime::Seconds();
+				float DeltaTime = (float)(SampleTime - LastMouseSampleTime);
+				FVector2D MouseDirection = MouseDelta.SafeNormal();
+				float MouseAngle = FMath::Lerp(LastMouseAngle, FMath::Atan2(-MouseDirection.Y, MouseDirection.X), FMath::Min<float>(10.f * DeltaTime, 1.f));		// lerp over 100ms
+				LastMouseAngle = MouseAngle;
+				LastMouseSampleTime = SampleTime;
+				OldMousePosition = LastMousePosition;
+				// UE_LOG(LogLandscape, Log, TEXT("(%f,%f) delta (%f,%f) angle %f"), LandscapeX, LandscapeY, MouseDirection.X, MouseDirection.Y, MouseAngle);
+			}
 		}
 	}
 
@@ -1189,7 +1193,7 @@ public:
 				1.0f / (AlphaBrushScale * SizeX),
 				1.0f / (AlphaBrushScale * SizeY),
 				0.f,
-				LastMouseAngle
+				EdMode->UISettings->bAlphaBrushAutoRotate ? LastMouseAngle : FMath::DegreesToRadians(EdMode->UISettings->AlphaBrushRotation)
 				);
 
 			int32 Channel = EdMode->UISettings->AlphaTextureChannel;
diff --git a/Engine/Source/Editor/LandscapeEditor/Private/LandscapeEditorObject.cpp b/Engine/Source/Editor/LandscapeEditor/Private/LandscapeEditorObject.cpp
--- a/Engine/Source/Editor/LandscapeEditor/Private/LandscapeEditorObject.cpp
+++ b/Engine/Source/Editor/LandscapeEditor/Private/LandscapeEditorObject.cpp
@@ -56,6 +56,7 @@ ULandscapeEditorObject::ULandscapeEditorObject(const class FPostConstructInitial
 	, bUseClayBrush(false)
 
 	, AlphaBrushScale(0.5f)
+	, bAlphaBrushAutoRotate(true)
 	, AlphaBrushRotation(0.0f)
 	, AlphaBrushPanU(0.5f)
 	, AlphaBrushPanV(0.5f)
@@ -161,6 +162,7 @@ void ULandscapeEditorObject::Load()
 	GConfig->GetBool(TEXT("LandscapeEdit"), TEXT("bUseClayBrush"), InbUseClayBrush, GEditorUserSettingsIni);
 	bUseClayBrush = InbUseClayBrush;
 	GConfig->GetFloat(TEXT("LandscapeEdit"), TEXT("AlphaBrushScale"), AlphaBrushScale, GEditorUserSettingsIni);
+	GConfig->GetBool(TEXT("LandscapeEdit"), TEXT("AlphaBrushAutoRotate"), bAlphaBrushAutoRotate, GEditorUserSettingsIni);
 	GConfig->GetFloat(TEXT("LandscapeEdit"), TEXT("AlphaBrushRotation"), AlphaBrushRotation, GEditorUserSettingsIni);
 	GConfig->GetFloat(TEXT("LandscapeEdit"), TEXT("AlphaBrushPanU"), AlphaBrushPanU, GEditorUserSettingsIni);
 	GConfig->GetFloat(TEXT("LandscapeEdit"), TEXT("AlphaBrushPanV"), AlphaBrushPanV, GEditorUserSettingsIni);
@@ -278,6 +280,7 @@ void ULandscapeEditorObject::Save()
 	GConfig->SetFloat(TEXT("LandscapeEdit"), TEXT("BrushFalloff"), BrushFalloff, GEditorUserSettingsIni);
 	GConfig->SetBool(TEXT("LandscapeEdit"), TEXT("bUseClayBrush"), bUseClayBrush, GEditorUserSettingsIni);
 	GConfig->SetFloat(TEXT("LandscapeEdit"), TEXT("AlphaBrushScale"), AlphaBrushScale, GEditorUserSettingsIni);
+	GConfig->SetBool(TEXT("LandscapeEdit"), TEXT("AlphaBrushAutoRotate"), bAlphaBrushAutoRotate, GEditorUserSettingsIni);
 	GConfig->SetFloat(TEXT("LandscapeEdit"), TEXT("AlphaBrushRotation"), AlphaBrushRotation, GEditorUserSettingsIni);
 	GConfig->SetFloat(TEXT("LandscapeEdit"), TEXT("AlphaBrushPanU"), AlphaBrushPanU, GEditorUserSettingsIni);
 	GConfig->SetFloat(TEXT("LandscapeEdit"), TEXT("AlphaBrushPanV"), AlphaBrushPanV, GEditorUserSettingsIni);
