Make Clang happy when compiling FLandscapeXYOffsetCache by fully qualifying the template super-type. Also rearranged the constructor order to be 'correct'.

diff --git a/Engine/Source/Editor/LandscapeEditor/Private/LandscapeEdModeTools.h b/Engine/Source/Editor/LandscapeEditor/Private/LandscapeEdModeTools.h
--- a/Engine/Source/Editor/LandscapeEditor/Private/LandscapeEdModeTools.h
+++ b/Engine/Source/Editor/LandscapeEditor/Private/LandscapeEdModeTools.h
@@ -818,8 +818,8 @@ struct FLandscapeXYOffsetCache : public TLandscapeEditCache < FXYOffsetmapAccess
 	FXYOffsetmapAccessor<bInUseInterp> XYOffsetmapAccessor;
 
 	FLandscapeXYOffsetCache(const FLandscapeToolTarget& InTarget)
-		: XYOffsetmapAccessor(InTarget.LandscapeInfo.Get())
-		, TLandscapeEditCache(XYOffsetmapAccessor)
+		: TLandscapeEditCache< FXYOffsetmapAccessor<bInUseInterp>, FVector >(XYOffsetmapAccessor)
+		, XYOffsetmapAccessor(InTarget.LandscapeInfo.Get())
 	{
 	}
 };
