Landscape : Fixed TTP #342158 : Fixed retopoligization tool problem

diff --git a/Engine/Source/Editor/LandscapeEditor/Private/LandscapeEdModeTools.h b/Engine/Source/Editor/LandscapeEditor/Private/LandscapeEdModeTools.h
--- a/Engine/Source/Editor/LandscapeEditor/Private/LandscapeEdModeTools.h
+++ b/Engine/Source/Editor/LandscapeEditor/Private/LandscapeEdModeTools.h
@@ -702,7 +702,7 @@ struct FXYOffsetmapAccessor
 		LandscapeEdit->GetXYOffsetDataFast(X1, Y1, X2, Y2, Data);
 
 		TMap<FIntPoint, uint16> NewHeights;
-		LandscapeEdit->GetHeightDataFast(X1, Y1, X2, Y2, NewHeights);
+		LandscapeEdit->GetHeightData(X1, Y1, X2, Y2, NewHeights);
 		for (int32 Y = Y1; Y <= Y2; ++Y)
 		{
 			for (int32 X = X1; X <= X2; ++X)
