Fixed issue where returning to landscape mode after creating a landscape causes a crash.
#ttp 341779 - CRITICAL: Regression: EDITOR: CRASH: Switching back to Landscape edit more
#branch UE4
#proj Editor.LandscapeEditor
In FEdModeLandscape::Enter, initializes the current tool prior to creating the landscape toolkit, as the toolkit has a dependency on the current tool being valid.

diff --git a/Engine/Source/Editor/LandscapeEditor/Private/LandscapeEdMode.cpp b/Engine/Source/Editor/LandscapeEditor/Private/LandscapeEdMode.cpp
--- a/Engine/Source/Editor/LandscapeEditor/Private/LandscapeEdMode.cpp
+++ b/Engine/Source/Editor/LandscapeEditor/Private/LandscapeEdMode.cpp
@@ -406,6 +406,25 @@ void FEdModeLandscape::Enter()
 	// Load UI settings from config file
 	UISettings->Load();
 
+	// Initialize current tool prior to creating the landscape toolkit in case it has a dependency on it
+	if (LandscapeList.Num() == 0)
+	{
+		SetCurrentToolMode("ToolMode_Manage", false);
+		SetCurrentTool("NewLandscape");
+	}
+	else
+	{
+		if (CurrentToolMode == NULL)
+		{
+			SetCurrentToolMode("ToolMode_Sculpt", false);
+			SetCurrentTool("Sculpt");
+		}
+		else
+		{
+			SetCurrentTool(CurrentToolMode->CurrentToolName);
+		}
+	}
+
 	// Create the landscape editor window
 	if (!Toolkit.IsValid())
 	{
@@ -426,24 +445,6 @@ void FEdModeLandscape::Enter()
 	{
 		GizmoBrush->EnterBrush();
 	}
-
-	if (LandscapeList.Num() == 0)
-	{
-		SetCurrentToolMode("ToolMode_Manage", false);
-		SetCurrentTool("NewLandscape");
-	}
-	else
-	{
-		if (CurrentToolMode == NULL)
-		{
-			SetCurrentToolMode("ToolMode_Sculpt", false);
-			SetCurrentTool("Sculpt");
-		}
-		else
-		{
-			SetCurrentTool(CurrentToolMode->CurrentToolName);
-		}
-	}
 }
 
 
