Fix crash when trying to use the landscape painting tools with no paintable layers

diff --git a/Engine/Source/Editor/LandscapeEditor/Private/LandscapeEdMode.cpp b/Engine/Source/Editor/LandscapeEditor/Private/LandscapeEdMode.cpp
--- a/Engine/Source/Editor/LandscapeEditor/Private/LandscapeEdMode.cpp
+++ b/Engine/Source/Editor/LandscapeEditor/Private/LandscapeEdMode.cpp
@@ -1183,7 +1183,7 @@ bool FEdModeLandscape::InputKey(FEditorViewportClient* ViewportClient, FViewport
 		{
 			if (Event == IE_Pressed && (IsCtrlDown(Viewport) || (Viewport->IsPenActive() && Viewport->GetTabletPressure() > 0.f)))
 			{
-				if (CurrentTool)
+				if (CurrentTool && (CurrentTool->GetSupportedTargetTypes() == ELandscapeToolTargetTypeMask::NA || CurrentToolTarget.TargetType != ELandscapeToolTargetType::Invalid))
 				{
 					FVector HitLocation;
 					if (LandscapeMouseTrace(ViewportClient, HitLocation))
@@ -1783,8 +1783,8 @@ void FEdModeLandscape::UpdateTargetList()
 
 		if (!bFoundSelected)
 		{
-			CurrentToolTarget.TargetType = ELandscapeToolTargetType::Heightmap;
-			CurrentToolTarget.LayerInfo = NULL;
+			CurrentToolTarget.TargetType = ELandscapeToolTargetType::Invalid;
+			CurrentToolTarget.LayerInfo = nullptr;
 			CurrentToolTarget.LayerName = NAME_None;
 		}
 	}
