More fixes for Clang warnings

diff --git a/Engine/Plugins/Editor/EpicSurvey/Source/EpicSurvey/Private/EpicSurvey.cpp b/Engine/Plugins/Editor/EpicSurvey/Source/EpicSurvey/Private/EpicSurvey.cpp
--- a/Engine/Plugins/Editor/EpicSurvey/Source/EpicSurvey/Private/EpicSurvey.cpp
+++ b/Engine/Plugins/Editor/EpicSurvey/Source/EpicSurvey/Private/EpicSurvey.cpp
@@ -34,12 +34,12 @@ TSharedRef< FEpicSurvey > FEpicSurvey::Create()
 FEpicSurvey::FEpicSurvey( const TSharedRef< FUICommandList >& InActionList ) 
 	: InitializationState( EContentInitializationState::NotStarted )
 	, ActionList( InActionList )
-	, bIsShowingToolbarNotification( false )
+	, SurveyNotificationDelayTime(0)
 	, SurveyNotificationDuration(5.0f)
-	, SurveyPulseDuration(5)
 	, SurveyPulseTimeInterval(5)
-	, SurveyNotificationDelayTime(0)
+	, SurveyPulseDuration(5)
 	, bSurveyIconPulsing(false)
+	, bIsShowingToolbarNotification( false )
 	, CurrentCulture(ECultureSpecification::Full)
 {
 	
diff --git a/Engine/Plugins/Editor/EpicSurvey/Source/EpicSurvey/Private/SurveyBranch.h b/Engine/Plugins/Editor/EpicSurvey/Source/EpicSurvey/Private/SurveyBranch.h
--- a/Engine/Plugins/Editor/EpicSurvey/Source/EpicSurvey/Private/SurveyBranch.h
+++ b/Engine/Plugins/Editor/EpicSurvey/Source/EpicSurvey/Private/SurveyBranch.h
@@ -11,7 +11,7 @@ class FSurveyBranch : public TSharedFromThis< FSurveyBranch >
 public:
 	struct BranchPointsType 
 	{
-		BranchPointsType(const TSharedRef<FSurveyBranch>& InBranchBlock) : Points(0), BranchBlock(InBranchBlock) {}
+		BranchPointsType(const TSharedRef<FSurveyBranch>& InBranchBlock) : BranchBlock(InBranchBlock), Points(0) {}
 		TSharedRef<FSurveyBranch> BranchBlock;
 		int32 Points;
 	};
diff --git a/Engine/Plugins/Experimental/LiveEditor/Source/LiveEditor/Private/LiveEditorDeviceSetupWizard.cpp b/Engine/Plugins/Experimental/LiveEditor/Source/LiveEditor/Private/LiveEditorDeviceSetupWizard.cpp
--- a/Engine/Plugins/Experimental/LiveEditor/Source/LiveEditor/Private/LiveEditorDeviceSetupWizard.cpp
+++ b/Engine/Plugins/Experimental/LiveEditor/Source/LiveEditor/Private/LiveEditorDeviceSetupWizard.cpp
@@ -43,9 +43,9 @@ struct FConfigurationState : public FLiveEditorDeviceSetupWizard::FState
 {
 	FConfigurationState( FLiveEditorDeviceSetupWizard *_Owner, int32 _NextState )
 	:	FLiveEditorDeviceSetupWizard::FState(_NextState),
-		Owner(_Owner),
 		bHasButtons(false),
-		bHasEndlessEncoders(false)
+		bHasEndlessEncoders(false),
+		Owner(_Owner)
 	{
 	}
 
diff --git a/Engine/Source/Developer/OutputLog/Private/SOutputLog.cpp b/Engine/Source/Developer/OutputLog/Private/SOutputLog.cpp
--- a/Engine/Source/Developer/OutputLog/Private/SOutputLog.cpp
+++ b/Engine/Source/Developer/OutputLog/Private/SOutputLog.cpp
@@ -6,8 +6,8 @@
 #include "SScrollBorder.h"
 
 SConsoleInputBox::SConsoleInputBox()
-	: bIgnoreUIUpdate(false)
-	, SelectedSuggestion(-1)
+	: SelectedSuggestion(-1)
+	, bIgnoreUIUpdate(false)
 {
 }
 
diff --git a/Engine/Source/Developer/SuperSearch/Private/SSuperSearch.cpp b/Engine/Source/Developer/SuperSearch/Private/SSuperSearch.cpp
--- a/Engine/Source/Developer/SuperSearch/Private/SSuperSearch.cpp
+++ b/Engine/Source/Developer/SuperSearch/Private/SSuperSearch.cpp
@@ -5,8 +5,8 @@
 #include "SScrollBorder.h"
 
 SSuperSearchBox::SSuperSearchBox()
-	: bIgnoreUIUpdate(false)
-	, SelectedSuggestion(-1)
+	: SelectedSuggestion(-1)
+	, bIgnoreUIUpdate(false)
 {
 }
 
diff --git a/Engine/Source/Editor/GraphEditor/Private/GraphActionNode.h b/Engine/Source/Editor/GraphEditor/Private/GraphActionNode.h
--- a/Engine/Source/Editor/GraphEditor/Private/GraphActionNode.h
+++ b/Engine/Source/Editor/GraphEditor/Private/GraphActionNode.h
@@ -382,37 +382,37 @@ public:
 protected:
 
 	FGraphActionNode( FString InCategory, int32 InGrouping, TSharedPtr<FEdGraphSchemaAction> InAction )
-		: Category( InCategory )
+		: SectionID ( 0 )
+		, Category( InCategory )
 		, Grouping( InGrouping )
 		, bIsRenameRequestBeforeReady( false )
-		, SectionID ( 0 )
 	{
 		Actions.Add( InAction );
 	}
 
 	FGraphActionNode( FString InCategory, FString InCategoryChain, int32 InGrouping, TSharedPtr<FEdGraphSchemaAction> InAction )
-		: Category( InCategory )
+		: SectionID ( 0 )
+		, Category( InCategory )
 		, CategoryChain( InCategoryChain )
 		, Grouping( InGrouping )
 		, bIsRenameRequestBeforeReady( false )
-		, SectionID ( 0 )
 	{
 		Actions.Add( InAction );
 	}
 
 	FGraphActionNode( FString InCategory, int32 InGrouping, TArray< TSharedPtr<FEdGraphSchemaAction> >& InActionList, int32 InSectionID = 0 )
-	: Category( InCategory )
-	, Grouping( InGrouping )
-	, bIsRenameRequestBeforeReady( false )
-	, SectionID ( InSectionID )
+		: SectionID ( InSectionID )
+		, Category( InCategory )
+		, Grouping( InGrouping )
+		, bIsRenameRequestBeforeReady( false )
 	{
 		Actions.Append( InActionList );
 	}
 
 	FGraphActionNode( int32 InGrouping, int32 InNodeSection )		
-		: Grouping( InGrouping )
+		: SectionID ( InNodeSection )
+		, Grouping( InGrouping )
 		, bIsRenameRequestBeforeReady( false )
-		, SectionID ( InNodeSection )
 	{
 		Actions.Add( NULL );		
 	}
diff --git a/Engine/Source/Editor/LandscapeEditor/Private/LandscapeEdModeComponentTools.cpp b/Engine/Source/Editor/LandscapeEditor/Private/LandscapeEdModeComponentTools.cpp
--- a/Engine/Source/Editor/LandscapeEditor/Private/LandscapeEdModeComponentTools.cpp
+++ b/Engine/Source/Editor/LandscapeEditor/Private/LandscapeEdModeComponentTools.cpp
@@ -290,6 +290,8 @@ public:
 		: LandscapeInfo(InTarget.LandscapeInfo.Get())
 	{}
 
+	virtual ~FLandscapeToolStrokeMoveToLevel() {}
+
 	virtual void Apply(FEditorViewportClient* ViewportClient, FLandscapeBrush* Brush, const ULandscapeEditorObject* UISettings, const TArray<FLandscapeToolMousePosition>& MousePositions)
 	{
 		ALandscape* Landscape = LandscapeInfo ? LandscapeInfo->LandscapeActor.Get() : NULL;
@@ -753,6 +755,8 @@ public:
 		,	XYOffsetCache(InTarget)
 	{}
 
+	virtual ~FLandscapeToolStrokeAddComponent() {}
+
 	virtual void Apply(FEditorViewportClient* ViewportClient, FLandscapeBrush* Brush, const ULandscapeEditorObject* UISettings, const TArray<FLandscapeToolMousePosition>& MousePositions)
 	{
 		ALandscapeProxy* Landscape = LandscapeInfo ? LandscapeInfo->GetCurrentLevelLandscapeProxy(true) : NULL;
@@ -921,6 +925,8 @@ public:
 		: LandscapeInfo(InTarget.LandscapeInfo.Get())
 	{}
 
+	virtual ~FLandscapeToolStrokeDeleteComponent() {}
+
 	virtual void Apply(FEditorViewportClient* ViewportClient, FLandscapeBrush* Brush, const ULandscapeEditorObject* UISettings, const TArray<FLandscapeToolMousePosition>& MousePositions)
 	{
 		if (LandscapeInfo)
diff --git a/Engine/Source/Editor/LandscapeEditor/Private/LandscapeEdModePaintTools.cpp b/Engine/Source/Editor/LandscapeEditor/Private/LandscapeEdModePaintTools.cpp
--- a/Engine/Source/Editor/LandscapeEditor/Private/LandscapeEdModePaintTools.cpp
+++ b/Engine/Source/Editor/LandscapeEditor/Private/LandscapeEdModePaintTools.cpp
@@ -85,6 +85,8 @@ public:
 		: FLandscapeToolStrokePaintBase<ToolTarget>(InEdMode, InTarget)
 	{}
 
+	virtual ~FLandscapeToolStrokePaint() {}
+
 	virtual void Apply(FEditorViewportClient* ViewportClient, FLandscapeBrush* Brush, const ULandscapeEditorObject* UISettings, const TArray<FLandscapeToolMousePosition>& MousePositions) override
 	{
 		// Get list of verts to update
@@ -447,6 +449,8 @@ public:
 		}
 	}
 
+	virtual ~FLandscapeToolStrokeFlatten() {}
+
 	virtual void Apply(FEditorViewportClient* ViewportClient, FLandscapeBrush* Brush, const ULandscapeEditorObject* UISettings, const TArray<FLandscapeToolMousePosition>& MousePositions) override
 	{
 		if (!this->LandscapeInfo) return;
@@ -650,6 +654,8 @@ public:
 		: FLandscapeToolStrokePaintBase<ToolTarget>(InEdMode, InTarget)
 	{}
 
+	virtual ~FLandscapeToolStrokeNoise() {}
+
 	virtual void Apply(FEditorViewportClient* ViewportClient, FLandscapeBrush* Brush, const ULandscapeEditorObject* UISettings, const TArray<FLandscapeToolMousePosition>& MousePositions) override
 	{
 		if (!this->LandscapeInfo) return;
diff --git a/Engine/Source/Editor/LandscapeEditor/Private/LandscapeEdModeTools.h b/Engine/Source/Editor/LandscapeEditor/Private/LandscapeEdModeTools.h
--- a/Engine/Source/Editor/LandscapeEditor/Private/LandscapeEdModeTools.h
+++ b/Engine/Source/Editor/LandscapeEditor/Private/LandscapeEdModeTools.h
@@ -657,11 +657,18 @@ struct FLandscapeHeightCache : public TLandscapeEditCache < FHeightmapAccessor<t
 
 	FHeightmapAccessor<true> HeightmapAccessor;
 
+#ifdef __clang__ // @todo
+	#pragma clang diagnostic push
+	#pragma clang diagnostic ignored "-Wreorder"
+#endif
 	FLandscapeHeightCache(const FLandscapeToolTarget& InTarget)
 		: HeightmapAccessor(InTarget.LandscapeInfo.Get())
 		, TLandscapeEditCache(HeightmapAccessor)
 	{
 	}
+#ifdef __clang__
+	#pragma clang diagnostic pop
+#endif
 };
 
 //
@@ -915,11 +922,18 @@ struct FLandscapeAlphaCache : public TLandscapeEditCache < FAlphamapAccessor<tru
 
 	FAlphamapAccessor<true, false> AlphamapAccessor;
 
+#ifdef __clang__ // @todo
+	#pragma clang diagnostic push
+	#pragma clang diagnostic ignored "-Wreorder"
+#endif
 	FLandscapeAlphaCache(const FLandscapeToolTarget& InTarget)
 		: AlphamapAccessor(InTarget.LandscapeInfo.Get(), InTarget.LayerInfo.Get())
 		, TLandscapeEditCache(AlphamapAccessor)
 	{
 	}
+#ifdef __clang__
+	#pragma clang diagnostic pop
+#endif
 };
 
 struct FLandscapeVisCache : public TLandscapeEditCache < FAlphamapAccessor<false, false>, uint8 >
@@ -929,11 +943,18 @@ struct FLandscapeVisCache : public TLandscapeEditCache < FAlphamapAccessor<false
 
 	FAlphamapAccessor<false, false> AlphamapAccessor;
 
+#ifdef __clang__ // @todo
+	#pragma clang diagnostic push
+	#pragma clang diagnostic ignored "-Wreorder"
+#endif
 	FLandscapeVisCache(const FLandscapeToolTarget& InTarget)
 		: AlphamapAccessor(InTarget.LandscapeInfo.Get(), ALandscapeProxy::VisibilityLayer)
 		, TLandscapeEditCache(AlphamapAccessor)
 	{
 	}
+#ifdef __clang__
+	#pragma clang diagnostic pop
+#endif
 };
 
 //
@@ -984,11 +1005,18 @@ struct FLandscapeFullWeightCache : public TLandscapeEditCache < FFullWeightmapAc
 
 	FFullWeightmapAccessor<false> WeightmapAccessor;
 
+#ifdef __clang__ // @todo
+	#pragma clang diagnostic push
+	#pragma clang diagnostic ignored "-Wreorder"
+#endif
 	FLandscapeFullWeightCache(const FLandscapeToolTarget& InTarget)
 		: WeightmapAccessor(InTarget.LandscapeInfo.Get())
 		, TLandscapeEditCache(WeightmapAccessor)
 	{
 	}
+#ifdef __clang__
+	#pragma clang diagnostic pop
+#endif
 
 	// Only for all weight case... the accessor type should be TArray<uint8>
 	void GetCachedData(int32 X1, int32 Y1, int32 X2, int32 Y2, TArray<uint8>& OutData, int32 ArraySize)
@@ -1088,11 +1116,18 @@ struct FLandscapeDataCache : public TLandscapeEditCache < FDatamapAccessor<false
 
 	FDatamapAccessor<false> DataAccessor;
 
+#ifdef __clang__ // @todo
+	#pragma clang diagnostic push
+	#pragma clang diagnostic ignored "-Wreorder"
+#endif
 	FLandscapeDataCache(const FLandscapeToolTarget& InTarget)
 		: DataAccessor(InTarget.LandscapeInfo.Get())
 		, TLandscapeEditCache(DataAccessor)
 	{
 	}
+#ifdef __clang__
+	#pragma clang diagnostic pop
+#endif
 };
 
 
diff --git a/Engine/Source/Editor/Layers/Private/ActorLayerCollectionViewModel.cpp b/Engine/Source/Editor/Layers/Private/ActorLayerCollectionViewModel.cpp
--- a/Engine/Source/Editor/Layers/Private/ActorLayerCollectionViewModel.cpp
+++ b/Engine/Source/Editor/Layers/Private/ActorLayerCollectionViewModel.cpp
@@ -6,9 +6,9 @@
 #define LOCTEXT_NAMESPACE "LayersView"
 
 FActorLayerCollectionViewModel::FActorLayerCollectionViewModel( const TSharedRef< ILayers >& InWorldLayers, const TWeakObjectPtr< UEditorEngine >& InEditor )
-	: WorldLayers( InWorldLayers )
+	: bIsRefreshing( false )
+	, WorldLayers( InWorldLayers )
 	, Editor( InEditor )
-	, bIsRefreshing( false )
 {
 
 }
diff --git a/Engine/Source/Editor/Layers/Private/LayerCollectionViewModel.cpp b/Engine/Source/Editor/Layers/Private/LayerCollectionViewModel.cpp
--- a/Engine/Source/Editor/Layers/Private/LayerCollectionViewModel.cpp
+++ b/Engine/Source/Editor/Layers/Private/LayerCollectionViewModel.cpp
@@ -4,11 +4,11 @@
 #define LOCTEXT_NAMESPACE "LayersView"
 
 FLayerCollectionViewModel::FLayerCollectionViewModel( const TSharedRef< ILayers >& InWorldLayers, const TWeakObjectPtr< UEditorEngine >& InEditor )
-	: WorldLayers( InWorldLayers )
-	, Editor( InEditor )
+	: bIsRefreshing( false )
 	, Filters( MakeShareable( new LayerFilterCollection ) )
 	, CommandList( MakeShareable( new FUICommandList ) )
-	, bIsRefreshing( false )
+	, WorldLayers( InWorldLayers )
+	, Editor( InEditor )
 {
 
 }
diff --git a/Engine/Source/Editor/MaterialEditor/Private/FindInMaterial.cpp b/Engine/Source/Editor/MaterialEditor/Private/FindInMaterial.cpp
--- a/Engine/Source/Editor/MaterialEditor/Private/FindInMaterial.cpp
+++ b/Engine/Source/Editor/MaterialEditor/Private/FindInMaterial.cpp
@@ -15,17 +15,17 @@ FFindInMaterialResult::FFindInMaterialResult(const FString& InValue)
 }
 
 FFindInMaterialResult::FFindInMaterialResult(const FString& InValue, TSharedPtr<FFindInMaterialResult>& InParent, UClass* InClass, int InDuplicationIndex)
-: Value(InValue), DuplicationIndex(InDuplicationIndex), Parent(InParent), Class(InClass), Pin(NULL), GraphNode(NULL)
+: Parent(InParent), Value(InValue), DuplicationIndex(InDuplicationIndex), Class(InClass), Pin(NULL), GraphNode(NULL)
 {
 }
 
 FFindInMaterialResult::FFindInMaterialResult(const FString& InValue, TSharedPtr<FFindInMaterialResult>& InParent, UEdGraphPin* InPin)
-: Value(InValue), DuplicationIndex(0), Parent(InParent), Class(InPin->GetClass()), Pin(InPin), GraphNode(NULL)
+: Parent(InParent), Value(InValue), DuplicationIndex(0), Class(InPin->GetClass()), Pin(InPin), GraphNode(NULL)
 {
 }
 
 FFindInMaterialResult::FFindInMaterialResult(const FString& InValue, TSharedPtr<FFindInMaterialResult>& InParent, UEdGraphNode* InNode)
-: Value(InValue), DuplicationIndex(0), Parent(InParent), Class(InNode->GetClass()), Pin(NULL), GraphNode(InNode)
+: Parent(InParent), Value(InValue), DuplicationIndex(0), Class(InNode->GetClass()), Pin(NULL), GraphNode(InNode)
 {
 	if (GraphNode.IsValid())
 	{
diff --git a/Engine/Source/Editor/MaterialEditor/Private/MaterialInstanceEditor.cpp b/Engine/Source/Editor/MaterialEditor/Private/MaterialInstanceEditor.cpp
--- a/Engine/Source/Editor/MaterialEditor/Private/MaterialInstanceEditor.cpp
+++ b/Engine/Source/Editor/MaterialEditor/Private/MaterialInstanceEditor.cpp
@@ -32,8 +32,8 @@ class SMaterialTreeWidgetItem : public SMultiColumnTableRow< TWeakObjectPtr<UMat
 {
 public:
 	SLATE_BEGIN_ARGS(SMaterialTreeWidgetItem)
-		: _WidgetInfoToVisualize()
-		, _ParentIndex( -1 )
+		: _ParentIndex( -1 )
+		, _WidgetInfoToVisualize()
 		{}
 		SLATE_ARGUMENT( int32, ParentIndex )
 		SLATE_ARGUMENT( TWeakObjectPtr<UMaterialInterface>, WidgetInfoToVisualize )
diff --git a/Engine/Source/Editor/Matinee/Private/MatineeHitProxy.h b/Engine/Source/Editor/Matinee/Private/MatineeHitProxy.h
--- a/Engine/Source/Editor/Matinee/Private/MatineeHitProxy.h
+++ b/Engine/Source/Editor/Matinee/Private/MatineeHitProxy.h
@@ -143,8 +143,8 @@ struct HMatineeTrackGraphPropBtn : public HHitProxy
 	HMatineeTrackGraphPropBtn(class UInterpGroup* InGroup, int32 InSubTrackGroupIndex, UInterpTrack* InTrack ) :
 		HHitProxy(HPP_UI),
 		Group(InGroup),
-		SubTrackGroupIndex( InSubTrackGroupIndex ),
-		Track(InTrack)
+		Track(InTrack),
+		SubTrackGroupIndex( InSubTrackGroupIndex )
 	{}
 };
 
diff --git a/Engine/Source/Editor/NewsFeed/Private/Implementation/NewsFeedCache.cpp b/Engine/Source/Editor/NewsFeed/Private/Implementation/NewsFeedCache.cpp
--- a/Engine/Source/Editor/NewsFeed/Private/Implementation/NewsFeedCache.cpp
+++ b/Engine/Source/Editor/NewsFeed/Private/Implementation/NewsFeedCache.cpp
@@ -13,8 +13,8 @@ const FString FeedFilename(TEXT("newsfeed.json"));
  *****************************************************************************/
 
 FNewsFeedCache::FNewsFeedCache( )
-	: CurrentCultureSpec(ENewsFeedCultureSpec::Full)
-	, LoaderState(ENewsFeedState::NotStarted)
+	: LoaderState(ENewsFeedState::NotStarted)
+	, CurrentCultureSpec(ENewsFeedCultureSpec::Full)
 {
 	TickDelegate = FTickerDelegate::CreateRaw(this, &FNewsFeedCache::HandleTicker);
 }
diff --git a/Engine/Source/Editor/PListEditor/Private/PListNodeArray.h b/Engine/Source/Editor/PListEditor/Private/PListNodeArray.h
--- a/Engine/Source/Editor/PListEditor/Private/PListNodeArray.h
+++ b/Engine/Source/Editor/PListEditor/Private/PListNodeArray.h
@@ -12,7 +12,7 @@ class FPListNodeArray : public IPListNode
 {
 public:
 	FPListNodeArray(SPListEditorPanel* InEditorWidget)
-		: IPListNode(InEditorWidget), bFiltered(false), bArrayMember(false), TableRow(nullptr), ArrayIndex(-1), bKeyValid(false)
+		: IPListNode(InEditorWidget), TableRow(nullptr), ArrayIndex(-1), bFiltered(false), bArrayMember(false), bKeyValid(false)
 	{}
 
 public:
diff --git a/Engine/Source/Editor/PListEditor/Private/PListNodeBoolean.h b/Engine/Source/Editor/PListEditor/Private/PListNodeBoolean.h
--- a/Engine/Source/Editor/PListEditor/Private/PListNodeBoolean.h
+++ b/Engine/Source/Editor/PListEditor/Private/PListNodeBoolean.h
@@ -12,7 +12,7 @@ class FPListNodeBoolean : public IPListNode
 {
 public:
 	FPListNodeBoolean(SPListEditorPanel* InEditorWidget)
-		: IPListNode(InEditorWidget), bFiltered(false), bValue(false), bArrayMember(false), ArrayIndex(-1), bKeyValid(false)
+		: IPListNode(InEditorWidget), bValue(false), ArrayIndex(-1), bFiltered(false), bArrayMember(false), bKeyValid(false)
 	{}
 
 public:
diff --git a/Engine/Source/Editor/PListEditor/Private/PListNodeDictionary.h b/Engine/Source/Editor/PListEditor/Private/PListNodeDictionary.h
--- a/Engine/Source/Editor/PListEditor/Private/PListNodeDictionary.h
+++ b/Engine/Source/Editor/PListEditor/Private/PListNodeDictionary.h
@@ -12,7 +12,7 @@ class FPListNodeDictionary : public IPListNode
 {
 public:
 	FPListNodeDictionary(SPListEditorPanel* InEditorWidget)
-		: IPListNode(InEditorWidget), bArrayMember(false), TableRow(nullptr), ArrayIndex(-1), bFiltered(false)
+		: IPListNode(InEditorWidget), TableRow(nullptr), ArrayIndex(-1), bArrayMember(false), bFiltered(false)
 	{}
 
 public:
diff --git a/Engine/Source/Editor/PListEditor/Private/PListNodeString.h b/Engine/Source/Editor/PListEditor/Private/PListNodeString.h
--- a/Engine/Source/Editor/PListEditor/Private/PListNodeString.h
+++ b/Engine/Source/Editor/PListEditor/Private/PListNodeString.h
@@ -12,7 +12,7 @@ class FPListNodeString : public IPListNode
 {
 public:
 	FPListNodeString(SPListEditorPanel* InEditorWidget)
-		: IPListNode(InEditorWidget), bFiltered(false), bArrayMember(false), ArrayIndex(-1), bKeyValid(false), bValueValid(false)
+		: IPListNode(InEditorWidget), ArrayIndex(-1), bFiltered(false), bArrayMember(false), bKeyValid(false), bValueValid(false)
 	{}
 
 public:
diff --git a/Engine/Source/Editor/PhAT/Private/PhATSharedData.cpp b/Engine/Source/Editor/PhAT/Private/PhATSharedData.cpp
--- a/Engine/Source/Editor/PhAT/Private/PhATSharedData.cpp
+++ b/Engine/Source/Editor/PhAT/Private/PhATSharedData.cpp
@@ -12,12 +12,12 @@
 #define LOCTEXT_NAMESPACE "PhATShared"
 
 FPhATSharedData::FPhATSharedData()
-	: COMRenderColor(255,255,100)
-	, PreviewScene( FPreviewScene::ConstructionValues().ShouldSimulatePhysics(true) )
-	, WidgetModeBeforeSimulation(FWidget::WM_None)
+	: PreviewScene( FPreviewScene::ConstructionValues().ShouldSimulatePhysics(true) )
+	, COMRenderColor(255,255,100)
 	, CopiedBodySetup(NULL)
 	, CopiedConstraintTemplate(NULL)
 	, bInsideSelChange(false)
+	, WidgetModeBeforeSimulation(FWidget::WM_None)
 {
 	// Editor variables
 	BodyEdit_MeshViewMode = PRM_Solid;
diff --git a/Engine/Source/Editor/PlacementMode/Private/PlacementMode.cpp b/Engine/Source/Editor/PlacementMode/Private/PlacementMode.cpp
--- a/Engine/Source/Editor/PlacementMode/Private/PlacementMode.cpp
+++ b/Engine/Source/Editor/PlacementMode/Private/PlacementMode.cpp
@@ -14,10 +14,10 @@
 #include "Editor/LevelEditor/Public/ILevelViewport.h"
 
 FPlacementMode::FPlacementMode()
-	: PlacementsChanged( false )
-	, CreatedPreviewActors( false )
-	, AssetsToPlace()
+	: AssetsToPlace()
 	, ActiveTransactionIndex( INDEX_NONE )
+	, PlacementsChanged( false )
+	, CreatedPreviewActors( false )
 {
 }
 
diff --git a/Engine/Source/Editor/SceneOutliner/Private/SceneOutlinerActorInfoColumn.cpp b/Engine/Source/Editor/SceneOutliner/Private/SceneOutlinerActorInfoColumn.cpp
--- a/Engine/Source/Editor/SceneOutliner/Private/SceneOutlinerActorInfoColumn.cpp
+++ b/Engine/Source/Editor/SceneOutliner/Private/SceneOutlinerActorInfoColumn.cpp
@@ -10,8 +10,8 @@
 TArray< TSharedPtr< ECustomColumnMode::Type > > FSceneOutlinerActorInfoColumn::ModeOptions;
 
 FSceneOutlinerActorInfoColumn::FSceneOutlinerActorInfoColumn( const TWeakPtr< ISceneOutliner >& InSceneOutlinerWeak, ECustomColumnMode::Type InDefaultCustomColumnMode )
-	: SceneOutlinerWeak( InSceneOutlinerWeak )
-	, CurrentMode( InDefaultCustomColumnMode )
+	: CurrentMode( InDefaultCustomColumnMode )
+	, SceneOutlinerWeak( InSceneOutlinerWeak )
 {
 
 }
diff --git a/Engine/Source/Editor/SequencerWidgets/Public/ITimeSlider.h b/Engine/Source/Editor/SequencerWidgets/Public/ITimeSlider.h
--- a/Engine/Source/Editor/SequencerWidgets/Public/ITimeSlider.h
+++ b/Engine/Source/Editor/SequencerWidgets/Public/ITimeSlider.h
@@ -9,14 +9,14 @@ struct FTimeSliderArgs
 {
 	FTimeSliderArgs()
 		: ScrubPosition(0)
-		, AllowZoom(true)
-		, OnScrubPositionChanged()
-		, OnBeginScrubberMovement()
-		, OnEndScrubberMovement()
 		, ViewRange( TRange<float>(0.0f, 5.0f) )
 		, ClampMin(-70000.0f)
 		, ClampMax(70000.0f)
+		, OnScrubPositionChanged()
+		, OnBeginScrubberMovement()
+		, OnEndScrubberMovement()
 		, OnViewRangeChanged()
+		, AllowZoom(true)
 	{}
 
 	/** The scrub position */
diff --git a/Engine/Source/Editor/UMGEditor/Private/Animation/UMGSequencerObjectBindingManager.cpp b/Engine/Source/Editor/UMGEditor/Private/Animation/UMGSequencerObjectBindingManager.cpp
--- a/Engine/Source/Editor/UMGEditor/Private/Animation/UMGSequencerObjectBindingManager.cpp
+++ b/Engine/Source/Editor/UMGEditor/Private/Animation/UMGSequencerObjectBindingManager.cpp
@@ -7,8 +7,8 @@
 #include "MovieScene.h"
 
 FUMGSequencerObjectBindingManager::FUMGSequencerObjectBindingManager( FWidgetBlueprintEditor& InWidgetBlueprintEditor, UMovieScene& InMovieScene )
-	: WidgetBlueprintEditor( InWidgetBlueprintEditor )
-	, MovieScene( &InMovieScene )
+	: MovieScene( &InMovieScene )
+	, WidgetBlueprintEditor( InWidgetBlueprintEditor )
 {
 	WidgetBlueprintEditor.GetOnWidgetPreviewUpdated().AddRaw( this, &FUMGSequencerObjectBindingManager::OnWidgetPreviewUpdated );
 
diff --git a/Engine/Source/Runtime/LauncherAutomatedService/Private/LauncherAutomatedServiceProvider.cpp b/Engine/Source/Runtime/LauncherAutomatedService/Private/LauncherAutomatedServiceProvider.cpp
--- a/Engine/Source/Runtime/LauncherAutomatedService/Private/LauncherAutomatedServiceProvider.cpp
+++ b/Engine/Source/Runtime/LauncherAutomatedService/Private/LauncherAutomatedServiceProvider.cpp
@@ -10,13 +10,13 @@ DEFINE_LOG_CATEGORY_STATIC(LauncherAutomatedService, Log, All);
  *****************************************************************************/
 
 FLauncherAutomatedServiceProvider::FLauncherAutomatedServiceProvider()
-	: LastDevicePingTime( 0.0 )
-	, TimeSinceLastShutdownRequest( 0.0 )
+	: bHasErrors( false )
 	, bHasLaunchedAllInstances( false )
 	, bIsReadyToShutdown( false )
 	, bShouldDeleteProfileWhenComplete( false )
+	, LastDevicePingTime( 0.0 )
 	, SessionID( FGuid::NewGuid() )
-	, bHasErrors( false )
+	, TimeSinceLastShutdownRequest( 0.0 )
 { }
 
 
