Landscape Tool : Fixed ensure for nav mesh for Add component tool (TTP #341789)

diff --git a/Engine/Source/Editor/LandscapeEditor/Private/LandscapeEdModeComponentTools.cpp b/Engine/Source/Editor/LandscapeEditor/Private/LandscapeEdModeComponentTools.cpp
--- a/Engine/Source/Editor/LandscapeEditor/Private/LandscapeEdModeComponentTools.cpp
+++ b/Engine/Source/Editor/LandscapeEditor/Private/LandscapeEdModeComponentTools.cpp
@@ -830,7 +830,11 @@ public:
 				}
 			}
 
-			Landscape->RegisterAllComponents();
+			// Need to register to use general height/xyoffset data update
+			for (int32 Idx = 0; Idx < NewComponents.Num(); Idx++)
+			{
+				NewComponents[Idx]->RegisterComponent();
+			}
 
 			if (bHasXYOffset)
 			{
