Landscape : Fixed Add Component regression (no interpolation height,  TTP #343057)

diff --git a/Engine/Source/Editor/LandscapeEditor/Private/LandscapeEdModeComponentTools.cpp b/Engine/Source/Editor/LandscapeEditor/Private/LandscapeEdModeComponentTools.cpp
--- a/Engine/Source/Editor/LandscapeEditor/Private/LandscapeEdModeComponentTools.cpp
+++ b/Engine/Source/Editor/LandscapeEditor/Private/LandscapeEdModeComponentTools.cpp
@@ -842,11 +842,9 @@ public:
 				XYOffsetCache.SetCachedData(X1, Y1, X2, Y2, XYOffsetData);
 				XYOffsetCache.Flush();
 			}
-			else
-			{
-				HeightCache.SetCachedData(X1, Y1, X2, Y2, Data);
-				HeightCache.Flush();
-			}
+
+			HeightCache.SetCachedData(X1, Y1, X2, Y2, Data);
+			HeightCache.Flush();
 
 			for (int32 Idx = 0; Idx < NewComponents.Num(); Idx++)
 			{
diff --git a/Engine/Source/Editor/LandscapeEditor/Private/LandscapeEdModeTools.h b/Engine/Source/Editor/LandscapeEditor/Private/LandscapeEdModeTools.h
--- a/Engine/Source/Editor/LandscapeEditor/Private/LandscapeEdModeTools.h
+++ b/Engine/Source/Editor/LandscapeEditor/Private/LandscapeEdModeTools.h
@@ -425,6 +425,26 @@ struct TLandscapeEditCache
 		CachedData.Add(ALandscape::MakeKey(LandscapeX, LandscapeY), Value);
 	}
 
+	bool IsZeroValue(const FVector& Value)
+	{
+		return (FMath::IsNearlyZero(Value.X) && FMath::IsNearlyZero(Value.Y));
+	}
+
+	bool IsZeroValue(const FVector2D& Value)
+	{
+		return (FMath::IsNearlyZero(Value.X) && FMath::IsNearlyZero(Value.Y));
+	}
+
+	bool IsZeroValue(const uint16& Value)
+	{
+		return Value == 0;
+	}
+
+	bool IsZeroValue(const uint8& Value)
+	{
+		return Value == 0;
+	}
+
 	bool GetCachedData(int32 X1, int32 Y1, int32 X2, int32 Y2, TArray<AccessorType>& OutData)
 	{
 		int32 NumSamples = (1 + X2 - X1)*(1 + Y2 - Y1);
@@ -440,7 +460,7 @@ struct TLandscapeEditCache
 				if (Ptr)
 				{
 					OutData[(X-X1) + (Y-Y1)*(1+X2-X1)] = *Ptr;
-					if (*Ptr != (AccessorType)0)
+					if (!IsZeroValue(*Ptr))
 					{
 						bHasNonZero = true;
 					}
@@ -634,7 +654,7 @@ struct FHeightmapAccessor
 			ULandscapeHeightfieldCollisionComponent* CollisionComponent = (*It)->CollisionComponent.Get();
 			if (CollisionComponent)
 			{
-				CollisionComponent->RecreateCollision(false);
+				CollisionComponent->RecreateCollision(true);
 				UNavigationSystem* NavSys = UNavigationSystem::GetCurrent(*It);
 				if (NavSys)
 				{
