Fix crash due to the the heightmap texture being locked twice when adding components with XY offset.
TTP#343040

diff --git a/Engine/Source/Editor/LandscapeEditor/Private/LandscapeEdModeComponentTools.cpp b/Engine/Source/Editor/LandscapeEditor/Private/LandscapeEdModeComponentTools.cpp
--- a/Engine/Source/Editor/LandscapeEditor/Private/LandscapeEdModeComponentTools.cpp
+++ b/Engine/Source/Editor/LandscapeEditor/Private/LandscapeEdModeComponentTools.cpp
@@ -756,7 +756,12 @@ public:
 		,	XYOffsetCache(InTarget)
 	{}
 
-	virtual ~FLandscapeToolStrokeAddComponent() {}
+	virtual ~FLandscapeToolStrokeAddComponent()
+	{
+		// We flush here so here ~FXYOffsetmapAccessor can safely lock the heightmap data to update bounds
+		HeightCache.Flush();
+		XYOffsetCache.Flush();
+	}
 
 	virtual void Apply(FEditorViewportClient* ViewportClient, FLandscapeBrush* Brush, const ULandscapeEditorObject* UISettings, const TArray<FLandscapeToolMousePosition>& MousePositions)
 	{
diff --git a/Engine/Source/Editor/LandscapeEditor/Private/LandscapeEdModeTools.h b/Engine/Source/Editor/LandscapeEditor/Private/LandscapeEdModeTools.h
--- a/Engine/Source/Editor/LandscapeEditor/Private/LandscapeEdModeTools.h
+++ b/Engine/Source/Editor/LandscapeEditor/Private/LandscapeEdModeTools.h
@@ -604,8 +604,12 @@ struct FHeightmapAccessor
 
 				for (ULandscapeComponent* Component : Components)
 				{
-					CollisionComponents.Add(Component->CollisionComponent.Get());
-					PreUpdateLocalBoxes.Add(FBox(FVector((float)X1, (float)Y1, Component->CachedLocalBox.Min.Z), FVector((float)X2, (float)Y2, Component->CachedLocalBox.Max.Z)));
+					ULandscapeHeightfieldCollisionComponent* CollisionComponent = Component->CollisionComponent.Get();
+					if (CollisionComponent)
+					{
+						CollisionComponents.Add(CollisionComponent);
+						PreUpdateLocalBoxes.Add(FBox(FVector((float)X1, (float)Y1, Component->CachedLocalBox.Min.Z), FVector((float)X2, (float)Y2, Component->CachedLocalBox.Max.Z)));
+					}
 				}
 
 				// Update landscape.
