Fix for the Move to Level tool on landscape with painted layers causing broken references on save

diff --git a/Engine/Source/Editor/LandscapeEditor/Private/LandscapeEdModeComponentTools.cpp b/Engine/Source/Editor/LandscapeEditor/Private/LandscapeEdModeComponentTools.cpp
--- a/Engine/Source/Editor/LandscapeEditor/Private/LandscapeEdModeComponentTools.cpp
+++ b/Engine/Source/Editor/LandscapeEditor/Private/LandscapeEdModeComponentTools.cpp
@@ -575,7 +575,7 @@ public:
 								int32 WeightmapSize = (Comp->SubsectionSizeQuads + 1) * Comp->NumSubsections;
 
 								// We need a new weightmap texture
-								CurrentWeightmapTexture = Comp->GetLandscapeProxy()->CreateLandscapeTexture(WeightmapSize, WeightmapSize, TEXTUREGROUP_Terrain_Weightmap, TSF_BGRA8);
+								CurrentWeightmapTexture = LandscapeProxy->CreateLandscapeTexture(WeightmapSize, WeightmapSize, TEXTUREGROUP_Terrain_Weightmap, TSF_BGRA8);
 								// Alloc dummy mips
 								Comp->CreateEmptyTextureMips(CurrentWeightmapTexture);
 								CurrentWeightmapTexture->PostEditChange();
