Fix for editor builds.

diff --git a/Engine/Source/Editor/LandscapeEditor/Private/LandscapeEditorDetailCustomization_NewLandscape.cpp b/Engine/Source/Editor/LandscapeEditor/Private/LandscapeEditorDetailCustomization_NewLandscape.cpp
--- a/Engine/Source/Editor/LandscapeEditor/Private/LandscapeEditorDetailCustomization_NewLandscape.cpp
+++ b/Engine/Source/Editor/LandscapeEditor/Private/LandscapeEditorDetailCustomization_NewLandscape.cpp
@@ -945,7 +945,7 @@ EVisibility FLandscapeEditorDetailCustomization_NewLandscape::GetHeightmapErrorV
 		return EVisibility::Visible;
 	}
 
-	if (HeightmapErrorAsByte != ELandscapeImportLayerError::None)
+	if (HeightmapErrorAsByte != ELandscapeImportHeightmapError::None)
 	{
 		return EVisibility::Visible;
 	}
@@ -1566,8 +1566,8 @@ void FLandscapeEditorStructCustomization_FLandscapeImportLayer::OnImportLayerCre
 
 EVisibility FLandscapeEditorStructCustomization_FLandscapeImportLayer::GetErrorVisibility(TSharedRef<IPropertyHandle> PropertyHandle_LayerError)
 {
-	uint8 LayerErrorAsByte = 0;
-	FPropertyAccess::Result Result = PropertyHandle_LayerError->GetValue(LayerErrorAsByte);
+	auto LayerError = ELandscapeImportLayerError::None;
+	FPropertyAccess::Result Result = PropertyHandle_LayerError->GetValue((uint8&)LayerError);
 	check(Result == FPropertyAccess::Success);
 
 	if (Result == FPropertyAccess::MultipleValues)
@@ -1585,7 +1585,7 @@ EVisibility FLandscapeEditorStructCustomization_FLandscapeImportLayer::GetErrorV
 		}
 	}
 
-	if (LayerErrorAsByte != ELandscapeImportLayerError::None)
+	if (LayerError != ELandscapeImportLayerError::None)
 	{
 		return EVisibility::Visible;
 	}
@@ -1594,8 +1594,8 @@ EVisibility FLandscapeEditorStructCustomization_FLandscapeImportLayer::GetErrorV
 
 FText FLandscapeEditorStructCustomization_FLandscapeImportLayer::GetErrorText(TSharedRef<IPropertyHandle> PropertyHandle_LayerError)
 {
-	uint8 LayerErrorAsByte = 0;
-	FPropertyAccess::Result Result = PropertyHandle_LayerError->GetValue(LayerErrorAsByte);
+	auto LayerError = ELandscapeImportLayerError::None;
+	FPropertyAccess::Result Result = PropertyHandle_LayerError->GetValue((uint8&)LayerError);
 	check(Result == FPropertyAccess::Success);
 
 	if (Result == FPropertyAccess::MultipleValues)
@@ -1603,7 +1603,7 @@ FText FLandscapeEditorStructCustomization_FLandscapeImportLayer::GetErrorText(TS
 		return NSLOCTEXT("PropertyEditor", "MultipleValues", "Multiple Values");
 	}
 
-	switch (LayerErrorAsByte)
+	switch (LayerError)
 	{
 	case ELandscapeImportLayerError::None:
 		return FText::GetEmpty();
