Allow undoing changes to the Landscape "New Landscape" panel
Fix crash when replacing material in the "New Landscape" panel with one with fewer material layers (or null)
Fix undo transactions for "interactive" changes (e.g. sliding sliders) so they don't flood the undo buffer with one transaction for every sub-step of the change

diff --git a/Engine/Source/Editor/LandscapeEditor/Classes/LandscapeEditorObject.h b/Engine/Source/Editor/LandscapeEditor/Classes/LandscapeEditorObject.h
--- a/Engine/Source/Editor/LandscapeEditor/Classes/LandscapeEditorObject.h
+++ b/Engine/Source/Editor/LandscapeEditor/Classes/LandscapeEditorObject.h
@@ -198,169 +198,169 @@ class ULandscapeEditorObject : public UObject
 	// Common Tool Settings:
 
 	// Strength of the tool
-	UPROPERTY(Category="Tool Settings", EditAnywhere, meta=(ShowForTools="Paint,Sculpt,Smooth,Flatten,Erosion,HydraErosion,Noise,Mask,CopyPaste", ClampMin="0", ClampMax="10", UIMin="0", UIMax="1"))
+	UPROPERTY(Category="Tool Settings", EditAnywhere, NonTransactional, meta=(ShowForTools="Paint,Sculpt,Smooth,Flatten,Erosion,HydraErosion,Noise,Mask,CopyPaste", ClampMin="0", ClampMax="10", UIMin="0", UIMax="1"))
 	float ToolStrength;
 
 	// Enable to make tools blend towards a target value
-	UPROPERTY()
+	UPROPERTY(NonTransactional)
 	bool bUseWeightTargetValue;
 
 	// Enable to make tools blend towards a target value
-	UPROPERTY(Category="Tool Settings", EditAnywhere, meta=(DisplayName="Use Target Value", editcondition="bUseWeightTargetValue", ShowForTools="Paint,Sculpt,Noise", ClampMin="0", ClampMax="10", UIMin="0", UIMax="1"))
+	UPROPERTY(Category="Tool Settings", EditAnywhere, NonTransactional, meta=(DisplayName="Use Target Value", editcondition="bUseWeightTargetValue", ShowForTools="Paint,Sculpt,Noise", ClampMin="0", ClampMax="10", UIMin="0", UIMax="1"))
 	float WeightTargetValue;
 
 	// I have no idea what this is for but it's used by the noise and erosion tools, and isn't exposed to the UI
-	UPROPERTY()
+	UPROPERTY(NonTransactional)
 	float MaximumValueRadius;
 
 	// Flatten Tool:
 
 	// Whether to flatten by lowering, raising, or both
-	UPROPERTY(Category="Tool Settings", EditAnywhere, meta=(ShowForTools="Flatten"))
+	UPROPERTY(Category="Tool Settings", EditAnywhere, NonTransactional, meta=(ShowForTools="Flatten"))
 	TEnumAsByte<ELandscapeToolFlattenMode::Type> FlattenMode;
 
 	// Flattens to the angle of the clicked point, instead of horizontal
-	UPROPERTY(Category="Tool Settings", EditAnywhere, meta=(ShowForTools="Flatten", ShowForTargetTypes="Heightmap"))
+	UPROPERTY(Category="Tool Settings", EditAnywhere, NonTransactional, meta=(ShowForTools="Flatten", ShowForTargetTypes="Heightmap"))
 	bool bUseSlopeFlatten;
 
 	// Constantly picks new values to flatten towards when dragging around, instead of only using the first clicked point
-	UPROPERTY(Category="Tool Settings", EditAnywhere, meta=(ShowForTools="Flatten", ShowForTargetTypes="Heightmap"))
+	UPROPERTY(Category="Tool Settings", EditAnywhere, NonTransactional, meta=(ShowForTools="Flatten", ShowForTargetTypes="Heightmap"))
 	bool bPickValuePerApply;
 
 	// Enable to flatten towards a target height
-	UPROPERTY()
+	UPROPERTY(NonTransactional)
 	bool bUseFlattenTarget;
 
 	// Target height to flatten towards (in Unreal Units)
-	UPROPERTY(Category="Tool Settings", EditAnywhere, meta=(ShowForTools="Flatten", ShowForTargetTypes="Heightmap", editcondition="bUseFlattenTarget", UIMin="-32768", UIMax="32768"))
+	UPROPERTY(Category="Tool Settings", EditAnywhere, NonTransactional, meta=(ShowForTools="Flatten", ShowForTargetTypes="Heightmap", editcondition="bUseFlattenTarget", UIMin="-32768", UIMax="32768"))
 	float FlattenTarget;
 
 	// Whether to show the preview grid for the flatten target height
-	UPROPERTY(Category = "Tool Settings", EditAnywhere, AdvancedDisplay, meta = (DisplayName = "Show Preview Grid", ShowForTools = "Flatten", ShowForTargetTypes = "Heightmap", editcondition = "bUseFlattenTarget", HideEditConditionToggle, UIMin = "-32768", UIMax = "32768"))
+	UPROPERTY(Category = "Tool Settings", EditAnywhere, NonTransactional, AdvancedDisplay, meta = (DisplayName = "Show Preview Grid", ShowForTools = "Flatten", ShowForTargetTypes = "Heightmap", editcondition = "bUseFlattenTarget", HideEditConditionToggle, UIMin = "-32768", UIMax = "32768"))
 	bool bShowFlattenTargetPreview;
 
 	// Ramp Tool:
 
 	// Width of ramp
-	UPROPERTY(Category="Tool Settings", EditAnywhere, meta=(ShowForTools="Ramp", ClampMin="1", UIMin="1", UIMax="8192", SliderExponent=3))
+	UPROPERTY(Category="Tool Settings", EditAnywhere, NonTransactional, meta=(ShowForTools="Ramp", ClampMin="1", UIMin="1", UIMax="8192", SliderExponent=3))
 	float RampWidth;
 
 	// Falloff on side of ramp
-	UPROPERTY(Category="Tool Settings", EditAnywhere, meta=(DisplayName="Side Falloff", ShowForTools="Ramp", ClampMin="0", ClampMax="1", UIMin="0", UIMax="1"))
+	UPROPERTY(Category="Tool Settings", EditAnywhere, NonTransactional, meta=(DisplayName="Side Falloff", ShowForTools="Ramp", ClampMin="0", ClampMax="1", UIMin="0", UIMax="1"))
 	float RampSideFalloff;
 
 	// Smooth Tool:
 
 	// Scale multiplier for the smoothing filter kernel
-	UPROPERTY(Category="Tool Settings", EditAnywhere, meta=(DisplayName="Filter Kernel Scale", ShowForTools="Smooth", ClampMin="0", ClampMax="1", UIMin="0", UIMax="1"))
+	UPROPERTY(Category="Tool Settings", EditAnywhere, NonTransactional, meta=(DisplayName="Filter Kernel Scale", ShowForTools="Smooth", ClampMin="0", ClampMax="1", UIMin="0", UIMax="1"))
 	float SmoothFilterKernelScale;
 
 	// If checked, performs a detail preserving smooth using the specified detail smoothing value
-	UPROPERTY()
+	UPROPERTY(NonTransactional)
 	bool bDetailSmooth;
 
 	// Larger detail smoothing values remove more details, while smaller values preserve more details
-	UPROPERTY(Category="Tool Settings", EditAnywhere, meta=(DisplayName="Detail Smooth", editcondition="bDetailSmooth", ShowForTools="Smooth", ClampMin="0", ClampMax="0.99"))
+	UPROPERTY(Category="Tool Settings", EditAnywhere, NonTransactional, meta=(DisplayName="Detail Smooth", editcondition="bDetailSmooth", ShowForTools="Smooth", ClampMin="0", ClampMax="0.99"))
 	float DetailScale;
 
 	// Erosion Tool:
 
 	// The minimum height difference necessary for the erosion effects to be applied. Smaller values will result in more erosion being applied
-	UPROPERTY(Category="Tool Settings", EditAnywhere, meta=(DisplayName="Threshold", ShowForTools="Erosion", ClampMin="0", ClampMax="256", UIMin="0", UIMax="128"))
+	UPROPERTY(Category="Tool Settings", EditAnywhere, NonTransactional, meta=(DisplayName="Threshold", ShowForTools="Erosion", ClampMin="0", ClampMax="256", UIMin="0", UIMax="128"))
 	int32 ErodeThresh;
 
 	// The thickness of the surface for the layer weight erosion effect
-	UPROPERTY(Category="Tool Settings", EditAnywhere, meta=(DisplayName="Surface Thickness", ShowForTools="Erosion", ClampMin="128", ClampMax="1024", UIMin="128", UIMax="512"))
+	UPROPERTY(Category="Tool Settings", EditAnywhere, NonTransactional, meta=(DisplayName="Surface Thickness", ShowForTools="Erosion", ClampMin="128", ClampMax="1024", UIMin="128", UIMax="512"))
 	int32 ErodeSurfaceThickness;
 
 	// Number of erosion iterations, more means more erosion but is slower
-	UPROPERTY(Category="Tool Settings", EditAnywhere, meta=(DisplayName="Iterations", ShowForTools="Erosion", ClampMin="1", ClampMax="300", UIMin="1", UIMax="150"))
+	UPROPERTY(Category="Tool Settings", EditAnywhere, NonTransactional, meta=(DisplayName="Iterations", ShowForTools="Erosion", ClampMin="1", ClampMax="300", UIMin="1", UIMax="150"))
 	int32 ErodeIterationNum;
 
 	// Whether to erode by lowering, raising, or both
-	UPROPERTY(Category="Tool Settings", EditAnywhere, meta=(DisplayName="Noise Mode", ShowForTools="Erosion"))
+	UPROPERTY(Category="Tool Settings", EditAnywhere, NonTransactional, meta=(DisplayName="Noise Mode", ShowForTools="Erosion"))
 	TEnumAsByte<ELandscapeToolErosionMode::Type> ErosionNoiseMode;
 
 	// The size of the perlin noise filter used
-	UPROPERTY(Category="Tool Settings", EditAnywhere, meta=(DisplayName="Noise Scale", ShowForTools="Erosion", ClampMin="1", ClampMax="512", UIMin="1.1", UIMax="256"))
+	UPROPERTY(Category="Tool Settings", EditAnywhere, NonTransactional, meta=(DisplayName="Noise Scale", ShowForTools="Erosion", ClampMin="1", ClampMax="512", UIMin="1.1", UIMax="256"))
 	float ErosionNoiseScale;
 
 	// Hydraulic Erosion Tool:
 
 	// The amount of rain to apply to the surface. Larger values will result in more erosion
-	UPROPERTY(Category="Tool Settings", EditAnywhere, meta=(ShowForTools="HydraErosion", ClampMin="1", ClampMax="512", UIMin="1", UIMax="256"))
+	UPROPERTY(Category="Tool Settings", EditAnywhere, NonTransactional, meta=(ShowForTools="HydraErosion", ClampMin="1", ClampMax="512", UIMin="1", UIMax="256"))
 	int32 RainAmount;
 
 	// The amount of sediment that the water can carry. Larger values will result in more erosion
-	UPROPERTY(Category="Tool Settings", EditAnywhere, meta=(DisplayName="Sediment Cap.", ShowForTools="HydraErosion", ClampMin="0.1", ClampMax="1.0"))
+	UPROPERTY(Category="Tool Settings", EditAnywhere, NonTransactional, meta=(DisplayName="Sediment Cap.", ShowForTools="HydraErosion", ClampMin="0.1", ClampMax="1.0"))
 	float SedimentCapacity;
 
 	// Number of erosion iterations, more means more erosion but is slower
-	UPROPERTY(Category="Tool Settings", EditAnywhere, meta=(DisplayName="Iterations", ShowForTools="HydraErosion", ClampMin="1", ClampMax="300", UIMin="1", UIMax="150"))
+	UPROPERTY(Category="Tool Settings", EditAnywhere, NonTransactional, meta=(DisplayName="Iterations", ShowForTools="HydraErosion", ClampMin="1", ClampMax="300", UIMin="1", UIMax="150"))
 	int32 HErodeIterationNum;
 
 	// Initial Rain Distribution
-	UPROPERTY(Category="Tool Settings", EditAnywhere, meta=(DisplayName="Initial Rain Distribution", ShowForTools="HydraErosion"))
+	UPROPERTY(Category="Tool Settings", EditAnywhere, NonTransactional, meta=(DisplayName="Initial Rain Distribution", ShowForTools="HydraErosion"))
 	TEnumAsByte<ELandscapeToolHydroErosionMode::Type> RainDistMode;
 
 	// The size of the noise filter for applying initial rain to the surface
-	UPROPERTY(Category="Tool Settings", EditAnywhere, meta=(ShowForTools="HydraErosion", ClampMin="1", ClampMax="512", UIMin="1.1", UIMax="256"))
+	UPROPERTY(Category="Tool Settings", EditAnywhere, NonTransactional, meta=(ShowForTools="HydraErosion", ClampMin="1", ClampMax="512", UIMin="1.1", UIMax="256"))
 	float RainDistScale;
 
 	// If checked, performs a detail-preserving smooth to the erosion effect using the specified detail smoothing value
-	UPROPERTY()
+	UPROPERTY(NonTransactional)
 	bool bHErosionDetailSmooth;
 
 	// Larger detail smoothing values remove more details, while smaller values preserve more details
-	UPROPERTY(Category="Tool Settings", EditAnywhere, meta=(DisplayName="Detail Smooth", editcondition="bHErosionDetailSmooth", ShowForTools="HydraErosion", ClampMin="0", ClampMax="0.99"))
+	UPROPERTY(Category="Tool Settings", EditAnywhere, NonTransactional, meta=(DisplayName="Detail Smooth", editcondition="bHErosionDetailSmooth", ShowForTools="HydraErosion", ClampMin="0", ClampMax="0.99"))
 	float HErosionDetailScale;
 
 	// Noise Tool:
 
 	// Whether to apply noise that raises, lowers, or both
-	UPROPERTY(Category="Tool Settings", EditAnywhere, meta=(DisplayName="Noise Mode", ShowForTools="Noise"))
+	UPROPERTY(Category="Tool Settings", EditAnywhere, NonTransactional, meta=(DisplayName="Noise Mode", ShowForTools="Noise"))
 	TEnumAsByte<ELandscapeToolNoiseMode::Type> NoiseMode;
 
 	// The size of the perlin noise filter used
-	UPROPERTY(Category="Tool Settings", EditAnywhere, meta=(DisplayName="Noise Scale", ShowForTools="Noise", ClampMin="1", ClampMax="512", UIMin="1.1", UIMax="256"))
+	UPROPERTY(Category="Tool Settings", EditAnywhere, NonTransactional, meta=(DisplayName="Noise Scale", ShowForTools="Noise", ClampMin="1", ClampMax="512", UIMin="1.1", UIMax="256"))
 	float NoiseScale;
 
 	// Mask Tool:
 
 	// Uses selected region as a mask for other tools
-	UPROPERTY(Category="Tool Settings", EditAnywhere, meta=(DisplayName="Use Region as Mask", ShowForTools="Mask", ShowForMask))
+	UPROPERTY(Category="Tool Settings", EditAnywhere, NonTransactional, meta=(DisplayName="Use Region as Mask", ShowForTools="Mask", ShowForMask))
 	bool bUseSelectedRegion;
 
 	// If enabled, protects the selected region from changes
 	// If disabled, only allows changes in the selected region
-	UPROPERTY(Category="Tool Settings", EditAnywhere, meta=(DisplayName="Negative Mask", ShowForTools="Mask", ShowForMask))
+	UPROPERTY(Category="Tool Settings", EditAnywhere, NonTransactional, meta=(DisplayName="Negative Mask", ShowForTools="Mask", ShowForMask))
 	bool bUseNegativeMask;
 
 	// Copy/Paste Tool:
 
 	// Whether to paste will only raise, only lower, or both
-	UPROPERTY(Category="Tool Settings", EditAnywhere, meta=(ShowForTools="CopyPaste"))
+	UPROPERTY(Category="Tool Settings", EditAnywhere, NonTransactional, meta=(ShowForTools="CopyPaste"))
 	TEnumAsByte<ELandscapeToolPasteMode::Type> PasteMode;
 
 	// If set, copies/pastes all layers, otherwise only copy/pastes the layer selected in the targets panel
-	UPROPERTY(Category="Tool Settings", EditAnywhere, meta=(DisplayName="Gizmo copy/paste all layers", ShowForTools="CopyPaste"))
+	UPROPERTY(Category="Tool Settings", EditAnywhere, NonTransactional, meta=(DisplayName="Gizmo copy/paste all layers", ShowForTools="CopyPaste"))
 	bool bApplyToAllTargets;
 
 	// Makes sure the gizmo is snapped perfectly to the landscape so that the sample points line up, which makes copy/paste less blurry. Irrelevant if gizmo is scaled
-	UPROPERTY(Category="Tool Settings", EditAnywhere, meta=(DisplayName="Snap Gizmo to Landscape grid", ShowForTools="CopyPaste"))
+	UPROPERTY(Category="Tool Settings", EditAnywhere, NonTransactional, meta=(DisplayName="Snap Gizmo to Landscape grid", ShowForTools="CopyPaste"))
 	bool bSnapGizmo;
 
 	// Smooths the edges of the gizmo data into the landscape. Without this, the edges of the pasted data will be sharp
-	UPROPERTY(Category="Tool Settings", EditAnywhere, meta=(DisplayName="Use Smooth Gizmo Brush", ShowForTools="CopyPaste"))
+	UPROPERTY(Category="Tool Settings", EditAnywhere, NonTransactional, meta=(DisplayName="Use Smooth Gizmo Brush", ShowForTools="CopyPaste"))
 	bool bSmoothGizmoBrush;
 
-	UPROPERTY(Category="Tool Settings", EditAnywhere, AdvancedDisplay, meta=(DisplayName="Heightmap", ShowForTools="CopyPaste"))
+	UPROPERTY(Category="Tool Settings", EditAnywhere, NonTransactional, AdvancedDisplay, meta=(DisplayName="Heightmap", ShowForTools="CopyPaste"))
 	FString GizmoHeightmapFilenameString;
 
-	UPROPERTY(Category="Tool Settings", EditAnywhere, AdvancedDisplay, meta=(DisplayName="Heightmap Size", ShowForTools="CopyPaste"))
+	UPROPERTY(Category="Tool Settings", EditAnywhere, NonTransactional, AdvancedDisplay, meta=(DisplayName="Heightmap Size", ShowForTools="CopyPaste"))
 	FIntPoint GizmoImportSize;
 
-	UPROPERTY(Category="Tool Settings", EditAnywhere, AdvancedDisplay, meta=(DisplayName="Layers", ShowForTools="CopyPaste"))
+	UPROPERTY(Category="Tool Settings", EditAnywhere, NonTransactional, AdvancedDisplay, meta=(DisplayName="Layers", ShowForTools="CopyPaste"))
 	TArray<FGizmoImportLayer> GizmoImportLayers;
 
 	TArray<FGizmoHistory> GizmoHistories;
@@ -369,19 +369,19 @@ class ULandscapeEditorObject : public UObject
 	// Resize Landscape Tool
 
 	// Number of quads per landscape component section
-	UPROPERTY(Category="Change Component Size", EditAnywhere, meta=(DisplayName="Section Size", ShowForTools="ResizeLandscape"))
+	UPROPERTY(Category="Change Component Size", EditAnywhere, NonTransactional, meta=(DisplayName="Section Size", ShowForTools="ResizeLandscape"))
 	int32 ResizeLandscape_QuadsPerSection;
 
 	// Number of sections per landscape component
-	UPROPERTY(Category="Change Component Size", EditAnywhere, meta=(DisplayName="Sections Per Component", ShowForTools="ResizeLandscape"))
+	UPROPERTY(Category="Change Component Size", EditAnywhere, NonTransactional, meta=(DisplayName="Sections Per Component", ShowForTools="ResizeLandscape"))
 	int32 ResizeLandscape_SectionsPerComponent;
 
 	// Number of components in resulting landscape
-	UPROPERTY(Category="Change Component Size", EditAnywhere, meta=(DisplayName="Number of Components", ShowForTools="ResizeLandscape"))
+	UPROPERTY(Category="Change Component Size", EditAnywhere, NonTransactional, meta=(DisplayName="Number of Components", ShowForTools="ResizeLandscape"))
 	FIntPoint ResizeLandscape_ComponentCount;
 
 	// 
-	UPROPERTY(Category="Change Component Size", EditAnywhere, meta=(DisplayName="Resize Mode", ShowForTools="ResizeLandscape"))
+	UPROPERTY(Category="Change Component Size", EditAnywhere, NonTransactional, meta=(DisplayName="Resize Mode", ShowForTools="ResizeLandscape"))
 	TEnumAsByte<ELandscapeConvertMode::Type> ResizeLandscape_ConvertMode;
 
 	int32 ResizeLandscape_Original_QuadsPerSection;
@@ -410,82 +410,82 @@ class ULandscapeEditorObject : public UObject
 
 	// Import
 
-	UPROPERTY(Category="New Landscape", VisibleAnywhere, meta=(ShowForTools="NewLandscape"))
+	UPROPERTY(Category="New Landscape", VisibleAnywhere, NonTransactional, meta=(ShowForTools="NewLandscape"))
 	TEnumAsByte<ELandscapeImportHeightmapError::Type> ImportLandscape_HeightmapError;
-	UPROPERTY(Category="New Landscape", EditAnywhere, meta=(DisplayName="Heightmap File", ShowForTools="NewLandscape"))
+	UPROPERTY(Category="New Landscape", EditAnywhere, NonTransactional, meta=(DisplayName="Heightmap File", ShowForTools="NewLandscape"))
 	FString ImportLandscape_HeightmapFilename;
-	UPROPERTY()
+	UPROPERTY(NonTransactional)
 	int32 ImportLandscape_Width;
-	UPROPERTY()
+	UPROPERTY(NonTransactional)
 	int32 ImportLandscape_Height;
-	UPROPERTY()
+	UPROPERTY(NonTransactional)
 	TArray<uint16> ImportLandscape_Data;
 
-	UPROPERTY(Category="New Landscape", EditAnywhere, EditFixedSize, meta=(DisplayName="Layers", ShowForTools="NewLandscape"))
+	UPROPERTY(Category="New Landscape", EditAnywhere, NonTransactional, EditFixedSize, meta=(DisplayName="Layers", ShowForTools="NewLandscape"))
 	TArray<FLandscapeImportLayer> ImportLandscape_Layers;
 
 	// Common Brush Settings:
 
 	// The radius of the brush, in unreal units
-	UPROPERTY(Category="Brush Settings", EditAnywhere, meta=(DisplayName="Brush Size", ShowForBrushes="BrushSet_Circle,BrushSet_Alpha,BrushSet_Pattern", ClampMin="1", ClampMax="65536", UIMin="1", UIMax="8192", SliderExponent="3"))
+	UPROPERTY(Category="Brush Settings", EditAnywhere, NonTransactional, meta=(DisplayName="Brush Size", ShowForBrushes="BrushSet_Circle,BrushSet_Alpha,BrushSet_Pattern", ClampMin="1", ClampMax="65536", UIMin="1", UIMax="8192", SliderExponent="3"))
 	float BrushRadius;
 
 	// The falloff at the edge of the brush, as a fraction of the brush's size. 0 = no falloff, 1 = all falloff
-	UPROPERTY(Category="Brush Settings", EditAnywhere, meta=(ShowForBrushes="BrushSet_Circle,BrushSet_Gizmo,BrushSet_Pattern", ClampMin="0", ClampMax="1"))
+	UPROPERTY(Category="Brush Settings", EditAnywhere, NonTransactional, meta=(ShowForBrushes="BrushSet_Circle,BrushSet_Gizmo,BrushSet_Pattern", ClampMin="0", ClampMax="1"))
 	float BrushFalloff;
 
 	// Selects the Clay Brush painting mode
-	UPROPERTY(Category="Brush Settings", EditAnywhere, meta=(ShowForTools="Sculpt", ShowForBrushes="BrushSet_Circle,BrushSet_Alpha,BrushSet_Pattern"))
+	UPROPERTY(Category="Brush Settings", EditAnywhere, NonTransactional, meta=(ShowForTools="Sculpt", ShowForBrushes="BrushSet_Circle,BrushSet_Alpha,BrushSet_Pattern"))
 	bool bUseClayBrush;
 
 	// Alpha/Pattern Brush:
 
 	// Scale of the brush texture. A scale of 1.000 maps the brush texture to the landscape at a 1 pixel = 1 vertex size
-	UPROPERTY(Category="Brush Settings", EditAnywhere, meta=(DisplayName="Texture Scale", ShowForBrushes="BrushSet_Pattern", ClampMin="0.005", ClampMax="5", SliderExponent="3"))
+	UPROPERTY(Category="Brush Settings", EditAnywhere, NonTransactional, meta=(DisplayName="Texture Scale", ShowForBrushes="BrushSet_Pattern", ClampMin="0.005", ClampMax="5", SliderExponent="3"))
 	float AlphaBrushScale;
 
 	// Rotate brush to follow mouse
-	UPROPERTY(Category = "Brush Settings", EditAnywhere, meta = (DisplayName = "Auto-Rotate", ShowForBrushes = "BrushSet_Alpha"))
+	UPROPERTY(Category = "Brush Settings", EditAnywhere, NonTransactional, meta = (DisplayName = "Auto-Rotate", ShowForBrushes = "BrushSet_Alpha"))
 	bool bAlphaBrushAutoRotate;
 
 	// Rotates the brush mask texture
-	UPROPERTY(Category="Brush Settings", EditAnywhere, meta=(DisplayName="Texture Rotation", ShowForBrushes="BrushSet_Alpha,BrushSet_Pattern", ClampMin="-360", ClampMax="360", UIMin="-180", UIMax="180"))
+	UPROPERTY(Category="Brush Settings", EditAnywhere, NonTransactional, meta=(DisplayName="Texture Rotation", ShowForBrushes="BrushSet_Alpha,BrushSet_Pattern", ClampMin="-360", ClampMax="360", UIMin="-180", UIMax="180"))
 	float AlphaBrushRotation;
 
 	// Horizontally offsets the brush mask texture
-	UPROPERTY(Category="Brush Settings", EditAnywhere, meta=(DisplayName="Texture Pan U", ShowForBrushes="BrushSet_Pattern", ClampMin="0", ClampMax="1"))
+	UPROPERTY(Category="Brush Settings", EditAnywhere, NonTransactional, meta=(DisplayName="Texture Pan U", ShowForBrushes="BrushSet_Pattern", ClampMin="0", ClampMax="1"))
 	float AlphaBrushPanU;
 
 	// Vertically offsets the brush mask texture
-	UPROPERTY(Category="Brush Settings", EditAnywhere, meta=(DisplayName="Texture Pan V", ShowForBrushes="BrushSet_Pattern", ClampMin="0", ClampMax="1"))
+	UPROPERTY(Category="Brush Settings", EditAnywhere, NonTransactional, meta=(DisplayName="Texture Pan V", ShowForBrushes="BrushSet_Pattern", ClampMin="0", ClampMax="1"))
 	float AlphaBrushPanV;
 
 	// Mask texture to use
-	UPROPERTY(Category="Brush Settings", EditAnywhere, meta=(DisplayName="Texture", ShowForBrushes="BrushSet_Alpha,BrushSet_Pattern"))
+	UPROPERTY(Category="Brush Settings", EditAnywhere, NonTransactional, meta=(DisplayName="Texture", ShowForBrushes="BrushSet_Alpha,BrushSet_Pattern"))
 	UTexture2D* AlphaTexture;
 
 	// Channel of Mask Texture to use
-	UPROPERTY(Category="Brush Settings", EditAnywhere, meta=(DisplayName="Texture Channel", ShowForBrushes="BrushSet_Alpha,BrushSet_Pattern"))
+	UPROPERTY(Category="Brush Settings", EditAnywhere, NonTransactional, meta=(DisplayName="Texture Channel", ShowForBrushes="BrushSet_Alpha,BrushSet_Pattern"))
 	TEnumAsByte<EColorChannel::Type> AlphaTextureChannel;
 
-	UPROPERTY()
+	UPROPERTY(NonTransactional)
 	int32 AlphaTextureSizeX;
-	UPROPERTY()
+	UPROPERTY(NonTransactional)
 	int32 AlphaTextureSizeY;
-	UPROPERTY()
+	UPROPERTY(NonTransactional)
 	TArray<uint8> AlphaTextureData;
 
 	// Component Brush:
 
 	// Number of components X/Y to affect at once. 1 means 1x1, 2 means 2x2, etc
-	UPROPERTY(Category="Brush Settings", EditAnywhere, meta=(DisplayName="Brush Size", ShowForBrushes="BrushSet_Component", ClampMin="1", ClampMax="128", UIMin="1", UIMax="64", SliderExponent="3"))
+	UPROPERTY(Category="Brush Settings", EditAnywhere, NonTransactional, meta=(DisplayName="Brush Size", ShowForBrushes="BrushSet_Component", ClampMin="1", ClampMax="128", UIMin="1", UIMax="64", SliderExponent="3"))
 	int32 BrushComponentSize;
 
 
 	// Target Layer Settings:
 
 	// Limits painting to only the components that already have the selected layer
-	UPROPERTY(Category="Target Layers", EditAnywhere)
+	UPROPERTY(Category="Target Layers", EditAnywhere, NonTransactional)
 	TEnumAsByte<ELandscapeLayerPaintingRestriction::Type> PaintingRestriction;
 
 #if WITH_EDITOR
diff --git a/Engine/Source/Editor/LandscapeEditor/Private/LandscapeEdMode.cpp b/Engine/Source/Editor/LandscapeEditor/Private/LandscapeEdMode.cpp
--- a/Engine/Source/Editor/LandscapeEditor/Private/LandscapeEdMode.cpp
+++ b/Engine/Source/Editor/LandscapeEditor/Private/LandscapeEdMode.cpp
@@ -205,7 +205,7 @@ FEdModeLandscape::FEdModeLandscape()
 	CurrentToolTarget.TargetType = ELandscapeToolTargetType::Heightmap;
 	CurrentToolTarget.LayerInfo = nullptr;
 
-	UISettings = ConstructObject<ULandscapeEditorObject>(ULandscapeEditorObject::StaticClass());
+	UISettings = ConstructObject<ULandscapeEditorObject>(ULandscapeEditorObject::StaticClass(), GetTransientPackage(), NAME_None, RF_Transactional);
 	UISettings->SetParent(this);
 }
 
@@ -1324,6 +1324,7 @@ bool FEdModeLandscape::InputDelta(FEditorViewportClient* InViewportClient, FView
 			FVector DeltaScale = InScale;
 			DeltaScale.X = DeltaScale.Y = (FMath::Abs(InScale.X) > FMath::Abs(InScale.Y)) ? InScale.X : InScale.Y;
 
+			UISettings->Modify();
 			UISettings->NewLandscape_Location += InDrag;
 			UISettings->NewLandscape_Rotation += InRot;
 			UISettings->NewLandscape_Scale += DeltaScale;
@@ -1338,6 +1339,7 @@ bool FEdModeLandscape::InputDelta(FEditorViewportClient* InViewportClient, FView
 			FTransform Transform(UISettings->NewLandscape_Rotation, UISettings->NewLandscape_Location, UISettings->NewLandscape_Scale * UISettings->NewLandscape_QuadsPerSection * UISettings->NewLandscape_SectionsPerComponent);
 			HitLocation = Transform.InverseTransformPosition(HitLocation);
 
+			UISettings->Modify();
 			switch (DraggingEdge)
 			{
 			case ELandscapeEdge::X_Negative:
diff --git a/Engine/Source/Editor/LandscapeEditor/Private/LandscapeEditorDetailCustomization_NewLandscape.cpp b/Engine/Source/Editor/LandscapeEditor/Private/LandscapeEditorDetailCustomization_NewLandscape.cpp
--- a/Engine/Source/Editor/LandscapeEditor/Private/LandscapeEditorDetailCustomization_NewLandscape.cpp
+++ b/Engine/Source/Editor/LandscapeEditor/Private/LandscapeEditorDetailCustomization_NewLandscape.cpp
@@ -592,6 +592,11 @@ void FLandscapeEditorDetailCustomization_NewLandscape::OnChangeLandscapeResoluti
 	FEdModeLandscape* LandscapeEdMode = GetEditorMode();
 	if (LandscapeEdMode != NULL)
 	{
+		if (!GEditor->IsTransactionActive())
+		{
+			GEditor->BeginTransaction(LOCTEXT("ChangeResolutionX_Transaction", "Change Landscape Resolution X"));
+		}
+		LandscapeEdMode->UISettings->Modify();
 		LandscapeEdMode->UISettings->NewLandscape_ComponentCount.X = (NewValue / (LandscapeEdMode->UISettings->NewLandscape_SectionsPerComponent * LandscapeEdMode->UISettings->NewLandscape_QuadsPerSection));
 		LandscapeEdMode->UISettings->NewLandscape_ClampSize();
 	}
@@ -602,8 +607,14 @@ void FLandscapeEditorDetailCustomization_NewLandscape::OnCommitLandscapeResoluti
 	FEdModeLandscape* LandscapeEdMode = GetEditorMode();
 	if (LandscapeEdMode != NULL)
 	{
+		if (!GEditor->IsTransactionActive())
+		{
+			GEditor->BeginTransaction(LOCTEXT("ChangeResolutionX_Transaction", "Change Landscape Resolution X"));
+		}
+		LandscapeEdMode->UISettings->Modify();
 		LandscapeEdMode->UISettings->NewLandscape_ComponentCount.X = (NewValue / (LandscapeEdMode->UISettings->NewLandscape_SectionsPerComponent * LandscapeEdMode->UISettings->NewLandscape_QuadsPerSection));
 		LandscapeEdMode->UISettings->NewLandscape_ClampSize();
+		GEditor->EndTransaction();
 	}
 }
 
@@ -623,6 +634,11 @@ void FLandscapeEditorDetailCustomization_NewLandscape::OnChangeLandscapeResoluti
 	FEdModeLandscape* LandscapeEdMode = GetEditorMode();
 	if (LandscapeEdMode != NULL)
 	{
+		if (!GEditor->IsTransactionActive())
+		{
+			GEditor->BeginTransaction(LOCTEXT("ChangeResolutionY_Transaction", "Change Landscape Resolution Y"));
+		}
+		LandscapeEdMode->UISettings->Modify();
 		LandscapeEdMode->UISettings->NewLandscape_ComponentCount.Y = (NewValue / (LandscapeEdMode->UISettings->NewLandscape_SectionsPerComponent * LandscapeEdMode->UISettings->NewLandscape_QuadsPerSection));
 		LandscapeEdMode->UISettings->NewLandscape_ClampSize();
 	}
@@ -633,8 +649,14 @@ void FLandscapeEditorDetailCustomization_NewLandscape::OnCommitLandscapeResoluti
 	FEdModeLandscape* LandscapeEdMode = GetEditorMode();
 	if (LandscapeEdMode != NULL)
 	{
+		if (!GEditor->IsTransactionActive())
+		{
+			GEditor->BeginTransaction(LOCTEXT("ChangeResolutionY_Transaction", "Change Landscape Resolution Y"));
+		}
+		LandscapeEdMode->UISettings->Modify();
 		LandscapeEdMode->UISettings->NewLandscape_ComponentCount.Y = (NewValue / (LandscapeEdMode->UISettings->NewLandscape_SectionsPerComponent * LandscapeEdMode->UISettings->NewLandscape_QuadsPerSection));
 		LandscapeEdMode->UISettings->NewLandscape_ClampSize();
+		GEditor->EndTransaction();
 	}
 }
 
@@ -1568,9 +1590,8 @@ EVisibility FLandscapeEditorStructCustomization_FLandscapeImportLayer::GetErrorV
 {
 	auto LayerError = ELandscapeImportLayerError::None;
 	FPropertyAccess::Result Result = PropertyHandle_LayerError->GetValue((uint8&)LayerError);
-	check(Result == FPropertyAccess::Success);
-
-	if (Result == FPropertyAccess::MultipleValues)
+	if (Result == FPropertyAccess::Fail ||
+		Result == FPropertyAccess::MultipleValues)
 	{
 		return EVisibility::Visible;
 	}
@@ -1596,9 +1617,11 @@ FText FLandscapeEditorStructCustomization_FLandscapeImportLayer::GetErrorText(TS
 {
 	auto LayerError = ELandscapeImportLayerError::None;
 	FPropertyAccess::Result Result = PropertyHandle_LayerError->GetValue((uint8&)LayerError);
-	check(Result == FPropertyAccess::Success);
-
-	if (Result == FPropertyAccess::MultipleValues)
+	if (Result == FPropertyAccess::Fail)
+	{
+		return LOCTEXT("Import_LayerUnknownError", "Unknown Error");
+	}
+	else if (Result == FPropertyAccess::MultipleValues)
 	{
 		return NSLOCTEXT("PropertyEditor", "MultipleValues", "Multiple Values");
 	}
@@ -1618,7 +1641,8 @@ FText FLandscapeEditorStructCustomization_FLandscapeImportLayer::GetErrorText(TS
 	case ELandscapeImportLayerError::ColorPng:
 		return LOCTEXT("Import_LayerColorPng", "The Layer file appears to be a color png, grayscale is expected. The import *can* continue, but the result may not be what you expect...");
 	default:
-		check(0);
+		checkSlow(0);
+		return LOCTEXT("Import_LayerUnknownError", "Unknown Error");
 	}
 
 	return FText::GetEmpty();
diff --git a/Engine/Source/Editor/LandscapeEditor/Private/LandscapeEditorObject.cpp b/Engine/Source/Editor/LandscapeEditor/Private/LandscapeEditorObject.cpp
--- a/Engine/Source/Editor/LandscapeEditor/Private/LandscapeEditorObject.cpp
+++ b/Engine/Source/Editor/LandscapeEditor/Private/LandscapeEditorObject.cpp
@@ -118,19 +118,22 @@ void ULandscapeEditorObject::PostEditChangeProperty(FPropertyChangedEvent& Prope
 		SetAlphaTexture(AlphaTexture, AlphaTextureChannel);
 	}
 
-	if (PropertyChangedEvent.MemberProperty->GetFName() == GET_MEMBER_NAME_CHECKED(ULandscapeEditorObject, GizmoHeightmapFilenameString))
+	if (PropertyChangedEvent.MemberProperty == NULL ||
+		PropertyChangedEvent.MemberProperty->GetFName() == GET_MEMBER_NAME_CHECKED(ULandscapeEditorObject, GizmoHeightmapFilenameString))
 	{
 		GuessGizmoImportSize();
 	}
 
-	if (PropertyChangedEvent.MemberProperty->GetFName() == GET_MEMBER_NAME_CHECKED(ULandscapeEditorObject, NewLandscape_QuadsPerSection) ||
+	if (PropertyChangedEvent.MemberProperty == NULL ||
+		PropertyChangedEvent.MemberProperty->GetFName() == GET_MEMBER_NAME_CHECKED(ULandscapeEditorObject, NewLandscape_QuadsPerSection) ||
 		PropertyChangedEvent.MemberProperty->GetFName() == GET_MEMBER_NAME_CHECKED(ULandscapeEditorObject, NewLandscape_SectionsPerComponent) ||
 		PropertyChangedEvent.MemberProperty->GetFName() == GET_MEMBER_NAME_CHECKED(ULandscapeEditorObject, NewLandscape_ComponentCount))
 	{
 		NewLandscape_ClampSize();
 	}
 
-	if (PropertyChangedEvent.MemberProperty->GetFName() == GET_MEMBER_NAME_CHECKED(ULandscapeEditorObject, ResizeLandscape_QuadsPerSection) ||
+	if (PropertyChangedEvent.MemberProperty == NULL ||
+		PropertyChangedEvent.MemberProperty->GetFName() == GET_MEMBER_NAME_CHECKED(ULandscapeEditorObject, ResizeLandscape_QuadsPerSection) ||
 		PropertyChangedEvent.MemberProperty->GetFName() == GET_MEMBER_NAME_CHECKED(ULandscapeEditorObject, ResizeLandscape_SectionsPerComponent) ||
 		PropertyChangedEvent.MemberProperty->GetFName() == GET_MEMBER_NAME_CHECKED(ULandscapeEditorObject, ResizeLandscape_ConvertMode))
 	{
diff --git a/Engine/Source/Editor/PropertyEditor/Private/PropertyHandleImpl.cpp b/Engine/Source/Editor/PropertyEditor/Private/PropertyHandleImpl.cpp
--- a/Engine/Source/Editor/PropertyEditor/Private/PropertyHandleImpl.cpp
+++ b/Engine/Source/Editor/PropertyEditor/Private/PropertyHandleImpl.cpp
@@ -386,16 +386,18 @@ FPropertyAccess::Result FPropertyValueImpl::ImportText( const TArray<FObjectBase
 				bNotifiedPreChange = true;
 				NotifiedObj = Cur.Object;
 
-				// Begin a transaction only if we need to call PreChange
-				if( GEditor && bTransactable )
+				if (!bInteractiveChangeInProgress)
 				{
-					GEditor->BeginTransaction( TEXT("PropertyEditor"), FText::Format( NSLOCTEXT("PropertyEditor", "EditPropertyTransaction", "Edit {0}"), FText::FromString( InPropertyNode->GetDisplayName() ) ), NodeProperty );
+					// Begin a transaction only if we need to call PreChange
+					if (GEditor && bTransactable)
+					{
+						GEditor->BeginTransaction(TEXT("PropertyEditor"), FText::Format(NSLOCTEXT("PropertyEditor", "EditPropertyTransaction", "Edit {0}"), FText::FromString(InPropertyNode->GetDisplayName())), NodeProperty);
+					}
 				}
-				
-				InPropertyNode->NotifyPreChange( NodeProperty, NotifyHook );
 
-				bInteractiveChangeInProgress = (Flags & EPropertyValueSetFlags::InteractiveChange) != 0;
+				InPropertyNode->NotifyPreChange(NodeProperty, NotifyHook);
 
+				bInteractiveChangeInProgress = (Flags & EPropertyValueSetFlags::InteractiveChange) != 0;
 			}
 
 			// Set the new value.
@@ -440,12 +442,12 @@ FPropertyAccess::Result FPropertyValueImpl::ImportText( const TArray<FObjectBase
 			if (bFinished)
 			{
 				bInteractiveChangeInProgress = false;
-			}
 
-			if( bTransactable )
-			{
-				// End the transaction if we called PreChange
-				GEditor->EndTransaction();
+				if (bTransactable)
+				{
+					// End the transaction if we called PreChange
+					GEditor->EndTransaction();
+				}
 			}
 		}
 
