Flagged the landscape promotion test to be ATF_NonNullRHI.  So an end user will have to run it with rendering otherwise it will cause an ensure.

diff --git a/Engine/Source/Editor/LandscapeEditor/Private/Tests/LandscapeAutomationTests.cpp b/Engine/Source/Editor/LandscapeEditor/Private/Tests/LandscapeAutomationTests.cpp
--- a/Engine/Source/Editor/LandscapeEditor/Private/Tests/LandscapeAutomationTests.cpp
+++ b/Engine/Source/Editor/LandscapeEditor/Private/Tests/LandscapeAutomationTests.cpp
@@ -117,7 +117,7 @@ bool FEndModifyLandscapeCommand::Update()
 /**
 * Landscape creation / edit test
 */
-IMPLEMENT_SIMPLE_AUTOMATION_TEST(FLandscapeEditorTest, "Promotion.Editor Promotion Pass.Step 1 Main Editor Test.New Landscape", EAutomationTestFlags::ATF_Editor);
+IMPLEMENT_SIMPLE_AUTOMATION_TEST(FLandscapeEditorTest, "Promotion.Editor Promotion Pass.Step 1 Main Editor Test.New Landscape", EAutomationTestFlags::ATF_Editor | EAutomationTestFlags::ATF_NonNullRHI );
 bool FLandscapeEditorTest::RunTest(const FString& Parameters)
 {
 	//New level
