Landscape : Fixed undo problem (TTP #344535)

diff --git a/Engine/Source/Editor/LandscapeEditor/Classes/LandscapeEditorObject.h b/Engine/Source/Editor/LandscapeEditor/Classes/LandscapeEditorObject.h
--- a/Engine/Source/Editor/LandscapeEditor/Classes/LandscapeEditorObject.h
+++ b/Engine/Source/Editor/LandscapeEditor/Classes/LandscapeEditorObject.h
@@ -525,6 +525,11 @@ class ULandscapeEditorObject : public UObject
 
 	void UpdateComponentCount()
 	{
+		// ignore invalid cases
+		if (ResizeLandscape_QuadsPerSection == 0 || ResizeLandscape_SectionsPerComponent == 0 ||ResizeLandscape_ComponentCount.X == 0 || ResizeLandscape_ComponentCount.Y == 0) 
+		{
+			return;
+		}
 		const int32 ComponentSizeQuads = ResizeLandscape_QuadsPerSection * ResizeLandscape_SectionsPerComponent;
 		const int32 Original_ComponentSizeQuads = ResizeLandscape_Original_QuadsPerSection * ResizeLandscape_Original_SectionsPerComponent;
 		const FIntPoint OriginalResolution = ResizeLandscape_Original_ComponentCount * Original_ComponentSizeQuads;
