Fix Landscape Layer Painting Restriction showing up for sculpting tools

diff --git a/Engine/Source/Editor/LandscapeEditor/Classes/LandscapeEditorObject.h b/Engine/Source/Editor/LandscapeEditor/Classes/LandscapeEditorObject.h
--- a/Engine/Source/Editor/LandscapeEditor/Classes/LandscapeEditorObject.h
+++ b/Engine/Source/Editor/LandscapeEditor/Classes/LandscapeEditorObject.h
@@ -485,7 +485,7 @@ class ULandscapeEditorObject : public UObject
 	// Target Layer Settings:
 
 	// Limits painting to only the components that already have the selected layer
-	UPROPERTY(Category="Target Layers", EditAnywhere, NonTransactional)
+	UPROPERTY(Category="Target Layers", EditAnywhere, NonTransactional, meta=(ShowForTargetTypes="Weightmap,Visibility"))
 	TEnumAsByte<ELandscapeLayerPaintingRestriction::Type> PaintingRestriction;
 
 #if WITH_EDITOR
diff --git a/Engine/Source/Editor/LandscapeEditor/Private/LandscapeEditorDetailCustomization_TargetLayers.cpp b/Engine/Source/Editor/LandscapeEditor/Private/LandscapeEditorDetailCustomization_TargetLayers.cpp
--- a/Engine/Source/Editor/LandscapeEditor/Private/LandscapeEditorDetailCustomization_TargetLayers.cpp
+++ b/Engine/Source/Editor/LandscapeEditor/Private/LandscapeEditorDetailCustomization_TargetLayers.cpp
@@ -47,7 +47,8 @@ void FLandscapeEditorDetailCustomization_TargetLayers::CustomizeDetails(IDetailL
 
 	IDetailCategoryBuilder& TargetsCategory = DetailBuilder.EditCategory("Target Layers");
 
-	TargetsCategory.AddProperty(PropertyHandle_PaintingRestriction);
+	TargetsCategory.AddProperty(PropertyHandle_PaintingRestriction)
+		.Visibility(TAttribute<EVisibility>::Create(TAttribute<EVisibility>::FGetter::CreateStatic(&FLandscapeEditorDetailCustomization_TargetLayers::GetVisibility_PaintingRestriction)));
 
 	TargetsCategory.AddCustomBuilder(MakeShareable(new FLandscapeEditorCustomNodeBuilder_TargetLayers(DetailBuilder.GetThumbnailPool().ToSharedRef())));
 }
@@ -73,6 +74,26 @@ bool FLandscapeEditorDetailCustomization_TargetLayers::ShouldShowTargetLayers()
 	return false;
 }
 
+bool FLandscapeEditorDetailCustomization_TargetLayers::ShouldShowPaintingRestriction()
+{
+	FEdModeLandscape* LandscapeEdMode = GetEditorMode();
+	if (LandscapeEdMode)
+	{
+		if (LandscapeEdMode->CurrentToolTarget.TargetType == ELandscapeToolTargetType::Weightmap ||
+			LandscapeEdMode->CurrentToolTarget.TargetType == ELandscapeToolTargetType::Visibility)
+		{
+			return true;
+		}
+	}
+
+	return false;
+}
+
+EVisibility FLandscapeEditorDetailCustomization_TargetLayers::GetVisibility_PaintingRestriction()
+{
+	return ShouldShowPaintingRestriction() ? EVisibility::Visible : EVisibility::Collapsed;
+}
+
 //////////////////////////////////////////////////////////////////////////
 
 FEdModeLandscape* FLandscapeEditorCustomNodeBuilder_TargetLayers::GetEditorMode()
diff --git a/Engine/Source/Editor/LandscapeEditor/Private/LandscapeEditorDetailCustomization_TargetLayers.h b/Engine/Source/Editor/LandscapeEditor/Private/LandscapeEditorDetailCustomization_TargetLayers.h
--- a/Engine/Source/Editor/LandscapeEditor/Private/LandscapeEditorDetailCustomization_TargetLayers.h
+++ b/Engine/Source/Editor/LandscapeEditor/Private/LandscapeEditorDetailCustomization_TargetLayers.h
@@ -20,6 +20,8 @@ public:
 
 protected:
 	static bool ShouldShowTargetLayers();
+	static bool ShouldShowPaintingRestriction();
+	static EVisibility GetVisibility_PaintingRestriction();
 };
 
 class FLandscapeEditorCustomNodeBuilder_TargetLayers : public IDetailCustomNodeBuilder
