Fixed metadata on landscape tool buttons.

diff --git a/Engine/Source/Editor/LandscapeEditor/Private/LandscapeEditorDetailWidgets.cpp b/Engine/Source/Editor/LandscapeEditor/Private/LandscapeEditorDetailWidgets.cpp
--- a/Engine/Source/Editor/LandscapeEditor/Private/LandscapeEditorDetailWidgets.cpp
+++ b/Engine/Source/Editor/LandscapeEditor/Private/LandscapeEditorDetailWidgets.cpp
@@ -146,6 +146,8 @@ void SToolSelector::BuildMultiBlockWidget(const ISlateStyle* StyleSet, const FNa
 
 	Label = ToolSelectorButtonBlock->Label;
 	SmallText = ToolSelectorButtonBlock->SmallText;
+	
+	const FString MetaTag = FString::Printf(TEXT("LandscapeToolButton.%s"), Label.IsSet() == true ? *Label.Get().ToString() : TEXT("NoLabel"));
 
 	static FTextBlockStyle LabelStyle = FTextBlockStyle(FEditorStyle::GetWidgetStyle< FTextBlockStyle >(FEditorStyle::Join(StyleName, ".Label"))).SetShadowOffset(FVector2D::UnitVector);
 	static FTextBlockStyle SmallTextStyle = FTextBlockStyle(LabelStyle).SetFont(LabelStyle.Font.FontName, LabelStyle.Font.Size - 1).SetColorAndOpacity(FSlateColor::UseSubduedForeground());
@@ -153,7 +155,7 @@ void SToolSelector::BuildMultiBlockWidget(const ISlateStyle* StyleSet, const FNa
 	// Create the content for our button
 	TSharedRef< SWidget > ButtonContent =
 		SNew( SVerticalBox )
-		.AddMetaData<FTagMetaData>(FTagMetaData(TEXT("LandscapeToolButton")))
+		.AddMetaData<FTagMetaData>(FTagMetaData(FName(*MetaTag)))
 		// Icon image
 		+ SVerticalBox::Slot()
 		.AutoHeight()
