Fix noise generated when using the landscape flatten tool with a strength greater than 1.0

diff --git a/Engine/Source/Editor/LandscapeEditor/Private/LandscapeEdModePaintTools.cpp b/Engine/Source/Editor/LandscapeEditor/Private/LandscapeEdModePaintTools.cpp
--- a/Engine/Source/Editor/LandscapeEditor/Private/LandscapeEdModePaintTools.cpp
+++ b/Engine/Source/Editor/LandscapeEditor/Private/LandscapeEdModePaintTools.cpp
@@ -508,6 +508,8 @@ public:
 			if (It.Value() > 0.f)
 			{
 				int32 HeightDataIndex = (X - X1) + (Y - Y1)*(1 + X2 - X1);
+			
+				float Strength = FMath::Clamp<float>(It.Value() * UISettings->ToolStrength * Pressure, 0.f, 1.f);
 
 				if (!(UISettings->bUseSlopeFlatten && bTargetIsHeightmap))
 				{
@@ -517,18 +519,18 @@ public:
 					case ELandscapeToolNoiseMode::Add:
 						if (Delta < 0)
 						{
-							HeightData[HeightDataIndex] = FMath::Lerp(HeightData[HeightDataIndex], FlattenHeight, It.Value() * UISettings->ToolStrength * Pressure);
+							HeightData[HeightDataIndex] = FMath::Lerp(HeightData[HeightDataIndex], FlattenHeight, Strength);
 						}
 						break;
 					case ELandscapeToolNoiseMode::Sub:
 						if (Delta > 0)
 						{
-							HeightData[HeightDataIndex] = FMath::Lerp(HeightData[HeightDataIndex], FlattenHeight, It.Value() * UISettings->ToolStrength * Pressure);
+							HeightData[HeightDataIndex] = FMath::Lerp(HeightData[HeightDataIndex], FlattenHeight, Strength);
 						}
 						break;
 					default:
 					case ELandscapeToolNoiseMode::Both:
-						HeightData[HeightDataIndex] = FMath::Lerp(HeightData[HeightDataIndex], FlattenHeight, It.Value() * UISettings->ToolStrength * Pressure);
+						HeightData[HeightDataIndex] = FMath::Lerp(HeightData[HeightDataIndex], FlattenHeight, Strength);
 						break;
 					}
 				}
@@ -537,24 +539,24 @@ public:
 					typename ToolTarget::CacheClass::DataType DestValue = -(FlattenNormal.X * X + FlattenNormal.Y * Y + FlattenPlaneDist) / FlattenNormal.Z;
 					//float PlaneDist = FlattenNormal | FVector(X, Y, HeightData(HeightDataIndex)) + FlattenPlaneDist;
 					float PlaneDist = HeightData[HeightDataIndex] - DestValue;
-					DestValue = HeightData[HeightDataIndex] - PlaneDist * It.Value() * UISettings->ToolStrength * Pressure;
+					DestValue = HeightData[HeightDataIndex] - PlaneDist * Strength;
 					switch (UISettings->FlattenMode)
 					{
 					case ELandscapeToolNoiseMode::Add:
 						if (PlaneDist < 0)
 						{
-							HeightData[HeightDataIndex] = FMath::Lerp(HeightData[HeightDataIndex], DestValue, It.Value() * UISettings->ToolStrength * Pressure);
+							HeightData[HeightDataIndex] = FMath::Lerp(HeightData[HeightDataIndex], DestValue, Strength);
 						}
 						break;
 					case ELandscapeToolNoiseMode::Sub:
 						if (PlaneDist > 0)
 						{
-							HeightData[HeightDataIndex] = FMath::Lerp(HeightData[HeightDataIndex], DestValue, It.Value() * UISettings->ToolStrength * Pressure);
+							HeightData[HeightDataIndex] = FMath::Lerp(HeightData[HeightDataIndex], DestValue, Strength);
 						}
 						break;
 					default:
 					case ELandscapeToolNoiseMode::Both:
-						HeightData[HeightDataIndex] = FMath::Lerp(HeightData[HeightDataIndex], DestValue, It.Value() * UISettings->ToolStrength * Pressure);
+						HeightData[HeightDataIndex] = FMath::Lerp(HeightData[HeightDataIndex], DestValue, Strength);
 						break;
 					}
 				}
