Improve landscape sculpting tool, it's now continuous and framerate-independent
Refactored landscape tool strokes somewhat, removing unnecessary virtual functions (most of the stroke classes were single-level hierarchies, making virtuals pointless), empty functions and unnecessary dynamic allocation (using refactored TOptional)

diff --git a/Engine/Source/Editor/LandscapeEditor/Private/LandscapeEdModeComponentTools.cpp b/Engine/Source/Editor/LandscapeEditor/Private/LandscapeEdModeComponentTools.cpp
--- a/Engine/Source/Editor/LandscapeEditor/Private/LandscapeEdModeComponentTools.cpp
+++ b/Engine/Source/Editor/LandscapeEditor/Private/LandscapeEdModeComponentTools.cpp
@@ -18,23 +18,20 @@
 
 #define LOCTEXT_NAMESPACE "Landscape"
 
-class FLandscapeToolStrokeSelect : FLandscapeStrokeBase<const FLandscapeToolTarget>
+class FLandscapeToolStrokeSelect : public FLandscapeToolStrokeBase
 {
 	bool bInitializedComponentInvert;
 	bool bComponentInvert;
 
 public:
 	FLandscapeToolStrokeSelect(FEdModeLandscape* InEdMode, const FLandscapeToolTarget& InTarget)
-		: FLandscapeStrokeBase<const FLandscapeToolTarget>(InTarget)
-		, bInitializedComponentInvert(false)
+		: bInitializedComponentInvert(false)
 		, LandscapeInfo(InTarget.LandscapeInfo.Get())
 		, Cache(InTarget)
 	{
 	}
 
-	virtual ~FLandscapeToolStrokeSelect() {}
-
-	virtual void Apply(FEditorViewportClient* ViewportClient, FLandscapeBrush* Brush, const ULandscapeEditorObject* UISettings, const TArray<FLandscapeToolMousePosition>& MousePositions)
+	void Apply(FEditorViewportClient* ViewportClient, FLandscapeBrush* Brush, const ULandscapeEditorObject* UISettings, const TArray<FLandscapeToolMousePosition>& MousePositions)
 	{
 		if (LandscapeInfo)
 		{
@@ -169,7 +166,8 @@ class FLandscapeToolSelect : public FLandscapeToolBase<TStrokeClass>
 public:
 	FLandscapeToolSelect(FEdModeLandscape* InEdMode)
 		: FLandscapeToolBase<TStrokeClass>(InEdMode)
-	{}
+	{
+	}
 
 	virtual const TCHAR* GetToolName() override { return TEXT("Select"); }
 	virtual FText GetDisplayName() override { return NSLOCTEXT("UnrealEd", "LandscapeMode_Selection", "Component Selection"); };
@@ -185,7 +183,8 @@ class FLandscapeToolMask : public FLandscapeToolSelect<TStrokeClass>
 public:
 	FLandscapeToolMask(FEdModeLandscape* InEdMode)
 		: FLandscapeToolSelect<TStrokeClass>(InEdMode)
-	{}
+	{
+	}
 
 	virtual const TCHAR* GetToolName() override { return TEXT("Mask"); }
 	virtual FText GetDisplayName() override { return NSLOCTEXT("UnrealEd", "LandscapeMode_Mask", "Region Selection"); };
@@ -193,17 +192,18 @@ public:
 	virtual bool SupportsMask() override { return true; }
 };
 
-class FLandscapeToolStrokeVisibility
+class FLandscapeToolStrokeVisibility : public FLandscapeToolStrokeBase
 {
 public:
+	enum { UseContinuousApply = false };
+
 	FLandscapeToolStrokeVisibility(FEdModeLandscape* InEdMode, const FLandscapeToolTarget& InTarget)
 		: LandscapeInfo(InTarget.LandscapeInfo.Get())
 		, Cache(InTarget)
-	{}
-
-	virtual ~FLandscapeToolStrokeVisibility() {}
+	{
+	}
 
-	virtual void Apply(FEditorViewportClient* ViewportClient, FLandscapeBrush* Brush, const ULandscapeEditorObject* UISettings, const TArray<FLandscapeToolMousePosition>& MousePositions)
+	void Apply(FEditorViewportClient* ViewportClient, FLandscapeBrush* Brush, const ULandscapeEditorObject* UISettings, const TArray<FLandscapeToolMousePosition>& MousePositions)
 	{
 		if (LandscapeInfo)
 		{
@@ -270,7 +270,8 @@ class FLandscapeToolVisibility : public FLandscapeToolBase<FLandscapeToolStrokeV
 public:
 	FLandscapeToolVisibility(FEdModeLandscape* InEdMode)
 		: FLandscapeToolBase<FLandscapeToolStrokeVisibility>(InEdMode)
-	{}
+	{
+	}
 
 	virtual const TCHAR* GetToolName() override { return TEXT("Visibility"); }
 	virtual FText GetDisplayName() override { return NSLOCTEXT("UnrealEd", "LandscapeMode_Visibility", "Visibility"); };
@@ -287,13 +288,14 @@ public:
 class FLandscapeToolStrokeMoveToLevel
 {
 public:
+	enum { UseContinuousApply = false };
+
 	FLandscapeToolStrokeMoveToLevel(FEdModeLandscape* InEdMode, const FLandscapeToolTarget& InTarget)
 		: LandscapeInfo(InTarget.LandscapeInfo.Get())
-	{}
-
-	virtual ~FLandscapeToolStrokeMoveToLevel() {}
+	{
+	}
 
-	virtual void Apply(FEditorViewportClient* ViewportClient, FLandscapeBrush* Brush, const ULandscapeEditorObject* UISettings, const TArray<FLandscapeToolMousePosition>& MousePositions)
+	void Apply(FEditorViewportClient* ViewportClient, FLandscapeBrush* Brush, const ULandscapeEditorObject* UISettings, const TArray<FLandscapeToolMousePosition>& MousePositions)
 	{
 		ALandscape* Landscape = LandscapeInfo ? LandscapeInfo->LandscapeActor.Get() : NULL;
 
@@ -736,7 +738,8 @@ class FLandscapeToolMoveToLevel : public FLandscapeToolBase<FLandscapeToolStroke
 public:
 	FLandscapeToolMoveToLevel(FEdModeLandscape* InEdMode)
 		: FLandscapeToolBase<FLandscapeToolStrokeMoveToLevel>(InEdMode)
-	{}
+	{
+	}
 
 	virtual const TCHAR* GetToolName() override { return TEXT("MoveToLevel"); }
 	virtual FText GetDisplayName() override { return NSLOCTEXT("UnrealEd", "LandscapeMode_MoveToLevel", "Move to Streaming Level"); };
@@ -746,7 +749,7 @@ public:
 };
 
 
-class FLandscapeToolStrokeAddComponent
+class FLandscapeToolStrokeAddComponent : public FLandscapeToolStrokeBase
 {
 public:
 	FLandscapeToolStrokeAddComponent(FEdModeLandscape* InEdMode, const FLandscapeToolTarget& InTarget)
@@ -754,7 +757,8 @@ public:
 		,	LandscapeInfo(InTarget.LandscapeInfo.Get())
 		,	HeightCache(InTarget)
 		,	XYOffsetCache(InTarget)
-	{}
+	{
+	}
 
 	virtual ~FLandscapeToolStrokeAddComponent()
 	{
@@ -910,7 +914,8 @@ class FLandscapeToolAddComponent : public FLandscapeToolBase<FLandscapeToolStrok
 public:
 	FLandscapeToolAddComponent(FEdModeLandscape* InEdMode)
 		: FLandscapeToolBase<FLandscapeToolStrokeAddComponent>(InEdMode)
-	{}
+	{
+	}
 
 	virtual const TCHAR* GetToolName() override { return TEXT("AddComponent"); }
 	virtual FText GetDisplayName() override { return NSLOCTEXT("UnrealEd", "LandscapeMode_AddComponent", "Add New Landscape Component"); };
@@ -926,16 +931,15 @@ public:
 	}
 };
 
-class FLandscapeToolStrokeDeleteComponent
+class FLandscapeToolStrokeDeleteComponent : public FLandscapeToolStrokeBase
 {
 public:
 	FLandscapeToolStrokeDeleteComponent(FEdModeLandscape* InEdMode, const FLandscapeToolTarget& InTarget)
 		: LandscapeInfo(InTarget.LandscapeInfo.Get())
-	{}
-
-	virtual ~FLandscapeToolStrokeDeleteComponent() {}
+	{
+	}
 
-	virtual void Apply(FEditorViewportClient* ViewportClient, FLandscapeBrush* Brush, const ULandscapeEditorObject* UISettings, const TArray<FLandscapeToolMousePosition>& MousePositions)
+	void Apply(FEditorViewportClient* ViewportClient, FLandscapeBrush* Brush, const ULandscapeEditorObject* UISettings, const TArray<FLandscapeToolMousePosition>& MousePositions)
 	{
 		if (LandscapeInfo)
 		{
@@ -1121,7 +1125,8 @@ class FLandscapeToolDeleteComponent : public FLandscapeToolBase<FLandscapeToolSt
 public:
 	FLandscapeToolDeleteComponent(FEdModeLandscape* InEdMode)
 		: FLandscapeToolBase<FLandscapeToolStrokeDeleteComponent>(InEdMode)
-	{}
+	{
+	}
 
 	virtual const TCHAR* GetToolName() override { return TEXT("DeleteComponent"); }
 	virtual FText GetDisplayName() override { return NSLOCTEXT("UnrealEd", "LandscapeMode_DeleteComponent", "Delete Landscape Components"); };
@@ -1131,7 +1136,7 @@ public:
 };
 
 template<class ToolTarget>
-class FLandscapeToolStrokeCopy
+class FLandscapeToolStrokeCopy : public FLandscapeToolStrokeBase
 {
 public:
 	FLandscapeToolStrokeCopy(FEdModeLandscape* InEdMode, const FLandscapeToolTarget& InTarget)
@@ -1140,9 +1145,8 @@ public:
 		, Cache(InTarget)
 		, HeightCache(InTarget)
 		, WeightCache(InTarget)
-	{}
-
-	virtual ~FLandscapeToolStrokeCopy() {}
+	{
+	}
 
 	struct FGizmoPreData
 	{
@@ -1150,7 +1154,7 @@ public:
 		float Data;
 	};
 
-	virtual void Apply(FEditorViewportClient* ViewportClient, FLandscapeBrush* Brush, const ULandscapeEditorObject* UISettings, const TArray<FLandscapeToolMousePosition>& MousePositions)
+	void Apply(FEditorViewportClient* ViewportClient, FLandscapeBrush* Brush, const ULandscapeEditorObject* UISettings, const TArray<FLandscapeToolMousePosition>& MousePositions)
 	{
 		//ULandscapeInfo* LandscapeInfo = EdMode->CurrentToolTarget.LandscapeInfo;
 		ALandscapeGizmoActiveActor* Gizmo = EdMode->CurrentGizmoActor.Get();
@@ -1443,7 +1447,8 @@ class FLandscapeToolCopy : public FLandscapeToolBase<FLandscapeToolStrokeCopy<To
 public:
 	FLandscapeToolCopy(FEdModeLandscape* InEdMode)
 		: FLandscapeToolBase<FLandscapeToolStrokeCopy<ToolTarget> >(InEdMode)
-	{}
+	{
+	}
 
 	virtual const TCHAR* GetToolName() override { return TEXT("Copy"); }
 	virtual FText GetDisplayName() override { return NSLOCTEXT("UnrealEd", "LandscapeMode_Copy", "Copy"); };
@@ -1463,7 +1468,7 @@ public:
 	{
 		check(this->MousePositions.Num() == 0);
 		this->bToolActive = true;
-		this->ToolStroke = new FLandscapeToolStrokeCopy<ToolTarget>(this->EdMode, InTarget);
+		this->ToolStroke.Emplace(this->EdMode, InTarget);
 
 		this->EdMode->GizmoBrush->Tick(ViewportClient, 0.1f);
 		this->EdMode->GizmoBrush->BeginStroke(InHitLocation.X, InHitLocation.Y, this);
@@ -1482,8 +1487,7 @@ public:
 			this->MousePositions.Empty();
 		}
 
-		delete this->ToolStroke;
-		this->ToolStroke = NULL;
+		this->ToolStroke.Reset();
 		this->bToolActive = false;
 		this->EdMode->GizmoBrush->EndStroke();
 	}
@@ -1496,7 +1500,7 @@ public:
 };
 
 template<class ToolTarget>
-class FLandscapeToolStrokePaste
+class FLandscapeToolStrokePaste : public FLandscapeToolStrokeBase
 {
 public:
 	FLandscapeToolStrokePaste(FEdModeLandscape* InEdMode, const FLandscapeToolTarget& InTarget)
@@ -1505,11 +1509,10 @@ public:
 		, Cache(InTarget)
 		, HeightCache(InTarget)
 		, WeightCache(InTarget)
-	{}
-
-	virtual ~FLandscapeToolStrokePaste() {}
+	{
+	}
 
-	virtual void Apply(FEditorViewportClient* ViewportClient, FLandscapeBrush* Brush, const ULandscapeEditorObject* UISettings, const TArray<FLandscapeToolMousePosition>& MousePositions)
+	void Apply(FEditorViewportClient* ViewportClient, FLandscapeBrush* Brush, const ULandscapeEditorObject* UISettings, const TArray<FLandscapeToolMousePosition>& MousePositions)
 	{
 		//ULandscapeInfo* LandscapeInfo = EdMode->CurrentToolTarget.LandscapeInfo;
 		ALandscapeGizmoActiveActor* Gizmo = EdMode->CurrentGizmoActor.Get();
@@ -1743,7 +1746,8 @@ public:
 	FLandscapeToolPaste(FEdModeLandscape* InEdMode)
 		: FLandscapeToolBase<FLandscapeToolStrokePaste<ToolTarget> >(InEdMode)
 		, bUseGizmoRegion(false)
-	{}
+	{
+	}
 
 	virtual const TCHAR* GetToolName() override { return TEXT("Paste"); }
 	virtual FText GetDisplayName() override { return NSLOCTEXT("UnrealEd", "LandscapeMode_Region", "Region Copy/Paste"); };
@@ -1767,7 +1771,7 @@ public:
 	virtual bool BeginTool(FEditorViewportClient* ViewportClient, const FLandscapeToolTarget& InTarget, const FVector& InHitLocation) override
 	{
 		this->bToolActive = true;
-		this->ToolStroke = new FLandscapeToolStrokePaste<ToolTarget>(this->EdMode, InTarget);
+		this->ToolStroke.Emplace(this->EdMode, InTarget);
 
 		this->EdMode->GizmoBrush->Tick(ViewportClient, 0.1f);
 		if (bUseGizmoRegion)
@@ -1793,8 +1797,7 @@ public:
 			this->MousePositions.Empty();
 		}
 
-		delete this->ToolStroke;
-		this->ToolStroke = NULL;
+		this->ToolStroke.Reset();
 		this->bToolActive = false;
 		if (bUseGizmoRegion)
 		{
@@ -1837,8 +1840,8 @@ public:
 	// Copy tool doesn't use any view information, so just do it as one function
 	void Copy()
 	{
-		CopyTool->BeginTool(nullptr, this->EdMode->CurrentToolTarget, FVector::ZeroVector);
-		CopyTool->EndTool(nullptr);
+		CopyTool.BeginTool(nullptr, this->EdMode->CurrentToolTarget, FVector::ZeroVector);
+		CopyTool.EndTool(nullptr);
 	}
 
 	void Paste()
@@ -1850,7 +1853,7 @@ public:
 	}
 
 protected:
-	TUniqueObj<FLandscapeToolCopy<ToolTarget>> CopyTool;
+	FLandscapeToolCopy<ToolTarget> CopyTool;
 };
 
 void FEdModeLandscape::CopyDataToGizmo()
diff --git a/Engine/Source/Editor/LandscapeEditor/Private/LandscapeEdModeErosionTools.cpp b/Engine/Source/Editor/LandscapeEditor/Private/LandscapeEdModeErosionTools.cpp
--- a/Engine/Source/Editor/LandscapeEditor/Private/LandscapeEdModeErosionTools.cpp
+++ b/Engine/Source/Editor/LandscapeEditor/Private/LandscapeEdModeErosionTools.cpp
@@ -17,7 +17,7 @@
 //
 // FLandscapeToolErosionBase
 //
-class FLandscapeToolStrokeErosionBase
+class FLandscapeToolStrokeErosionBase : public FLandscapeToolStrokeBase
 {
 public:
 	FLandscapeToolStrokeErosionBase(FEdModeLandscape* InEdMode, const FLandscapeToolTarget& InTarget)
@@ -25,9 +25,9 @@ public:
 		, HeightCache(InTarget)
 		, WeightCache(InTarget)
 		, bWeightApplied(InTarget.TargetType != ELandscapeToolTargetType::Heightmap)
-	{}
+	{
+	}
 
-	virtual void Apply(FEditorViewportClient* ViewportClient, FLandscapeBrush* Brush, const ULandscapeEditorObject* UISettings, const TArray<FLandscapeToolMousePosition>& MousePositions) = 0;
 protected:
 	ULandscapeInfo* LandscapeInfo;
 	FLandscapeHeightCache HeightCache;
@@ -41,7 +41,8 @@ class FLandscapeToolErosionBase : public FLandscapeToolBase<TStrokeClass>
 public:
 	FLandscapeToolErosionBase(FEdModeLandscape* InEdMode)
 		: FLandscapeToolBase<TStrokeClass>(InEdMode)
-	{}
+	{
+	}
 
 	virtual ELandscapeToolTargetTypeMask::Type GetSupportedTargetTypes() override
 	{
@@ -58,11 +59,10 @@ class FLandscapeToolStrokeErosion : public FLandscapeToolStrokeErosionBase
 public:
 	FLandscapeToolStrokeErosion(FEdModeLandscape* InEdMode, const FLandscapeToolTarget& InTarget)
 		: FLandscapeToolStrokeErosionBase(InEdMode, InTarget)
-	{}
-
-	virtual ~FLandscapeToolStrokeErosion() {}
+	{
+	}
 
-	virtual void Apply(FEditorViewportClient* ViewportClient, FLandscapeBrush* Brush, const ULandscapeEditorObject* UISettings, const TArray<FLandscapeToolMousePosition>& MousePositions) override
+	void Apply(FEditorViewportClient* ViewportClient, FLandscapeBrush* Brush, const ULandscapeEditorObject* UISettings, const TArray<FLandscapeToolMousePosition>& MousePositions)
 	{
 		if (!LandscapeInfo)
 		{
@@ -260,7 +260,8 @@ class FLandscapeToolErosion : public FLandscapeToolErosionBase<FLandscapeToolStr
 public:
 	FLandscapeToolErosion(FEdModeLandscape* InEdMode)
 		: FLandscapeToolErosionBase(InEdMode)
-	{}
+	{
+	}
 
 	virtual const TCHAR* GetToolName() override { return TEXT("Erosion"); }
 	virtual FText GetDisplayName() override { return NSLOCTEXT("UnrealEd", "LandscapeMode_Erosion", "Erosion"); };
@@ -276,11 +277,10 @@ class FLandscapeToolStrokeHydraErosion : public FLandscapeToolStrokeErosionBase
 public:
 	FLandscapeToolStrokeHydraErosion(FEdModeLandscape* InEdMode, const FLandscapeToolTarget& InTarget)
 		: FLandscapeToolStrokeErosionBase(InEdMode, InTarget)
-	{}
-
-	virtual ~FLandscapeToolStrokeHydraErosion() {}
+	{
+	}
 
-	virtual void Apply(FEditorViewportClient* ViewportClient, FLandscapeBrush* Brush, const ULandscapeEditorObject* UISettings, const TArray<FLandscapeToolMousePosition>& MousePositions) override
+	void Apply(FEditorViewportClient* ViewportClient, FLandscapeBrush* Brush, const ULandscapeEditorObject* UISettings, const TArray<FLandscapeToolMousePosition>& MousePositions)
 	{
 		if (!LandscapeInfo)
 		{
@@ -462,7 +462,8 @@ class FLandscapeToolHydraErosion : public FLandscapeToolErosionBase<FLandscapeTo
 public:
 	FLandscapeToolHydraErosion(FEdModeLandscape* InEdMode)
 		: FLandscapeToolErosionBase(InEdMode)
-	{}
+	{
+	}
 
 	virtual const TCHAR* GetToolName() override { return TEXT("HydraErosion"); } // formerly HydraulicErosion
 	virtual FText GetDisplayName() override { return NSLOCTEXT("UnrealEd", "LandscapeMode_HydraErosion", "Hydraulic Erosion"); };
diff --git a/Engine/Source/Editor/LandscapeEditor/Private/LandscapeEdModePaintTools.cpp b/Engine/Source/Editor/LandscapeEditor/Private/LandscapeEdModePaintTools.cpp
--- a/Engine/Source/Editor/LandscapeEditor/Private/LandscapeEdModePaintTools.cpp
+++ b/Engine/Source/Editor/LandscapeEditor/Private/LandscapeEdModePaintTools.cpp
@@ -43,7 +43,8 @@ class FLandscapeToolPaintBase : public FLandscapeToolBase<TStrokeClass>
 public:
 	FLandscapeToolPaintBase(FEdModeLandscape* InEdMode)
 		: FLandscapeToolBase<TStrokeClass>(InEdMode)
-	{}
+	{
+	}
 
 	virtual ELandscapeToolTargetTypeMask::Type GetSupportedTargetTypes() override
 	{
@@ -53,7 +54,7 @@ public:
 
 
 template<class ToolTarget>
-class FLandscapeToolStrokePaintBase
+class FLandscapeToolStrokePaintBase : public FLandscapeToolStrokeBase
 {
 public:
 	FLandscapeToolStrokePaintBase(FEdModeLandscape* InEdMode, const FLandscapeToolTarget& InTarget)
@@ -62,7 +63,6 @@ public:
 	{
 	}
 
-	virtual void Apply(FEditorViewportClient* ViewportClient, FLandscapeBrush* Brush, const ULandscapeEditorObject* UISettings, const TArray<FLandscapeToolMousePosition>& MousePositions) = 0;
 protected:
 	typename ToolTarget::CacheClass Cache;
 	ULandscapeInfo* LandscapeInfo;
@@ -81,13 +81,14 @@ class FLandscapeToolStrokePaint : public FLandscapeToolStrokePaintBase<ToolTarge
 {
 	TMap<FIntPoint, float> TotalInfluenceMap;	// amount of time and weight the brush has spent on each vertex.
 public:
+	enum { UseContinuousApply = (ToolTarget::TargetType == ELandscapeToolTargetType::Heightmap) };
+
 	FLandscapeToolStrokePaint(FEdModeLandscape* InEdMode, const FLandscapeToolTarget& InTarget)
 		: FLandscapeToolStrokePaintBase<ToolTarget>(InEdMode, InTarget)
-	{}
-
-	virtual ~FLandscapeToolStrokePaint() {}
+	{
+	}
 
-	virtual void Apply(FEditorViewportClient* ViewportClient, FLandscapeBrush* Brush, const ULandscapeEditorObject* UISettings, const TArray<FLandscapeToolMousePosition>& MousePositions) override
+	void Apply(FEditorViewportClient* ViewportClient, FLandscapeBrush* Brush, const ULandscapeEditorObject* UISettings, const TArray<FLandscapeToolMousePosition>& MousePositions)
 	{
 		// Get list of verts to update
 		TMap<FIntPoint, float> BrushInfo;
@@ -119,19 +120,11 @@ public:
 		this->Cache.GetCachedData(X1, Y1, X2, Y2, Data);
 
 		// The source data we use for editing. 
-		// For heightmaps we use a cached snapshot, for weightmaps we use the live data.
 		TArray<typename ToolTarget::CacheClass::DataType>* SourceDataArrayPtr = &Data;
 		TArray<typename ToolTarget::CacheClass::DataType> OriginalData;
 
-		if (ToolTarget::TargetType == ELandscapeToolTargetType::Heightmap)
+		if ((ToolTarget::TargetType == ELandscapeToolTargetType::Weightmap) && !bUseWeightTargetValue)
 		{
-			// Heightmaps use the original data rather than the data edited during the stroke.
-			this->Cache.GetOriginalData(X1, Y1, X2, Y2, OriginalData);
-			SourceDataArrayPtr = &OriginalData;
-		}
-		else
-			if (!bUseWeightTargetValue)
-			{
 			// When painting weights (and not using target value mode), we use a source value that tends more
 			// to the current value as we paint over the same region multiple times.
 			this->Cache.GetOriginalData(X1, Y1, X2, Y2, OriginalData);
@@ -149,7 +142,7 @@ public:
 					SourceValue = FMath::Lerp<typename ToolTarget::CacheClass::DataType>(SourceValue, CurrentValue, FMath::Min<float>(VertexInfluence * 0.05f, 1.f));
 				}
 			}
-			}
+		}
 
 		FMatrix ToWorld = ToolTarget::ToWorldMatrix(this->LandscapeInfo);
 		FMatrix FromWorld = ToolTarget::FromWorldMatrix(this->LandscapeInfo);
@@ -158,23 +151,27 @@ public:
 		float AdjustedStrength = ToolTarget::StrengthMultiplier(this->LandscapeInfo, UISettings->BrushRadius);
 		typename ToolTarget::CacheClass::DataType DestValue = ToolTarget::CacheClass::ClampValue(255.f * UISettings->WeightTargetValue);
 
+		float DeltaTime = FMath::Min<float>(FApp::GetDeltaTime(), 0.1f); // Under 10 fps slow down paint speed
+		// * 3.0f to partially compensate for impact of DeltaTime on slowing the tools down compared to the old framerate-dependent version
+		float PaintStrength = UISettings->ToolStrength * Pressure * AdjustedStrength * DeltaTime * 3.0f;
+
 		FPlane BrushPlane;
 		TArray<FVector> Normals;
 
 		if (bUseClayBrush)
 		{
 			// Calculate normals for brush verts in data space
-			Normals.Empty(OriginalData.Num());
-			Normals.AddZeroed(OriginalData.Num());
+			Normals.Empty((*SourceDataArrayPtr).Num());
+			Normals.AddZeroed((*SourceDataArrayPtr).Num());
 
 			for (int32 Y = Y1; Y < Y2; Y++)
 			{
 				for (int32 X = X1; X < X2; X++)
 				{
-					FVector Vert00 = ToWorld.TransformPosition(FVector((float)X + 0.f, (float)Y + 0.f, DATA_AT(OriginalData, X + 0, Y + 0)));
-					FVector Vert01 = ToWorld.TransformPosition(FVector((float)X + 0.f, (float)Y + 1.f, DATA_AT(OriginalData, X + 0, Y + 1)));
-					FVector Vert10 = ToWorld.TransformPosition(FVector((float)X + 1.f, (float)Y + 0.f, DATA_AT(OriginalData, X + 1, Y + 0)));
-					FVector Vert11 = ToWorld.TransformPosition(FVector((float)X + 1.f, (float)Y + 1.f, DATA_AT(OriginalData, X + 1, Y + 1)));
+					FVector Vert00 = ToWorld.TransformPosition(FVector((float)X + 0.f, (float)Y + 0.f, DATA_AT(*SourceDataArrayPtr, X + 0, Y + 0)));
+					FVector Vert01 = ToWorld.TransformPosition(FVector((float)X + 0.f, (float)Y + 1.f, DATA_AT(*SourceDataArrayPtr, X + 0, Y + 1)));
+					FVector Vert10 = ToWorld.TransformPosition(FVector((float)X + 1.f, (float)Y + 0.f, DATA_AT(*SourceDataArrayPtr, X + 1, Y + 0)));
+					FVector Vert11 = ToWorld.TransformPosition(FVector((float)X + 1.f, (float)Y + 1.f, DATA_AT(*SourceDataArrayPtr, X + 1, Y + 1)));
 
 					FVector FaceNormal1 = ((Vert00 - Vert10) ^ (Vert10 - Vert11)).SafeNormal();
 					FVector FaceNormal2 = ((Vert11 - Vert01) ^ (Vert01 - Vert00)).SafeNormal();
@@ -204,7 +201,7 @@ public:
 				ALandscape::UnpackKey(It.Key(), X, Y);
 				float Weight = It.Value();
 
-				AveragePoint += FVector((float)X * Weight, (float)Y * Weight, (float)DATA_AT(OriginalData, FMath::FloorToInt(X), FMath::FloorToInt(Y)) * Weight);
+				AveragePoint += FVector((float)X * Weight, (float)Y * Weight, (float)DATA_AT(*SourceDataArrayPtr, FMath::FloorToInt(X), FMath::FloorToInt(Y)) * Weight);
 
 				FVector SampleNormal = DATA_AT(Normals, X, Y);
 				AverageNormal += SampleNormal * Weight;
@@ -220,7 +217,7 @@ public:
 
 			// Convert to world space
 			FVector AverageLocation = ToWorld.TransformPosition(AveragePoint);
-			FVector StrengthVector = ToWorld.TransformVector(FVector(0, 0, UISettings->ToolStrength * Pressure * AdjustedStrength));
+			FVector StrengthVector = ToWorld.TransformVector(FVector(0, 0, PaintStrength));
 
 			// Brush pushes out in the normal direction
 			FVector OffsetVector = AverageNormal * StrengthVector.Z;
@@ -233,17 +230,17 @@ public:
 			BrushPlane = FPlane(AverageLocation + OffsetVector, AverageNormal);
 		}
 
-		// Apply the brush	
-		for (auto It = BrushInfo.CreateIterator(); It; ++It)
+		// Apply the brush
+		for (auto BrushPair : BrushInfo)
 		{
 			int32 X, Y;
-			ALandscape::UnpackKey(It.Key(), X, Y);
+			ALandscape::UnpackKey(BrushPair.Key, X, Y);
 
 			// Update influence map
-			float VertexInfluence = TotalInfluenceMap.FindRef(It.Key());
-			TotalInfluenceMap.Add(It.Key(), VertexInfluence + It.Value());
+			float VertexInfluence = TotalInfluenceMap.FindRef(BrushPair.Key);
+			TotalInfluenceMap.Add(BrushPair.Key, VertexInfluence + BrushPair.Value);
 
-			float PaintAmount = It.Value() * UISettings->ToolStrength * Pressure * AdjustedStrength;
+			float PaintAmount = BrushPair.Value * PaintStrength;
 			typename ToolTarget::CacheClass::DataType& CurrentValue = DATA_AT(Data, X, Y);
 			const typename ToolTarget::CacheClass::DataType& SourceValue = DATA_AT(*SourceDataArrayPtr, X, Y);
 
@@ -258,9 +255,8 @@ public:
 					CurrentValue = FMath::Lerp(CurrentValue, DestValue, PaintAmount / AdjustedStrength);
 				}
 			}
-			else
-				if (bUseClayBrush)
-				{
+			else if (bUseClayBrush)
+			{
 				// Brush application starts from original world location at start of stroke
 				FVector WorldLoc = ToWorld.TransformPosition(FVector(X, Y, SourceValue));
 
@@ -268,7 +264,7 @@ public:
 				WorldLoc.Z = (BrushPlane.W - BrushPlane.X*WorldLoc.X - BrushPlane.Y*WorldLoc.Y) / BrushPlane.Z;
 
 				// Painted amount lerps based on brush falloff.
-				float PaintValue = FMath::Lerp<float>((float)SourceValue, FromWorld.TransformPosition(WorldLoc).Z, It.Value());
+				float PaintValue = FMath::Lerp<float>((float)SourceValue, FromWorld.TransformPosition(WorldLoc).Z, BrushPair.Value);
 
 				if (bInvert)
 				{
@@ -278,18 +274,18 @@ public:
 				{
 					CurrentValue = ToolTarget::CacheClass::ClampValue(FMath::Max<int32>(FMath::RoundToInt(PaintValue), CurrentValue));
 				}
+			}
+			else
+			{
+				if (bInvert)
+				{
+					CurrentValue = ToolTarget::CacheClass::ClampValue(FMath::Min<int32>(SourceValue - FMath::RoundToInt(PaintAmount), CurrentValue));
 				}
 				else
 				{
-					if (bInvert)
-					{
-						CurrentValue = ToolTarget::CacheClass::ClampValue(FMath::Min<int32>(SourceValue - FMath::RoundToInt(PaintAmount), CurrentValue));
-					}
-					else
-					{
-						CurrentValue = ToolTarget::CacheClass::ClampValue(FMath::Max<int32>(SourceValue + FMath::RoundToInt(PaintAmount), CurrentValue));
-					}
+					CurrentValue = ToolTarget::CacheClass::ClampValue(FMath::Max<int32>(SourceValue + FMath::RoundToInt(PaintAmount), CurrentValue));
 				}
+			}
 		}
 
 		this->Cache.SetCachedData(X1, Y1, X2, Y2, Data, UISettings->PaintingRestriction);
@@ -303,7 +299,8 @@ class FLandscapeToolPaint : public FLandscapeToolPaintBase<FWeightmapToolTarget,
 public:
 	FLandscapeToolPaint(FEdModeLandscape* InEdMode)
 		: FLandscapeToolPaintBase<FWeightmapToolTarget, FLandscapeToolStrokePaint<FWeightmapToolTarget>>(InEdMode)
-	{}
+	{
+	}
 
 	virtual const TCHAR* GetToolName() override { return TEXT("Paint"); }
 	virtual FText GetDisplayName() override { return NSLOCTEXT("UnrealEd", "LandscapeMode_Paint", "Paint"); };
@@ -314,7 +311,8 @@ class FLandscapeToolSculpt : public FLandscapeToolPaintBase<FHeightmapToolTarget
 public:
 	FLandscapeToolSculpt(FEdModeLandscape* InEdMode)
 		: FLandscapeToolPaintBase<FHeightmapToolTarget, FLandscapeToolStrokePaint<FHeightmapToolTarget>>(InEdMode)
-	{}
+	{
+	}
 
 	virtual const TCHAR* GetToolName() override { return TEXT("Sculpt"); }
 	virtual FText GetDisplayName() override { return NSLOCTEXT("UnrealEd", "LandscapeMode_Sculpt", "Sculpt"); };
@@ -330,11 +328,10 @@ class FLandscapeToolStrokeSmooth : public FLandscapeToolStrokePaintBase<ToolTarg
 public:
 	FLandscapeToolStrokeSmooth(FEdModeLandscape* InEdMode, const FLandscapeToolTarget& InTarget)
 		: FLandscapeToolStrokePaintBase<ToolTarget>(InEdMode, InTarget)
-	{}
-
-	virtual ~FLandscapeToolStrokeSmooth() {}
+	{
+	}
 
-	virtual void Apply(FEditorViewportClient* ViewportClient, FLandscapeBrush* Brush, const ULandscapeEditorObject* UISettings, const TArray<FLandscapeToolMousePosition>& MousePositions) override
+	void Apply(FEditorViewportClient* ViewportClient, FLandscapeBrush* Brush, const ULandscapeEditorObject* UISettings, const TArray<FLandscapeToolMousePosition>& MousePositions)
 	{
 		if (!this->LandscapeInfo) return;
 
@@ -414,7 +411,8 @@ class FLandscapeToolSmooth : public FLandscapeToolPaintBase<ToolTarget, FLandsca
 public:
 	FLandscapeToolSmooth(FEdModeLandscape* InEdMode)
 		: FLandscapeToolPaintBase<ToolTarget, FLandscapeToolStrokeSmooth<ToolTarget> >(InEdMode)
-	{}
+	{
+	}
 
 	virtual const TCHAR* GetToolName() override { return TEXT("Smooth"); }
 	virtual FText GetDisplayName() override { return NSLOCTEXT("UnrealEd", "LandscapeMode_Smooth", "Smooth"); };
@@ -449,9 +447,7 @@ public:
 		}
 	}
 
-	virtual ~FLandscapeToolStrokeFlatten() {}
-
-	virtual void Apply(FEditorViewportClient* ViewportClient, FLandscapeBrush* Brush, const ULandscapeEditorObject* UISettings, const TArray<FLandscapeToolMousePosition>& MousePositions) override
+	void Apply(FEditorViewportClient* ViewportClient, FLandscapeBrush* Brush, const ULandscapeEditorObject* UISettings, const TArray<FLandscapeToolMousePosition>& MousePositions)
 	{
 		if (!this->LandscapeInfo) return;
 
@@ -572,14 +568,12 @@ template<class ToolTarget>
 class FLandscapeToolFlatten : public FLandscapeToolPaintBase < ToolTarget, FLandscapeToolStrokeFlatten<ToolTarget> >
 {
 protected:
-	FVector LastMousePosition;
 	UStaticMesh* PlaneMesh;
 	UStaticMeshComponent* MeshComponent;
 
 public:
 	FLandscapeToolFlatten(FEdModeLandscape* InEdMode)
 		: FLandscapeToolPaintBase<ToolTarget, FLandscapeToolStrokeFlatten<ToolTarget> >(InEdMode)
-		, LastMousePosition(FVector::ZeroVector)
 		, PlaneMesh(LoadObject<UStaticMesh>(NULL, TEXT("/Engine/EditorLandscapeResources/FlattenPlaneMesh.FlattenPlaneMesh")))
 		, MeshComponent(NULL)
 	{
@@ -603,12 +597,13 @@ public:
 
 		if (ViewportClient->IsLevelEditorClient() && MeshComponent != NULL)
 		{
-			this->EdMode->LandscapeMouseTrace((FEditorViewportClient*)ViewportClient, x, y, LastMousePosition);
+			FVector MousePosition;
+			this->EdMode->LandscapeMouseTrace((FEditorViewportClient*)ViewportClient, x, y, MousePosition);
 
 			const FTransform LocalToWorld = this->EdMode->CurrentToolTarget.LandscapeInfo->GetLandscapeProxy()->ActorToWorld();
 			FVector Origin;
-			Origin.X = FMath::RoundToFloat(LastMousePosition.X);
-			Origin.Y = FMath::RoundToFloat(LastMousePosition.Y);
+			Origin.X = FMath::RoundToFloat(MousePosition.X);
+			Origin.Y = FMath::RoundToFloat(MousePosition.Y);
 			Origin.Z = (FMath::RoundToFloat((this->EdMode->UISettings->FlattenTarget - LocalToWorld.GetTranslation().Z) / LocalToWorld.GetScale3D().Z * LANDSCAPE_INV_ZSCALE) - 0.1f) * LANDSCAPE_ZSCALE;
 			MeshComponent->SetRelativeLocation(Origin, false);
 		}
@@ -654,11 +649,10 @@ class FLandscapeToolStrokeNoise : public FLandscapeToolStrokePaintBase<ToolTarge
 public:
 	FLandscapeToolStrokeNoise(FEdModeLandscape* InEdMode, const FLandscapeToolTarget& InTarget)
 		: FLandscapeToolStrokePaintBase<ToolTarget>(InEdMode, InTarget)
-	{}
-
-	virtual ~FLandscapeToolStrokeNoise() {}
+	{
+	}
 
-	virtual void Apply(FEditorViewportClient* ViewportClient, FLandscapeBrush* Brush, const ULandscapeEditorObject* UISettings, const TArray<FLandscapeToolMousePosition>& MousePositions) override
+	void Apply(FEditorViewportClient* ViewportClient, FLandscapeBrush* Brush, const ULandscapeEditorObject* UISettings, const TArray<FLandscapeToolMousePosition>& MousePositions)
 	{
 		if (!this->LandscapeInfo) return;
 
@@ -743,7 +737,8 @@ class FLandscapeToolNoise : public FLandscapeToolPaintBase < ToolTarget, FLandsc
 public:
 	FLandscapeToolNoise(FEdModeLandscape* InEdMode)
 		: FLandscapeToolPaintBase<ToolTarget, FLandscapeToolStrokeNoise<ToolTarget> >(InEdMode)
-	{}
+	{
+	}
 
 	virtual const TCHAR* GetToolName() override { return TEXT("Noise"); }
 	virtual FText GetDisplayName() override { return NSLOCTEXT("UnrealEd", "LandscapeMode_Noise", "Noise"); };
diff --git a/Engine/Source/Editor/LandscapeEditor/Private/LandscapeEdModeTools.h b/Engine/Source/Editor/LandscapeEditor/Private/LandscapeEdModeTools.h
--- a/Engine/Source/Editor/LandscapeEditor/Private/LandscapeEdModeTools.h
+++ b/Engine/Source/Editor/LandscapeEditor/Private/LandscapeEdModeTools.h
@@ -17,12 +17,15 @@ struct FNoiseParameter
 
 	// Constructors.
 
-	FNoiseParameter() {}
+	FNoiseParameter()
+	{
+	}
 	FNoiseParameter(float InBase, float InScale, float InAmount) :
 		Base(InBase),
 		NoiseScale(InScale),
 		NoiseAmount(InAmount)
-	{}
+	{
+	}
 
 	// Sample
 	float Sample(int32 X, int32 Y) const
@@ -1199,18 +1202,24 @@ struct FWeightmapToolTarget
 	static FMatrix FromWorldMatrix(ULandscapeInfo* LandscapeInfo) { return FMatrix::Identity; }
 };
 
-template<class TInputType>
-class FLandscapeStrokeBase
+/**
+* FLandscapeToolStrokeBase - base class for tool strokes (used by FLandscapeToolBase)
+*/
+
+class FLandscapeToolStrokeBase
 {
 public:
-	FLandscapeStrokeBase(TInputType& InTarget) {}
-	virtual void Apply(FEditorViewportClient* ViewportClient, FLandscapeBrush* Brush, const ULandscapeEditorObject* UISettings, const TArray<FLandscapeToolMousePosition>& MousePositions) = 0;
+	// Whether to call Apply() every frame even if the mouse hasn't moved
+	enum { UseContinuousApply = false };
+
+	// Signature of Apply() method:
+	// void Apply(FEditorViewportClient* ViewportClient, FLandscapeBrush* Brush, const ULandscapeEditorObject* UISettings, const TArray<FLandscapeToolMousePosition>& MousePositions);
 };
 
 
 /**
- * FLandscapeToolPaintBase - base class for painting tools
- *		ToolTarget - the target for the tool (weight or heightmaap)
+ * FLandscapeToolBase - base class for painting tools
+ *		ToolTarget - the target for the tool (weight or heightmap)
  *		StrokeClass - the class that implements the behavior for a mouse stroke applying the tool.
  */
 template<class TStrokeClass>
@@ -1220,33 +1229,48 @@ public:
 	FLandscapeToolBase(FEdModeLandscape* InEdMode)
 		: EdMode(InEdMode)
 		, bToolActive(false)
-		, ToolStroke(NULL)
-	{}
+	{
+	}
 
 	virtual bool BeginTool(FEditorViewportClient* ViewportClient, const FLandscapeToolTarget& InTarget, const FVector& InHitLocation) override
 	{
 		if (!ensure(MousePositions.Num() == 0))
 		{
-			MousePositions.Empty();
+			MousePositions.Empty(1);
 		}
 
 		bToolActive = true;
-		ToolStroke = new TStrokeClass(EdMode, InTarget);
+		ToolStroke.Emplace(EdMode, InTarget);
 
 		EdMode->CurrentBrush->BeginStroke(InHitLocation.X, InHitLocation.Y, this);
 
-		new(MousePositions)FLandscapeToolMousePosition(InHitLocation.X, InHitLocation.Y, IsShiftDown(ViewportClient->Viewport));
+		// Save the mouse position
+		LastMousePosition = FVector2D(InHitLocation);
+		MousePositions.Emplace(InHitLocation.X, InHitLocation.Y, IsShiftDown(ViewportClient->Viewport));
+		TimeSinceLastMouseMove = 0.0f;
+
 		ToolStroke->Apply(ViewportClient, EdMode->CurrentBrush, EdMode->UISettings, MousePositions);
-		MousePositions.Empty();
+
+		MousePositions.Empty(1);
 		return true;
 	}
 
 	virtual void Tick(FEditorViewportClient* ViewportClient, float DeltaTime) override
 	{
-		if (bToolActive && MousePositions.Num())
+		if (bToolActive)
 		{
-			ToolStroke->Apply(ViewportClient, EdMode->CurrentBrush, EdMode->UISettings, MousePositions);
-			MousePositions.Empty();
+			if (MousePositions.Num())
+			{
+				ToolStroke->Apply(ViewportClient, EdMode->CurrentBrush, EdMode->UISettings, MousePositions);
+				MousePositions.Empty(1);
+			}
+			else if (TStrokeClass::UseContinuousApply && TimeSinceLastMouseMove >= 0.25f)
+			{
+				MousePositions.Emplace(LastMousePosition.X, LastMousePosition.Y, IsShiftDown(ViewportClient->Viewport));
+				ToolStroke->Apply(ViewportClient, EdMode->CurrentBrush, EdMode->UISettings, MousePositions);
+				MousePositions.Empty(1);
+			}
+			TimeSinceLastMouseMove += DeltaTime;
 		}
 	}
 
@@ -1255,11 +1279,10 @@ public:
 		if (bToolActive && MousePositions.Num())
 		{
 			ToolStroke->Apply(ViewportClient, EdMode->CurrentBrush, EdMode->UISettings, MousePositions);
-			MousePositions.Empty();
+			MousePositions.Empty(1);
 		}
 
-		delete ToolStroke;
-		ToolStroke = NULL;
+		ToolStroke.Reset();
 		bToolActive = false;
 		EdMode->CurrentBrush->EndStroke();
 	}
@@ -1278,7 +1301,9 @@ public:
 			if (bToolActive)
 			{
 				// Save the mouse position
-				new(MousePositions)FLandscapeToolMousePosition(HitLocation.X, HitLocation.Y, IsShiftDown(ViewportClient->Viewport));
+				LastMousePosition = FVector2D(HitLocation);
+				MousePositions.Emplace(HitLocation.X, HitLocation.Y, IsShiftDown(ViewportClient->Viewport));
+				TimeSinceLastMouseMove = 0.0f;
 			}
 		}
 
@@ -1287,7 +1312,9 @@ public:
 
 protected:
 	TArray<FLandscapeToolMousePosition> MousePositions;
+	FVector2D LastMousePosition;
+	float TimeSinceLastMouseMove;
 	FEdModeLandscape* EdMode;
 	bool bToolActive;
-	TStrokeClass* ToolStroke;
+	TOptional<TStrokeClass> ToolStroke;
 };
diff --git a/Engine/Source/Editor/LandscapeEditor/Private/LandscapeEdModeXYOffset.cpp b/Engine/Source/Editor/LandscapeEditor/Private/LandscapeEdModeXYOffset.cpp
--- a/Engine/Source/Editor/LandscapeEditor/Private/LandscapeEdModeXYOffset.cpp
+++ b/Engine/Source/Editor/LandscapeEditor/Private/LandscapeEdModeXYOffset.cpp
@@ -50,17 +50,16 @@ namespace
 // FLandscapeToolRetopologize
 //
 
-class FLandscapeToolStrokeRetopologize
+class FLandscapeToolStrokeRetopologize : public FLandscapeToolStrokeBase
 {
 public:
 	FLandscapeToolStrokeRetopologize(FEdModeLandscape* InEdMode, const FLandscapeToolTarget& InTarget)
 		: LandscapeInfo(InTarget.LandscapeInfo.Get())
 		, Cache(InTarget)
-	{}
-
-	virtual ~FLandscapeToolStrokeRetopologize() {}
+	{
+	}
 
-	virtual void Apply(FEditorViewportClient* ViewportClient, FLandscapeBrush* Brush, const ULandscapeEditorObject* UISettings, const TArray<FLandscapeToolMousePosition>& MousePositions)
+	void Apply(FEditorViewportClient* ViewportClient, FLandscapeBrush* Brush, const ULandscapeEditorObject* UISettings, const TArray<FLandscapeToolMousePosition>& MousePositions)
 	{
 		if (!LandscapeInfo)
 		{
@@ -553,7 +552,8 @@ class FLandscapeToolRetopologize : public FLandscapeToolBase<FLandscapeToolStrok
 public:
 	FLandscapeToolRetopologize(FEdModeLandscape* InEdMode)
 		: FLandscapeToolBase<FLandscapeToolStrokeRetopologize>(InEdMode)
-	{}
+	{
+	}
 
 	virtual const TCHAR* GetToolName() override { return TEXT("Retopologize"); }
 	virtual FText GetDisplayName() override { return NSLOCTEXT("UnrealEd", "LandscapeMode_Retopologize", "Retopologize"); }
