Fixed New Landscape tooltips
TTP#345565

diff --git a/Engine/Source/Editor/LandscapeEditor/Classes/LandscapeEditorObject.h b/Engine/Source/Editor/LandscapeEditor/Classes/LandscapeEditorObject.h
--- a/Engine/Source/Editor/LandscapeEditor/Classes/LandscapeEditorObject.h
+++ b/Engine/Source/Editor/LandscapeEditor/Classes/LandscapeEditorObject.h
@@ -390,28 +390,38 @@ class ULandscapeEditorObject : public UObject
 
 	// New Landscape "Tool"
 
-	// Material initially applied to the landscape. Setting a material here exposes properties for setting up layer info based on the landscape blend nodes in the material
+	// Material initially applied to the landscape. Setting a material here exposes properties for setting up layer info based on the landscape blend nodes in the material.
 	UPROPERTY(Category="New Landscape", EditAnywhere, meta=(DisplayName="Material", ShowForTools="NewLandscape"))
 	TWeakObjectPtr<UMaterialInterface> NewLandscape_Material;
 
+	// The number of quads in a single landscape section. One section is the unit of LOD transition for landscape rendering.
 	UPROPERTY(Category="New Landscape", EditAnywhere, meta=(DisplayName="Section Size", ShowForTools="NewLandscape"))
 	int32 NewLandscape_QuadsPerSection;
+
+	// The number of sections in a single landscape component. This along with the section size determines the size of each landscape component. A component is the base unit of rendering and culling.
 	UPROPERTY(Category="New Landscape", EditAnywhere, meta=(DisplayName="Sections Per Component", ShowForTools="NewLandscape"))
 	int32 NewLandscape_SectionsPerComponent;
+
+	// The number of components in the X and Y direction, determining the overall size of the landscape.
 	UPROPERTY(Category="New Landscape", EditAnywhere, meta=(DisplayName="Number of Components", ShowForTools="NewLandscape"))
 	FIntPoint NewLandscape_ComponentCount;
 
+	// The location of the new landscape
 	UPROPERTY(Category="New Landscape", EditAnywhere, meta=(DisplayName="Location", ShowForTools="NewLandscape"))
 	FVector NewLandscape_Location;
+
+	// The rotation of the new landscape
 	UPROPERTY(Category="New Landscape", EditAnywhere, meta=(DisplayName="Rotation", ShowForTools="NewLandscape"))
 	FRotator NewLandscape_Rotation;
+
+	// The scale of the new landscape. This is the distance between each vertex on the landscape, defaulting to 100 units.
 	UPROPERTY(Category="New Landscape", EditAnywhere, meta=(DisplayName="Scale", ShowForTools="NewLandscape"))
 	FVector NewLandscape_Scale;
 
-	// Import
-
 	UPROPERTY(Category="New Landscape", VisibleAnywhere, NonTransactional, meta=(ShowForTools="NewLandscape"))
 	TEnumAsByte<ELandscapeImportHeightmapError::Type> ImportLandscape_HeightmapError;
+
+	// Specify a height map file in 16-bit RAW or PNG format
 	UPROPERTY(Category="New Landscape", EditAnywhere, NonTransactional, meta=(DisplayName="Heightmap File", ShowForTools="NewLandscape"))
 	FString ImportLandscape_HeightmapFilename;
 	UPROPERTY(NonTransactional)
@@ -421,6 +431,7 @@ class ULandscapeEditorObject : public UObject
 	UPROPERTY(NonTransactional)
 	TArray<uint16> ImportLandscape_Data;
 
+	// The landscape layers that will be created. Only layer names referenced in the material assigned above are shown here. Modify the material to add more layers.
 	UPROPERTY(Category="New Landscape", EditAnywhere, NonTransactional, EditFixedSize, meta=(DisplayName="Layers", ShowForTools="NewLandscape"))
 	TArray<FLandscapeImportLayer> ImportLandscape_Layers;
 
diff --git a/Engine/Source/Editor/LandscapeEditor/Private/LandscapeEditorDetailCustomization_NewLandscape.cpp b/Engine/Source/Editor/LandscapeEditor/Private/LandscapeEditorDetailCustomization_NewLandscape.cpp
--- a/Engine/Source/Editor/LandscapeEditor/Private/LandscapeEditorDetailCustomization_NewLandscape.cpp
+++ b/Engine/Source/Editor/LandscapeEditor/Private/LandscapeEditorDetailCustomization_NewLandscape.cpp
@@ -361,6 +361,7 @@ void FLandscapeEditorDetailCustomization_NewLandscape::CustomizeDetails(IDetailL
 			SNew(STextBlock)
 			.Font(DetailBuilder.GetDetailFont())
 			.Text(LOCTEXT("Resolution", "Overall Resolution"))
+			.ToolTipText(TAttribute<FText>(this, &FLandscapeEditorDetailCustomization_NewLandscape::GetOverallResolutionTooltip))
 		]
 	]
 	.ValueContent()
@@ -418,6 +419,7 @@ void FLandscapeEditorDetailCustomization_NewLandscape::CustomizeDetails(IDetailL
 			SNew(STextBlock)
 			.Font(DetailBuilder.GetDetailFont())
 			.Text(LOCTEXT("TotalComponents", "Total Components"))
+			.ToolTipText(LOCTEXT("NewLandscape_TotalComponents", "The total number of components that will be created for this landscape."))
 		]
 	]
 	.ValueContent()
@@ -479,6 +481,13 @@ void FLandscapeEditorDetailCustomization_NewLandscape::CustomizeDetails(IDetailL
 }
 END_SLATE_FUNCTION_BUILD_OPTIMIZATION
 
+FText FLandscapeEditorDetailCustomization_NewLandscape::GetOverallResolutionTooltip() const
+{
+	return (GetEditorMode() && GetEditorMode()->NewLandscapePreviewMode == ENewLandscapePreviewMode::ImportLandscape)
+	? LOCTEXT("ImportLandscape_OverallResolution", "Overall final resolution of the imported landscape in vertices")
+	: LOCTEXT("NewLandscape_OverallResolution", "Overall final resolution of the new landscape in vertices");
+}
+
 void FLandscapeEditorDetailCustomization_NewLandscape::SetScale(float NewValue, ETextCommit::Type, TSharedRef<IPropertyHandle> PropertyHandle)
 {
 	float OldValue = 0;
diff --git a/Engine/Source/Editor/LandscapeEditor/Private/LandscapeEditorDetailCustomization_NewLandscape.h b/Engine/Source/Editor/LandscapeEditor/Private/LandscapeEditorDetailCustomization_NewLandscape.h
--- a/Engine/Source/Editor/LandscapeEditor/Private/LandscapeEditorDetailCustomization_NewLandscape.h
+++ b/Engine/Source/Editor/LandscapeEditor/Private/LandscapeEditorDetailCustomization_NewLandscape.h
@@ -71,6 +71,8 @@ public:
 	FReply OnFitImportDataButtonClicked();
 	void ChooseBestComponentSizeForImport(FEdModeLandscape* LandscapeEdMode);
 
+	FText GetOverallResolutionTooltip() const;
+
 	// Import layers
 	EVisibility GetMaterialTipVisibility() const;
 
