Fix landscape ramp tool transform widgets not responding

diff --git a/Engine/Source/Editor/LandscapeEditor/Private/LandscapeEdModeRampTool.cpp b/Engine/Source/Editor/LandscapeEditor/Private/LandscapeEdModeRampTool.cpp
--- a/Engine/Source/Editor/LandscapeEditor/Private/LandscapeEdModeRampTool.cpp
+++ b/Engine/Source/Editor/LandscapeEditor/Private/LandscapeEdModeRampTool.cpp
@@ -161,10 +161,14 @@ public:
 				}
 
 				Points[SelectedPoint] = HitLocation;
+
+				GUnrealEd->RedrawLevelEditingViewports();
 			}
+
+			return true;
 		}
 
-		return true;
+		return false;
 	}
 
 	virtual bool HandleClick(HHitProxy* HitProxy, const FViewportClick& Click) override
