Fix landscape splines being difficult to select in spline edit mode

diff --git a/Engine/Source/Editor/LandscapeEditor/Private/LandscapeEdMode.cpp b/Engine/Source/Editor/LandscapeEditor/Private/LandscapeEdMode.cpp
--- a/Engine/Source/Editor/LandscapeEditor/Private/LandscapeEdMode.cpp
+++ b/Engine/Source/Editor/LandscapeEditor/Private/LandscapeEdMode.cpp
@@ -1181,6 +1181,7 @@ bool FEdModeLandscape::InputKey(FEditorViewportClient* ViewportClient, FViewport
 						else
 						{
 							bToolActive = CurrentTool->BeginTool(ViewportClient, CurrentToolTarget, HitLocation);
+							return bToolActive;
 						}
 					}
 				}
