Two more fixes for ramp tool misbehaving
Also you can now drag ramp points around

diff --git a/Engine/Source/Editor/LandscapeEditor/Private/LandscapeEdModeRampTool.cpp b/Engine/Source/Editor/LandscapeEditor/Private/LandscapeEdModeRampTool.cpp
--- a/Engine/Source/Editor/LandscapeEditor/Private/LandscapeEdModeRampTool.cpp
+++ b/Engine/Source/Editor/LandscapeEditor/Private/LandscapeEdModeRampTool.cpp
@@ -151,6 +151,12 @@ public:
 	{
 		if (bMovingPoint)
 		{
+			if (!Viewport->KeyState(EKeys::LeftMouseButton))
+			{
+				bMovingPoint = false;
+				return false;
+			}
+
 			FVector HitLocation;
 			if (EdMode->LandscapeMouseTrace(ViewportClient, x, y, HitLocation))
 			{
@@ -171,24 +177,6 @@ public:
 		return false;
 	}
 
-	virtual bool HandleClick(HHitProxy* HitProxy, const FViewportClick& Click) override
-	{
-		if (HitProxy)
-		{
-			if (HitProxy->IsA(HLandscapeRampToolPointHitProxy::StaticGetType()))
-			{
-				HLandscapeRampToolPointHitProxy* PointHitProxy = (HLandscapeRampToolPointHitProxy*)HitProxy;
-				SelectedPoint = PointHitProxy->Point;
-				GLevelEditorModeTools().SetWidgetMode(FWidget::WM_Translate);
-				GUnrealEd->RedrawLevelEditingViewports();
-
-				return true;
-			}
-		}
-
-		return false;
-	}
-
 	virtual bool InputKey(FEditorViewportClient* InViewportClient, FViewport* InViewport, FKey InKey, EInputEvent InEvent) override
 	{
 		if (InKey == EKeys::Enter && InEvent == IE_Pressed)
@@ -204,9 +192,37 @@ public:
 			ResetRamp();
 		}
 
-		if (SelectedPoint != INDEX_NONE)
+		// Handle clicking on points to select them and drag them around
+		if (InKey == EKeys::LeftMouseButton)
 		{
-			if (InKey == EKeys::End && InEvent == IE_Pressed)
+			if (InEvent == IE_Pressed)
+			{
+				if (!InViewport->KeyState(EKeys::MiddleMouseButton) && !InViewport->KeyState(EKeys::RightMouseButton) && !IsAltDown(InViewport) && InViewportClient->GetCurrentWidgetAxis() == EAxisList::None)
+				{
+					HHitProxy* HitProxy = InViewport->GetHitProxy(InViewport->GetMouseX(), InViewport->GetMouseY());
+					if (HitProxy && HitProxy->IsA(HLandscapeRampToolPointHitProxy::StaticGetType()))
+					{
+						HLandscapeRampToolPointHitProxy* PointHitProxy = (HLandscapeRampToolPointHitProxy*)HitProxy;
+						SelectedPoint = PointHitProxy->Point;
+						GLevelEditorModeTools().SetWidgetMode(FWidget::WM_Translate);
+						GUnrealEd->RedrawLevelEditingViewports();
+
+						bMovingPoint = true;
+						return true;
+					}
+				}
+				return false;
+			}
+			else if (InEvent == IE_Released)
+			{
+				bMovingPoint = false;
+				return false;
+			}
+		}
+
+		if (InKey == EKeys::End && InEvent == IE_Pressed)
+		{
+			if (SelectedPoint != INDEX_NONE)
 			{
 				const int32 MinX = FMath::FloorToInt(Points[SelectedPoint].X);
 				const int32 MinY = FMath::FloorToInt(Points[SelectedPoint].Y);
diff --git a/Engine/Source/Editor/UnrealEd/Private/EditorViewportClient.cpp b/Engine/Source/Editor/UnrealEd/Private/EditorViewportClient.cpp
--- a/Engine/Source/Editor/UnrealEd/Private/EditorViewportClient.cpp
+++ b/Engine/Source/Editor/UnrealEd/Private/EditorViewportClient.cpp
@@ -714,7 +714,7 @@ void FEditorViewportClient::ReceivedFocus(FViewport* InViewport)
 	ApplyRequiredCursorVisibility( true );
 
 	// Force a cursor update to make sure its returned to default as it could of been left in any state and wont update itself till an action is taken
-	SetRequiredCursorOverride( true , EMouseCursor::Default );
+	SetRequiredCursorOverride(false, EMouseCursor::Default);
 	FSlateApplication::Get().QueryCursor();
 
 	if( IsMatineeRecordingWindow() )
