Disable landscape editing in ES2 preview mode.
Fix potential crash if differing feature levels exist at any one time.

diff --git a/Engine/Source/Editor/LandscapeEditor/Private/LandscapeEdMode.cpp b/Engine/Source/Editor/LandscapeEditor/Private/LandscapeEdMode.cpp
--- a/Engine/Source/Editor/LandscapeEditor/Private/LandscapeEdMode.cpp
+++ b/Engine/Source/Editor/LandscapeEditor/Private/LandscapeEdMode.cpp
@@ -523,7 +523,7 @@ void FEdModeLandscape::Tick(FEditorViewportClient* ViewportClient, float DeltaTi
 {
 	FEdMode::Tick(ViewportClient, DeltaTime);
 
-	if (GEditor->PlayWorld != NULL)
+	if (!IsEditingEnabled())
 	{
 		return;
 	}
@@ -576,7 +576,7 @@ bool FEdModeLandscape::MouseMove(FEditorViewportClient* ViewportClient, FViewpor
 		}
 	}
 
-	if (GEditor->PlayWorld != NULL)
+	if (!IsEditingEnabled())
 	{
 		return false;
 	}
@@ -943,7 +943,7 @@ EEditAction::Type FEdModeLandscape::GetActionEditPaste()
 
 bool FEdModeLandscape::ProcessEditDuplicate()
 {
-	if (GEditor->PlayWorld != NULL)
+	if (!IsEditingEnabled())
 	{
 		return true;
 	}
@@ -963,7 +963,7 @@ bool FEdModeLandscape::ProcessEditDuplicate()
 
 bool FEdModeLandscape::ProcessEditDelete()
 {
-	if (GEditor->PlayWorld != NULL)
+	if (!IsEditingEnabled())
 	{
 		return true;
 	}
@@ -983,7 +983,7 @@ bool FEdModeLandscape::ProcessEditDelete()
 
 bool FEdModeLandscape::ProcessEditCut()
 {
-	if (GEditor->PlayWorld != NULL)
+	if (!IsEditingEnabled())
 	{
 		return true;
 	}
@@ -1003,7 +1003,7 @@ bool FEdModeLandscape::ProcessEditCut()
 
 bool FEdModeLandscape::ProcessEditCopy()
 {
-	if (GEditor->PlayWorld != NULL)
+	if (!IsEditingEnabled())
 	{
 		return true;
 	}
@@ -1045,7 +1045,7 @@ bool FEdModeLandscape::ProcessEditCopy()
 
 bool FEdModeLandscape::ProcessEditPaste()
 {
-	if (GEditor->PlayWorld != NULL)
+	if (!IsEditingEnabled())
 	{
 		return true;
 	}
@@ -1082,7 +1082,7 @@ bool FEdModeLandscape::ProcessEditPaste()
 
 bool FEdModeLandscape::HandleClick(FEditorViewportClient* InViewportClient, HHitProxy* HitProxy, const FViewportClick& Click)
 {
-	if (GEditor->PlayWorld != NULL)
+	if (!IsEditingEnabled())
 	{
 		return false;
 	}
@@ -1104,7 +1104,7 @@ bool FEdModeLandscape::HandleClick(FEditorViewportClient* InViewportClient, HHit
 /** FEdMode: Called when a key is pressed */
 bool FEdModeLandscape::InputKey(FEditorViewportClient* ViewportClient, FViewport* Viewport, FKey Key, EInputEvent Event)
 {
-	if (GEditor->PlayWorld != NULL)
+	if (!IsEditingEnabled())
 	{
 		return false;
 	}
@@ -1292,7 +1292,7 @@ bool FEdModeLandscape::InputKey(FEditorViewportClient* ViewportClient, FViewport
 /** FEdMode: Called when mouse drag input is applied */
 bool FEdModeLandscape::InputDelta(FEditorViewportClient* InViewportClient, FViewport* InViewport, FVector& InDrag, FRotator& InRot, FVector& InScale)
 {
-	if (GEditor->PlayWorld != NULL)
+	if (!IsEditingEnabled())
 	{
 		return false;
 	}
@@ -1803,7 +1803,7 @@ void FEdModeLandscape::Render(const FSceneView* View, FViewport* Viewport, FPrim
 	/** Call parent implementation */
 	FEdMode::Render(View, Viewport, PDI);
 
-	if (GEditor->PlayWorld != NULL)
+	if (!IsEditingEnabled())
 	{
 		return;
 	}
@@ -2163,7 +2163,7 @@ bool FEdModeLandscape::GetCustomInputCoordinateSystem(FMatrix& InMatrix, void* I
 /** FEdMode: Handling SelectActor */
 bool FEdModeLandscape::Select(AActor* InActor, bool bInSelected)
 {
-	if (GEditor->PlayWorld != NULL)
+	if (!IsEditingEnabled())
 	{
 		return false;
 	}
@@ -2198,7 +2198,7 @@ bool FEdModeLandscape::Select(AActor* InActor, bool bInSelected)
 /** FEdMode: Check to see if an actor can be selected in this mode - no side effects */
 bool FEdModeLandscape::IsSelectionAllowed(AActor* InActor, bool bInSelection) const
 {
-	if (GEditor->PlayWorld != NULL)
+	if (!IsEditingEnabled())
 	{
 		return false;
 	}
@@ -2602,6 +2602,35 @@ ALandscape* FEdModeLandscape::ChangeComponentSetting(int32 NumComponentsX, int32
 	return Landscape;
 }
 
+ELandscapeEditingState FEdModeLandscape::GetEditingState() const
+{
+	UWorld* World = GetWorld();
+
+	if (GEditor->bIsSimulatingInEditor)
+	{
+		return ELandscapeEditingState::SIEWorld;
+	}
+	else if (GEditor->PlayWorld != NULL)
+	{
+		return ELandscapeEditingState::PIEWorld;
+	}
+	else if (World == nullptr)
+	{
+		return ELandscapeEditingState::Unknown;
+	}
+	else if (World->FeatureLevel < ERHIFeatureLevel::SM4)
+	{
+		return ELandscapeEditingState::BadFeatureLevel;
+	}
+	else if (NewLandscapePreviewMode == ENewLandscapePreviewMode::None &&
+		!CurrentToolTarget.LandscapeInfo.IsValid())
+	{
+		return ELandscapeEditingState::NoLandscape;
+	}
+
+	return ELandscapeEditingState::Enabled;
+}
+
 bool LandscapeEditorUtils::SetHeightmapData(ALandscapeProxy* Landscape, const TArray<uint16>& Data)
 {
 	FIntRect ComponentsRect = Landscape->GetBoundingRect() + Landscape->LandscapeSectionOffset;
diff --git a/Engine/Source/Editor/LandscapeEditor/Private/SLandscapeEditor.cpp b/Engine/Source/Editor/LandscapeEditor/Private/SLandscapeEditor.cpp
--- a/Engine/Source/Editor/LandscapeEditor/Private/SLandscapeEditor.cpp
+++ b/Engine/Source/Editor/LandscapeEditor/Private/SLandscapeEditor.cpp
@@ -197,62 +197,72 @@ class FEdModeLandscape* SLandscapeEditor::GetEditorMode() const
 	return (FEdModeLandscape*)GLevelEditorModeTools().GetActiveMode(FBuiltinEditorModes::EM_Landscape);
 }
 
-void SLandscapeEditor::UpdateErrorText() const
+FText SLandscapeEditor::GetErrorText() const
 {
-	FEdModeLandscape* LandscapeEdMode = GetEditorMode();
-	if (LandscapeEdMode)
+	const FEdModeLandscape* LandscapeEdMode = GetEditorMode();
+	ELandscapeEditingState EditState = LandscapeEdMode->GetEditingState();
+	switch (EditState)
 	{
-		if (GEditor->bIsSimulatingInEditor)
+		case ELandscapeEditingState::SIEWorld:
 		{
+
 			if (LandscapeEdMode->NewLandscapePreviewMode != ENewLandscapePreviewMode::None)
 			{
-				Error->SetError(LOCTEXT("IsSimulatingError_create", "Can't create landscape while simulating!"));
+				return LOCTEXT("IsSimulatingError_create", "Can't create landscape while simulating!");
 			}
 			else
 			{
-				Error->SetError(LOCTEXT("IsSimulatingError_edit", "Can't edit landscape while simulating!"));
+				return LOCTEXT("IsSimulatingError_edit", "Can't edit landscape while simulating!");
 			}
-			return;
+			break;
 		}
-		else if (GEditor->PlayWorld != NULL)
+		case ELandscapeEditingState::PIEWorld:
 		{
 			if (LandscapeEdMode->NewLandscapePreviewMode != ENewLandscapePreviewMode::None)
 			{
-				Error->SetError(LOCTEXT("IsPIEError_create", "Can't create landscape in PIE!"));
+				return LOCTEXT("IsPIEError_create", "Can't create landscape in PIE!");
 			}
 			else
 			{
-				Error->SetError(LOCTEXT("IsPIEError_edit", "Can't edit landscape in PIE!"));
+				return LOCTEXT("IsPIEError_edit", "Can't edit landscape in PIE!");
 			}
-			return;
+			break;
 		}
-		else
+		case ELandscapeEditingState::BadFeatureLevel:
 		{
-			if (LandscapeEdMode->NewLandscapePreviewMode == ENewLandscapePreviewMode::None &&
-				!LandscapeEdMode->CurrentToolTarget.LandscapeInfo.IsValid())
+			if (LandscapeEdMode->NewLandscapePreviewMode != ENewLandscapePreviewMode::None)
 			{
-				Error->SetError(LOCTEXT("NoLandscapeError", "No Landscape!"));
-				return;
+				return LOCTEXT("IsFLError_create", "Can't create landscape with a feature level less than SM4!");
 			}
+			else
+			{
+				return LOCTEXT("IsFLError_edit", "Can't edit landscape with a feature level less than SM4!");
+			}
+			break;
+		}
+		case ELandscapeEditingState::NoLandscape:
+		{
+			return LOCTEXT("NoLandscapeError", "No Landscape!");
+		}
+		case ELandscapeEditingState::Enabled:
+		{
+			return FText::GetEmpty();
 		}
+		default:
+			checkNoEntry();
 	}
 
-	Error->SetError(FText::GetEmpty());
+	return FText::GetEmpty();
 }
 
 bool SLandscapeEditor::GetLandscapeEditorIsEnabled() const
 {
-	UpdateErrorText();
-
 	FEdModeLandscape* LandscapeEdMode = GetEditorMode();
 	if (LandscapeEdMode)
 	{
-		if (GEditor->PlayWorld == NULL)
-		{
-			return true;
-		}
+		Error->SetError(GetErrorText());
+		return LandscapeEdMode->GetEditingState() == ELandscapeEditingState::Enabled;
 	}
-
 	return false;
 }
 
diff --git a/Engine/Source/Editor/LandscapeEditor/Private/SLandscapeEditor.h b/Engine/Source/Editor/LandscapeEditor/Private/SLandscapeEditor.h
--- a/Engine/Source/Editor/LandscapeEditor/Private/SLandscapeEditor.h
+++ b/Engine/Source/Editor/LandscapeEditor/Private/SLandscapeEditor.h
@@ -70,7 +70,7 @@ public:
 protected:
 	class FEdModeLandscape* GetEditorMode() const;
 
-	void UpdateErrorText() const;
+	FText GetErrorText() const;
 
 	bool GetLandscapeEditorIsEnabled() const;
 
diff --git a/Engine/Source/Editor/LandscapeEditor/Public/LandscapeEdMode.h b/Engine/Source/Editor/LandscapeEditor/Public/LandscapeEdMode.h
--- a/Engine/Source/Editor/LandscapeEditor/Public/LandscapeEdMode.h
+++ b/Engine/Source/Editor/LandscapeEditor/Public/LandscapeEdMode.h
@@ -400,6 +400,15 @@ namespace ENewLandscapePreviewMode
 	};
 }
 
+enum class ELandscapeEditingState : uint8
+{
+	Unknown,
+	Enabled,
+	BadFeatureLevel,
+	PIEWorld,
+	SIEWorld,
+	NoLandscape,
+};
 
 /**
  * Landscape editor mode
@@ -616,6 +625,13 @@ public:
 	// For collision add visualization
 	FLandscapeAddCollision* LandscapeRenderAddCollision;
 
+	ELandscapeEditingState GetEditingState() const;
+
+	bool IsEditingEnabled() const
+	{
+		return GetEditingState() == ELandscapeEditingState::Enabled;
+	}
+
 private:
 	TArray<TSharedRef<FLandscapeTargetListInfo>> LandscapeTargetList;
 	TArray<FLandscapeListInfo> LandscapeList;
diff --git a/Engine/Source/Runtime/Landscape/Private/LandscapeRender.cpp b/Engine/Source/Runtime/Landscape/Private/LandscapeRender.cpp
--- a/Engine/Source/Runtime/Landscape/Private/LandscapeRender.cpp
+++ b/Engine/Source/Runtime/Landscape/Private/LandscapeRender.cpp
@@ -659,7 +659,7 @@ FLandscapeComponentSceneProxy::FLandscapeComponentSceneProxy(ULandscapeComponent
 #endif
 
 	bRequiresAdjacencyInformation = RequiresAdjacencyInformation(MaterialInterface, XYOffsetmapTexture == NULL ? &FLandscapeVertexFactory::StaticType : &FLandscapeXYOffsetVertexFactory::StaticType, InComponent->GetWorld()->FeatureLevel);
-	SharedBuffersKey = SubsectionSizeQuads | (NumSubsections << 16) | ( XYOffsetmapTexture == NULL ? 0 : 0x80000000);
+	SharedBuffersKey = (SubsectionSizeQuads & 0xffff) | ((NumSubsections & 0xf) << 16) | (FeatureLevel <= ERHIFeatureLevel::ES3_1 ? 0 : 1<<20) | (XYOffsetmapTexture == NULL ? 0 : 1 << 31);
 
 	DynamicMesh.LCI = ComponentLightInfo; 
 	DynamicMesh.CastShadow = true;
