Fix "Dragging Rotation Z slider does not update real-time"

diff --git a/Engine/Source/Editor/LandscapeEditor/Private/LandscapeEditorDetailCustomization_NewLandscape.cpp b/Engine/Source/Editor/LandscapeEditor/Private/LandscapeEditorDetailCustomization_NewLandscape.cpp
--- a/Engine/Source/Editor/LandscapeEditor/Private/LandscapeEditorDetailCustomization_NewLandscape.cpp
+++ b/Engine/Source/Editor/LandscapeEditor/Private/LandscapeEditorDetailCustomization_NewLandscape.cpp
@@ -227,9 +227,8 @@ void FLandscapeEditorDetailCustomization_NewLandscape::CustomizeDetails(IDetailL
 		.Roll_Static(&GetOptionalPropertyValue<float>, PropertyHandle_Rotation_Roll)
 		.Pitch_Static(&GetOptionalPropertyValue<float>, PropertyHandle_Rotation_Pitch)
 		.Yaw_Static(&GetOptionalPropertyValue<float>, PropertyHandle_Rotation_Yaw)
-		//.OnRollCommitted _Static(&SetPropertyValue<float>, PropertyHandle_Rotation_Roll) // not allowed to roll or pitch landscape
-		//.OnPitchCommitted_Static(&SetPropertyValue<float>, PropertyHandle_Rotation_Pitch) // not allowed to roll or pitch landscape
-		.OnYawCommitted_Static(&SetPropertyValue<float>, PropertyHandle_Rotation_Yaw)
+		.OnYawCommitted_Static(&SetPropertyValue<float>, PropertyHandle_Rotation_Yaw) // not allowed to roll or pitch landscape
+		.OnYawChanged_Lambda([=](float NewValue){ ensure(PropertyHandle_Rotation_Yaw->SetValue(NewValue, EPropertyValueSetFlags::InteractiveChange) == FPropertyAccess::Success); })
 	];
 
 	TSharedRef<IPropertyHandle> PropertyHandle_Scale = DetailBuilder.GetProperty(GET_MEMBER_NAME_CHECKED(ULandscapeEditorObject, NewLandscape_Scale));
