Allow (re)importing a height/weight map that doesn't match the landscape dimensions
It will now be padded/clipped to fit the same as the initial import

diff --git a/Engine/Source/Editor/LandscapeEditor/Private/LandscapeEdMode.cpp b/Engine/Source/Editor/LandscapeEditor/Private/LandscapeEdMode.cpp
--- a/Engine/Source/Editor/LandscapeEditor/Private/LandscapeEdMode.cpp
+++ b/Engine/Source/Editor/LandscapeEditor/Private/LandscapeEdMode.cpp
@@ -2307,21 +2307,39 @@ void FEdModeLandscape::ImportData(const FLandscapeTargetListInfo& TargetInfo, co
 		if (FFileHelper::LoadFileToArray(Data, *Filename) &&
 			LandscapeInfo->GetLandscapeExtent(MinX, MinY, MaxX, MaxY))
 		{
+			const int32 SizeX = (1 + MaxX - MinX);
+			const int32 SizeY = (1 + MaxY - MinY);
+
+			IImageWrapperPtr ImageWrapper = nullptr;
+			const TArray<uint8>* RawData = nullptr; // Pointer to actual used data, actual store could be in Data or ImageWrapper
+
 			if (Filename.EndsWith(".png"))
 			{
 				IImageWrapperModule& ImageWrapperModule = FModuleManager::LoadModuleChecked<IImageWrapperModule>("ImageWrapper");
-				IImageWrapperPtr ImageWrapper = ImageWrapperModule.CreateImageWrapper(EImageFormat::PNG);
+				ImageWrapper = ImageWrapperModule.CreateImageWrapper(EImageFormat::PNG);
 
 				if (ImageWrapper->SetCompressed(Data.GetData(), Data.Num()))
 				{
-					if (ImageWrapper->GetWidth() != (1 + MaxX - MinX) || ImageWrapper->GetHeight() != (1 + MaxY - MinY))
+					const int32 ImportSizeX = ImageWrapper->GetWidth();
+					const int32 ImportSizeY = ImageWrapper->GetHeight();
+
+					if (ImportSizeX != SizeX || ImportSizeY != SizeY)
 					{
-						FMessageDialog::Open(EAppMsgType::Ok,
-							FText::Format(NSLOCTEXT("UnrealEd", "LandscapeReImport_BadFileSize", "{0}'s filesize does not match with current Landscape extent"), FText::FromString(Filename)));
-						return;
+						FFormatNamedArguments Args;
+						Args.Add(TEXT("FileSizeX"), ImportSizeX);
+						Args.Add(TEXT("FileSizeY"), ImportSizeY);
+						Args.Add(TEXT("LandscapeSizeX"), SizeX);
+						Args.Add(TEXT("LandscapeSizeY"), SizeY);
+
+						auto Result = FMessageDialog::Open(EAppMsgType::OkCancel,
+							FText::Format(NSLOCTEXT("LandscapeEditor.NewLandscape", "Import_SizeMismatch", "File's size ({FileSizeX}\u00D7{FileSizeY}) does not match with current Landscape extent ({LandscapeSizeX}\u00D7{LandscapeSizeY}), if you continue it will be padded/clipped to fit"), Args));
+
+						if (Result != EAppReturnType::Ok)
+						{
+							return;
+						}
 					}
 
-					const TArray<uint8>* RawData = NULL;
 					if (TargetInfo.TargetType == ELandscapeToolTargetType::Heightmap)
 					{
 						if (ImageWrapper->GetFormat() != ERGBFormat::Gray)
@@ -2332,7 +2350,7 @@ void FEdModeLandscape::ImportData(const FLandscapeTargetListInfo& TargetInfo, co
 								return;
 							}
 						}
-						if (ImageWrapper->GetBitDepth() != 16)
+						else if (ImageWrapper->GetBitDepth() != 16)
 						{
 							EAppReturnType::Type Result = FMessageDialog::Open(EAppMsgType::OkCancel, NSLOCTEXT("LandscapeEditor.NewLandscape", "Import_HeightmapFileLowBitDepth", "The Heightmap file appears to be an 8-bit png, 16-bit is preferred. The import *can* continue, but the result may be lower quality than desired."));
 							if (Result != EAppReturnType::Ok)
@@ -2345,17 +2363,17 @@ void FEdModeLandscape::ImportData(const FLandscapeTargetListInfo& TargetInfo, co
 						{
 							ImageWrapper->GetRaw(ERGBFormat::Gray, 8, RawData);
 							Data.Reset();
-							Data.AddUninitialized(RawData->Num() * sizeof(uint16));
+							Data.SetNumUninitialized(RawData->Num() * sizeof(uint16));
 							uint16* DataPtr = (uint16*)Data.GetData();
 							for (int32 i = 0; i < RawData->Num(); i++)
 							{
 								DataPtr[i] = (*RawData)[i] * 0x101; // Expand to 16-bit
 							}
+							RawData = &Data;
 						}
 						else
 						{
 							ImageWrapper->GetRaw(ERGBFormat::Gray, 16, RawData);
-							Data = *RawData; // agh I want to use MoveTemp() here
 						}
 					}
 					else
@@ -2370,7 +2388,36 @@ void FEdModeLandscape::ImportData(const FLandscapeTargetListInfo& TargetInfo, co
 						}
 
 						ImageWrapper->GetRaw(ERGBFormat::Gray, 8, RawData);
-						Data = *RawData; // agh I want to use MoveTemp() here
+					}
+
+					if (ImportSizeX != SizeX || ImportSizeY != SizeY)
+					{
+						// Cloned from FLandscapeEditorDetailCustomization_NewLandscape.OnCreateButtonClicked
+						// so that reimports behave the same as the initial import :)
+
+						const int32 OffsetX = (SizeX - ImportSizeX) / 2;
+						const int32 OffsetY = (SizeY - ImportSizeY) / 2;
+
+						if (TargetInfo.TargetType == ELandscapeToolTargetType::Heightmap)
+						{
+							// have to use a TempData because RawData could be pointing to Data
+							TArray<uint8> TempData;
+							TempData.SetNumUninitialized(SizeX * SizeY * sizeof(uint16));
+
+							LandscapeEditorUtils::ExpandData<uint16>((uint16*)TempData.GetData(), (const uint16*)RawData->GetData(),
+								0, 0, ImportSizeX - 1, ImportSizeY - 1,
+								-OffsetX, -OffsetY, SizeX - OffsetX - 1, SizeY - OffsetY - 1);
+
+							Data = MoveTemp(TempData);
+							RawData = &Data;
+						}
+						else
+						{
+							Data = LandscapeEditorUtils::ExpandData<uint8>(*RawData,
+								0, 0, ImportSizeX - 1, ImportSizeY - 1,
+								-OffsetX, -OffsetY, SizeX - OffsetX - 1, SizeY - OffsetY - 1);
+							RawData = &Data;
+						}
 					}
 				}
 				else
@@ -2380,13 +2427,17 @@ void FEdModeLandscape::ImportData(const FLandscapeTargetListInfo& TargetInfo, co
 					return;
 				}
 			}
+			else
+			{
+				// I would love to deprecate the raw/r8/r16 support for landscape, r16 doesn't even handle endianness issues...
+			}
 
 			if (TargetInfo.TargetType == ELandscapeToolTargetType::Heightmap)
 			{
-				if (Data.Num() == (1 + MaxX - MinX)*(1 + MaxY - MinY) * 2)
+				if (RawData->Num() == SizeX * SizeY * sizeof(uint16))
 				{
 					FHeightmapAccessor<false> HeightmapAccessor(LandscapeInfo);
-					HeightmapAccessor.SetData(MinX, MinY, MaxX, MaxY, (uint16*)Data.GetData());
+					HeightmapAccessor.SetData(MinX, MinY, MaxX, MaxY, (const uint16*)RawData->GetData());
 				}
 				else
 				{
@@ -2396,10 +2447,10 @@ void FEdModeLandscape::ImportData(const FLandscapeTargetListInfo& TargetInfo, co
 			}
 			else
 			{
-				if (Data.Num() == (1 + MaxX - MinX)*(1 + MaxY - MinY))
+				if (RawData->Num() == SizeX * SizeY)
 				{
 					FAlphamapAccessor<false, true> AlphamapAccessor(LandscapeInfo, TargetInfo.LayerInfoObj.Get());
-					AlphamapAccessor.SetData(MinX, MinY, MaxX, MaxY, Data.GetData(), ELandscapeLayerPaintingRestriction::None);
+					AlphamapAccessor.SetData(MinX, MinY, MaxX, MaxY, RawData->GetData(), ELandscapeLayerPaintingRestriction::None);
 				}
 				else
 				{
diff --git a/Engine/Source/Editor/LandscapeEditor/Private/LandscapeEditorDetailCustomization_NewLandscape.cpp b/Engine/Source/Editor/LandscapeEditor/Private/LandscapeEditorDetailCustomization_NewLandscape.cpp
--- a/Engine/Source/Editor/LandscapeEditor/Private/LandscapeEditorDetailCustomization_NewLandscape.cpp
+++ b/Engine/Source/Editor/LandscapeEditor/Private/LandscapeEditorDetailCustomization_NewLandscape.cpp
@@ -529,7 +529,7 @@ TSharedRef<SWidget> FLandscapeEditorDetailCustomization_NewLandscape::GetSection
 
 	for (int32 i = 0; i < ARRAY_COUNT(SectionSizes); i++)
 	{
-		MenuBuilder.AddMenuEntry(FText::Format(LOCTEXT("NxNQuads", "{0}x{0} Quads"), FText::AsNumber(SectionSizes[i])), FText::GetEmpty(), FSlateIcon(), FExecuteAction::CreateStatic(&OnChangeSectionSize, PropertyHandle, SectionSizes[i]));
+		MenuBuilder.AddMenuEntry(FText::Format(LOCTEXT("NxNQuads", "{0}\u00D7{0} Quads"), FText::AsNumber(SectionSizes[i])), FText::GetEmpty(), FSlateIcon(), FExecuteAction::CreateStatic(&OnChangeSectionSize, PropertyHandle, SectionSizes[i]));
 	}
 
 	return MenuBuilder.MakeWidget();
@@ -551,7 +551,7 @@ FText FLandscapeEditorDetailCustomization_NewLandscape::GetSectionSize(TSharedRe
 		return NSLOCTEXT("PropertyEditor", "MultipleValues", "Multiple Values");
 	}
 
-	return FText::Format(LOCTEXT("NxNQuads", "{0}x{0} Quads"), FText::AsNumber(QuadsPerSection));
+	return FText::Format(LOCTEXT("NxNQuads", "{0}\u00D7{0} Quads"), FText::AsNumber(QuadsPerSection));
 }
 
 TSharedRef<SWidget> FLandscapeEditorDetailCustomization_NewLandscape::GetSectionsPerComponentMenu(TSharedRef<IPropertyHandle> PropertyHandle)
@@ -1084,7 +1084,7 @@ void FLandscapeEditorDetailCustomization_NewLandscape::OnImportHeightmapFilename
 							LandscapeEdMode->UISettings->ImportLandscape_HeightmapError = ELandscapeImportHeightmapError::ColorPng;
 							FMessageDialog::Open(EAppMsgType::Ok, LOCTEXT("Import_HeightmapFileColorPng", "The Heightmap file appears to be a color png, grayscale is expected. The import *can* continue, but the result may not be what you expect..."));
 						}
-						if (ImageWrapper->GetBitDepth() != 16)
+						else if (ImageWrapper->GetBitDepth() != 16)
 						{
 							LandscapeEdMode->UISettings->ImportLandscape_HeightmapError = ELandscapeImportHeightmapError::LowBitDepth;
 							FMessageDialog::Open(EAppMsgType::Ok, LOCTEXT("Import_HeightmapFileLowBitDepth", "The Heightmap file appears to be an 8-bit png, 16-bit is preferred. The import *can* continue, but the result may be lower quality than desired."));
diff --git a/Engine/Source/Editor/LandscapeEditor/Public/LandscapeEdMode.h b/Engine/Source/Editor/LandscapeEditor/Public/LandscapeEdMode.h
--- a/Engine/Source/Editor/LandscapeEditor/Public/LandscapeEdMode.h
+++ b/Engine/Source/Editor/LandscapeEditor/Public/LandscapeEdMode.h
@@ -645,7 +645,7 @@ private:
 namespace LandscapeEditorUtils
 {
 	template<typename T>
-	TArray<T> ExpandData(const TArray<T>& Data,
+	void ExpandData(T* OutData, const T* InData,
 		int32 OldMinX, int32 OldMinY, int32 OldMaxX, int32 OldMaxY,
 		int32 NewMinX, int32 NewMinY, int32 NewMaxX, int32 NewMaxY)
 	{
@@ -656,34 +656,47 @@ namespace LandscapeEditorUtils
 		const int32 OffsetX = NewMinX - OldMinX;
 		const int32 OffsetY = NewMinY - OldMinY;
 
-		TArray<T> Result;
-		Result.Empty(NewWidth * NewHeight);
-		Result.AddUninitialized(NewWidth * NewHeight);
-
 		for (int32 Y = 0; Y < NewHeight; ++Y)
 		{
 			const int32 OldY = FMath::Clamp<int32>(Y + OffsetY, 0, OldHeight - 1);
 
 			// Pad anything to the left
-			const T PadLeft = Data[OldY * OldWidth + 0];
+			const T PadLeft = InData[OldY * OldWidth + 0];
 			for (int32 X = 0; X < -OffsetX; ++X)
 			{
-				Result[Y * NewWidth + X] = PadLeft;
+				OutData[Y * NewWidth + X] = PadLeft;
 			}
 
 			// Copy one row of the old data
 			{
 				const int32 X = FMath::Max(0, -OffsetX);
 				const int32 OldX = FMath::Clamp<int32>(X + OffsetX, 0, OldWidth - 1);
-				FMemory::Memcpy(&Result[Y * NewWidth + X], &Data[OldY * OldWidth + OldX], FMath::Min<int32>(OldWidth, NewWidth) * sizeof(T));
+				FMemory::Memcpy(&OutData[Y * NewWidth + X], &InData[OldY * OldWidth + OldX], FMath::Min<int32>(OldWidth, NewWidth) * sizeof(T));
 			}
 
-			const T PadRight = Data[OldY * OldWidth + OldWidth - 1];
+			const T PadRight = InData[OldY * OldWidth + OldWidth - 1];
 			for (int32 X = -OffsetX + OldWidth; X < NewWidth; ++X)
 			{
-				Result[Y * NewWidth + X] = PadRight;
+				OutData[Y * NewWidth + X] = PadRight;
 			}
 		}
+	}
+
+	template<typename T>
+	TArray<T> ExpandData(const TArray<T>& Data,
+		int32 OldMinX, int32 OldMinY, int32 OldMaxX, int32 OldMaxY,
+		int32 NewMinX, int32 NewMinY, int32 NewMaxX, int32 NewMaxY)
+	{
+		const int32 NewWidth = NewMaxX - NewMinX + 1;
+		const int32 NewHeight = NewMaxY - NewMinY + 1;
+
+		TArray<T> Result;
+		Result.Empty(NewWidth * NewHeight);
+		Result.AddUninitialized(NewWidth * NewHeight);
+
+		ExpandData(Result.GetData(), Data.GetData(),
+			OldMinX, OldMinY, OldMaxX, OldMaxY,
+			NewMinX, NewMinY, NewMaxX, NewMaxY);
 
 		return Result;
 	}
