Fixing crash when using the landscape component resize tool in resample mode on a landscape with deleted/missing components

diff --git a/Engine/Source/Editor/LandscapeEditor/Private/LandscapeEdMode.cpp b/Engine/Source/Editor/LandscapeEditor/Private/LandscapeEdMode.cpp
--- a/Engine/Source/Editor/LandscapeEditor/Private/LandscapeEdMode.cpp
+++ b/Engine/Source/Editor/LandscapeEditor/Private/LandscapeEdMode.cpp
@@ -2500,7 +2500,7 @@ ALandscape* FEdModeLandscape::ChangeComponentSetting(int32 NumComponentsX, int32
 
 				// GetHeightData alters its args, so make temp copies to avoid screwing things up
 				int32 TMinX = OldMinX, TMinY = OldMinY, TMaxX = OldMaxX, TMaxY = OldMaxY;
-				LandscapeEdit.GetHeightData(TMinX, TMinY, TMaxX, OldMaxY, HeightData.GetData(), 0);
+				LandscapeEdit.GetHeightData(TMinX, TMinY, TMaxX, TMaxY, HeightData.GetData(), 0);
 
 				HeightData = LandscapeEditorUtils::ResampleData(HeightData,
 					OldVertsX, OldVertsY, NewVertsX, NewVertsY);
@@ -2621,7 +2621,7 @@ ALandscape* FEdModeLandscape::ChangeComponentSetting(int32 NumComponentsX, int32
 				ULandscapeInfo* NewLandscapeInfo = Landscape->GetLandscapeInfo();
 				for (const TPair<FIntPoint, ULandscapeComponent*>& Entry : LandscapeInfo->XYtoComponentMap)
 				{
-					ULandscapeComponent* NewComponent = NewLandscapeInfo->XYtoComponentMap[Entry.Key];
+					ULandscapeComponent* NewComponent = NewLandscapeInfo->XYtoComponentMap.FindRef(Entry.Key);
 					if (NewComponent)
 					{
 						ULandscapeHeightfieldCollisionComponent* OldCollisionComponent = Entry.Value->CollisionComponent.Get();
@@ -2640,6 +2640,10 @@ ALandscape* FEdModeLandscape::ChangeComponentSetting(int32 NumComponentsX, int32
 					}
 				}
 			}
+			else
+			{
+				// TODO: remap foliage when not resampling (i.e. when there isn't a 1:1 mapping between old and new component)
+			}
 
 			// Delete the old Landscape and all its proxies
 			for (TActorIterator<ALandscapeProxy> It(OldLandscapeProxy->GetWorld()); It; ++It)
