Fix landscape copy/paste tool in the all layers case

diff --git a/Engine/Source/Editor/LandscapeEditor/Private/LandscapeEdModeComponentTools.cpp b/Engine/Source/Editor/LandscapeEditor/Private/LandscapeEdModeComponentTools.cpp
--- a/Engine/Source/Editor/LandscapeEditor/Private/LandscapeEdModeComponentTools.cpp
+++ b/Engine/Source/Editor/LandscapeEditor/Private/LandscapeEdModeComponentTools.cpp
@@ -1233,6 +1233,7 @@ public:
 			bool bFullCopy = !EdMode->UISettings->bUseSelectedRegion || !LandscapeInfo->SelectedRegion.Num();
 			//bool bInverted = EdMode->UISettings->bUseSelectedRegion && EdMode->UISettings->bUseNegativeMask;
 
+			// TODO: This is a mess and badly needs refactoring
 			for (int32 Y = 0; Y < SizeY; ++Y)
 			{
 				for (int32 X = 0; X < SizeX; ++X)
@@ -1244,6 +1245,13 @@ public:
 					{
 						for (int32 i = -1; (!bApplyToAll && i < 0) || i < LayerNum; ++i)
 						{
+							// Don't try to copy data for null layers
+							if ( (bApplyToAll && i >= 0 && !LandscapeInfo->Layers[i].LayerInfoObj) ||
+								(!bApplyToAll && !EdMode->CurrentToolTarget.LayerInfo.Get()))
+							{
+								continue;
+							}
+
 							FGizmoPreData GizmoPreData[NeighborNum];
 
 							for (int32 LocalY = 0; LocalY < 2; ++LocalY)
@@ -1302,7 +1310,7 @@ public:
 								bDidCopy = true;
 							}
 
-							if (LerpedData.Data > 0.f)
+							if (LerpedData.Ratio > 0.f)
 							{
 								// Added for LayerNames
 								if (bApplyToAll)
@@ -1409,9 +1417,9 @@ public:
 				Gizmo->SampleData(SizeX, SizeY);
 
 				// Update LayerInfos
-				for (auto It = LayerInfoSet.CreateConstIterator(); It; ++It)
+				for (ULandscapeLayerInfoObject* LayerInfo : LayerInfoSet)
 				{
-					Gizmo->LayerInfos.Add(*It);
+					Gizmo->LayerInfos.Add(LayerInfo);
 				}
 			}
 
diff --git a/Engine/Source/Editor/LandscapeEditor/Private/LandscapeEdModeTools.h b/Engine/Source/Editor/LandscapeEditor/Private/LandscapeEdModeTools.h
--- a/Engine/Source/Editor/LandscapeEditor/Private/LandscapeEdModeTools.h
+++ b/Engine/Source/Editor/LandscapeEditor/Private/LandscapeEdModeTools.h
@@ -245,7 +245,7 @@ inline void LowPassFilter(int32 X1, int32 Y1, int32 X2, int32 Y2, TMap<FIntPoint
 
 
 //
-// FLandscapeHeightCache
+// TLandscapeEditCache
 //
 template<class Accessor, typename AccessorType>
 struct TLandscapeEditCache
@@ -426,7 +426,7 @@ struct TLandscapeEditCache
 
 	void SetValue(int32 LandscapeX, int32 LandscapeY, AccessorType Value)
 	{
-		CachedData.Add(ALandscape::MakeKey(LandscapeX, LandscapeY), Value);
+		CachedData.Add(ALandscape::MakeKey(LandscapeX, LandscapeY), Forward<AccessorType>(Value));
 	}
 
 	bool IsZeroValue(const FVector& Value)
@@ -1045,22 +1045,31 @@ struct FLandscapeFullWeightCache : public TLandscapeEditCache < FFullWeightmapAc
 	// Only for all weight case... the accessor type should be TArray<uint8>
 	void GetCachedData(int32 X1, int32 Y1, int32 X2, int32 Y2, TArray<uint8>& OutData, int32 ArraySize)
 	{
-		int32 NumSamples = (1 + X2 - X1)*(1 + Y2 - Y1) * ArraySize;
+		const int32 XSize = (1 + X2 - X1);
+		const int32 YSize = (1 + Y2 - Y1);
+		const int32 Stride = XSize * ArraySize;
+		int32 NumSamples = XSize * YSize * ArraySize;
 		OutData.Empty(NumSamples);
 		OutData.AddUninitialized(NumSamples);
 
 		for (int32 Y = Y1; Y <= Y2; Y++)
 		{
+			const int32 YOffset = (Y - Y1) * Stride;
 			for (int32 X = X1; X <= X2; X++)
 			{
+				const int32 XYOffset = YOffset + (X - X1) * ArraySize;
 				TArray<uint8>* Ptr = GetValueRef(X, Y);
 				if (Ptr)
 				{
 					for (int32 Z = 0; Z < ArraySize; Z++)
 					{
-						OutData[((X - X1) + (Y - Y1)*(1 + X2 - X1)) * ArraySize + Z] = (*Ptr)[Z];
+						OutData[XYOffset + Z] = (*Ptr)[Z];
 					}
 				}
+				else
+				{
+					FMemory::Memzero((void*)&OutData[XYOffset], (SIZE_T)ArraySize);
+				}
 			}
 		}
 	}
@@ -1080,7 +1089,7 @@ struct FLandscapeFullWeightCache : public TLandscapeEditCache < FFullWeightmapAc
 				{
 					Value[Z] = Data[((X - X1) + (Y - Y1)*(1 + X2 - X1)) * ArraySize + Z];
 				}
-				SetValue(X, Y, Value);
+				SetValue(X, Y, MoveTemp(Value));
 			}
 		}
 
diff --git a/Engine/Source/Runtime/Landscape/Private/LandscapeEditInterface.cpp b/Engine/Source/Runtime/Landscape/Private/LandscapeEditInterface.cpp
--- a/Engine/Source/Runtime/Landscape/Private/LandscapeEditInterface.cpp
+++ b/Engine/Source/Runtime/Landscape/Private/LandscapeEditInterface.cpp
@@ -2530,7 +2530,7 @@ void FLandscapeEditDataInterface::SetAlphaData(const TSet<ULandscapeLayerInfoObj
 			{
 				ULandscapeLayerInfoObject* const LayerInfo = *It;
 				const bool bFound = Component->WeightmapLayerAllocations.ContainsByPredicate([LayerInfo](const FWeightmapLayerAllocationInfo& Allocation){ return Allocation.LayerInfo == LayerInfo; });
-				if (bFound)
+				if (!bFound)
 				{
 					NeedAllocationInfos.Add(LayerInfo);
 				}
@@ -2650,7 +2650,7 @@ void FLandscapeEditDataInterface::SetAlphaData(const TSet<ULandscapeLayerInfoObj
 								{
 									uint8& Weight = LayerDataInfos[LayerIdx].TexDataPtr[TexDataIndex];
 
-									Weight = Data[DataIndex]; // Only for whole weight
+									Weight = LayerDataInfos[LayerIdx].InDataPtr[DataIndex]; // Only for whole weight
 									if (Weight != 0)
 									{
 										LayerEditDataAllZero[LayerIdx] = false;
