Fix moving splines to another level potentially adding a large offset
Fix crash when undoing pasting a spline
Fix description of "move spline to current level" in undo/redo

diff --git a/Engine/Source/Editor/LandscapeEditor/Private/LandscapeEdModeSplineTools.cpp b/Engine/Source/Editor/LandscapeEditor/Private/LandscapeEdModeSplineTools.cpp
--- a/Engine/Source/Editor/LandscapeEditor/Private/LandscapeEdModeSplineTools.cpp
+++ b/Engine/Source/Editor/LandscapeEditor/Private/LandscapeEdModeSplineTools.cpp
@@ -64,6 +64,7 @@ public:
 
 	void CreateSplineComponent(ALandscapeProxy* Landscape, FVector Scale3D)
 	{
+		Landscape->Modify();
 		Landscape->SplineComponent = ConstructObject<ULandscapeSplinesComponent>(ULandscapeSplinesComponent::StaticClass(), Landscape, NAME_None, RF_Transactional);
 		Landscape->SplineComponent->RelativeScale3D = Scale3D;
 		Landscape->SplineComponent->AttachTo(Landscape->GetRootComponent());
@@ -1565,6 +1566,7 @@ public:
 		{
 			CreateSplineComponent(Landscape, FVector(1.f) / Landscape->GetRootComponent()->RelativeScale3D);
 		}
+		Landscape->SplineComponent->Modify();
 
 		const TCHAR* Data = NULL;
 		FString PasteString;
@@ -1660,7 +1662,7 @@ void FEdModeLandscape::SelectAllConnectedSplineSegments()
 
 void FEdModeLandscape::SplineMoveToCurrentLevel()
 {
-	FScopedTransaction Transaction(LOCTEXT("LandscapeSpline_AddControlPoint", "Add Landscape Spline Control Point"));
+	FScopedTransaction Transaction(LOCTEXT("LandscapeSpline_MoveToCurrentLevel", "Move Landscape Spline to current level"));
 
 	if (SplinesTool /*&& SplinesTool == CurrentTool*/)
 	{
@@ -1697,11 +1699,11 @@ void FEdModeLandscape::SplineMoveToCurrentLevel()
 					{
 						SplinesTool->CreateSplineComponent(Landscape, FromProxy->SplineComponent->RelativeScale3D);
 					}
+					Landscape->SplineComponent->Modify();
 
 					const FTransform OldToNewTransform =
-						Landscape->SplineComponent->ComponentToWorld.GetRelativeTransform(FromProxy->SplineComponent->ComponentToWorld);
+						FromProxy->SplineComponent->ComponentToWorld.GetRelativeTransform(Landscape->SplineComponent->ComponentToWorld);
 
-					Landscape->SplineComponent->Modify();
 					if (FromProxies.Find(FromProxy) == NULL)
 					{
 						FromProxies.Add(FromProxy);
@@ -1758,8 +1760,8 @@ void FEdModeLandscape::SplineMoveToCurrentLevel()
 					{
 						SplinesTool->CreateSplineComponent(Landscape, FromProxy->SplineComponent->RelativeScale3D);
 					}
-
 					Landscape->SplineComponent->Modify();
+
 					if (FromProxies.Find(FromProxy) == NULL)
 					{
 						FromProxies.Add(FromProxy);
