Fix crash using erosion tool with no landscape layers

diff --git a/Engine/Source/Editor/LandscapeEditor/Private/LandscapeEdModeTools.h b/Engine/Source/Editor/LandscapeEditor/Private/LandscapeEdModeTools.h
--- a/Engine/Source/Editor/LandscapeEditor/Private/LandscapeEdModeTools.h
+++ b/Engine/Source/Editor/LandscapeEditor/Private/LandscapeEdModeTools.h
@@ -1045,6 +1045,12 @@ struct FLandscapeFullWeightCache : public TLandscapeEditCache < FFullWeightmapAc
 	// Only for all weight case... the accessor type should be TArray<uint8>
 	void GetCachedData(int32 X1, int32 Y1, int32 X2, int32 Y2, TArray<uint8>& OutData, int32 ArraySize)
 	{
+		if (ArraySize == 0)
+		{
+			OutData.Empty();
+			return;
+		}
+
 		const int32 XSize = (1 + X2 - X1);
 		const int32 YSize = (1 + Y2 - Y1);
 		const int32 Stride = XSize * ArraySize;
