Landscape wireframe on top visualization mode

diff --git a/Engine/Source/Editor/LandscapeEditor/Private/LandscapeEditorCommands.cpp b/Engine/Source/Editor/LandscapeEditor/Private/LandscapeEditorCommands.cpp
--- a/Engine/Source/Editor/LandscapeEditor/Private/LandscapeEditorCommands.cpp
+++ b/Engine/Source/Editor/LandscapeEditor/Private/LandscapeEditorCommands.cpp
@@ -82,4 +82,5 @@ void FLandscapeEditorCommands::RegisterCommands()
 	UI_COMMAND(ViewModeLOD, "LOD", "", EUserInterfaceActionType::RadioButton, FInputGesture());
 	UI_COMMAND(ViewModeLayerDensity, "Layer Density", "", EUserInterfaceActionType::RadioButton, FInputGesture());
 	UI_COMMAND(ViewModeLayerDebug, "Layer Debug", "", EUserInterfaceActionType::RadioButton, FInputGesture());
+	UI_COMMAND(ViewModeWireframeOnTop, "Wireframe on Top", "", EUserInterfaceActionType::RadioButton, FInputGesture());
 }
diff --git a/Engine/Source/Editor/LandscapeEditor/Private/LandscapeEditorModule.cpp b/Engine/Source/Editor/LandscapeEditor/Private/LandscapeEditorModule.cpp
--- a/Engine/Source/Editor/LandscapeEditor/Private/LandscapeEditorModule.cpp
+++ b/Engine/Source/Editor/LandscapeEditor/Private/LandscapeEditorModule.cpp
@@ -49,6 +49,7 @@ public:
 		CommandList->MapAction(LandscapeActions.ViewModeLOD,          FExecuteAction::CreateStatic(&ChangeLandscapeViewMode, ELandscapeViewMode::LOD),          FCanExecuteAction(), FIsActionChecked::CreateStatic(&IsLandscapeViewModeSelected, ELandscapeViewMode::LOD));
 		CommandList->MapAction(LandscapeActions.ViewModeLayerDensity, FExecuteAction::CreateStatic(&ChangeLandscapeViewMode, ELandscapeViewMode::LayerDensity), FCanExecuteAction(), FIsActionChecked::CreateStatic(&IsLandscapeViewModeSelected, ELandscapeViewMode::LayerDensity));
 		CommandList->MapAction(LandscapeActions.ViewModeLayerDebug,   FExecuteAction::CreateStatic(&ChangeLandscapeViewMode, ELandscapeViewMode::DebugLayer),   FCanExecuteAction(), FIsActionChecked::CreateStatic(&IsLandscapeViewModeSelected, ELandscapeViewMode::DebugLayer));
+		CommandList->MapAction(LandscapeActions.ViewModeWireframeOnTop,FExecuteAction::CreateStatic(&ChangeLandscapeViewMode, ELandscapeViewMode::WireframeOnTop), FCanExecuteAction(), FIsActionChecked::CreateStatic(&IsLandscapeViewModeSelected, ELandscapeViewMode::WireframeOnTop));
 
 		TSharedRef<FExtender> ViewportMenuExtender = MakeShareable(new FExtender);
 		ViewportMenuExtender->AddMenuExtension("LevelViewportLandscape", EExtensionHook::First, CommandList, FMenuExtensionDelegate::CreateStatic(&ConstructLandscapeViewportMenu));
@@ -90,6 +91,7 @@ public:
 					InMenuBuilder.AddMenuEntry(LandscapeActions.ViewModeLOD, NAME_None, LOCTEXT("LandscapeViewModeLOD", "LOD"));
 					InMenuBuilder.AddMenuEntry(LandscapeActions.ViewModeLayerDensity, NAME_None, LOCTEXT("LandscapeViewModeLayerDensity", "Layer Density"));
 					InMenuBuilder.AddMenuEntry(LandscapeActions.ViewModeLayerDebug, NAME_None, LOCTEXT("LandscapeViewModeLayerDebug", "Layer Debug"));
+					InMenuBuilder.AddMenuEntry(LandscapeActions.ViewModeWireframeOnTop, NAME_None, LOCTEXT("LandscapeViewModeWireframeOnTop", "Wireframe on Top"));				
 				}
 				InMenuBuilder.EndSection();
 			}
diff --git a/Engine/Source/Editor/LandscapeEditor/Public/LandscapeEditorCommands.h b/Engine/Source/Editor/LandscapeEditor/Public/LandscapeEditorCommands.h
--- a/Engine/Source/Editor/LandscapeEditor/Public/LandscapeEditorCommands.h
+++ b/Engine/Source/Editor/LandscapeEditor/Public/LandscapeEditorCommands.h
@@ -73,6 +73,7 @@ public:
 	TSharedPtr< FUICommandInfo > ViewModeLOD;
 	TSharedPtr< FUICommandInfo > ViewModeLayerDensity;
 	TSharedPtr< FUICommandInfo > ViewModeLayerDebug;
+	TSharedPtr< FUICommandInfo > ViewModeWireframeOnTop;
 
 	// Map
 	TMap<FName, TSharedPtr<FUICommandInfo>> NameToCommandMap;
diff --git a/Engine/Source/Runtime/Landscape/Private/LandscapeRender.cpp b/Engine/Source/Runtime/Landscape/Private/LandscapeRender.cpp
--- a/Engine/Source/Runtime/Landscape/Private/LandscapeRender.cpp
+++ b/Engine/Source/Runtime/Landscape/Private/LandscapeRender.cpp
@@ -1483,6 +1483,31 @@ void FLandscapeComponentSceneProxy::GetDynamicMeshElements(const TArray<const FS
 			}
 			break;
 
+			case ELandscapeViewMode::WireframeOnTop:
+			{
+				// base mesh
+				Mesh.MaterialRenderProxy = MaterialInterface->GetRenderProxy(false);
+				Mesh.bCanApplyViewModeOverrides = false;
+				Collector.AddMesh(ViewIndex, Mesh);
+				NumPasses++;
+				NumTriangles += Mesh.GetNumPrimitives();
+				NumDrawCalls += Mesh.Elements.Num();
+
+				// wireframe on top
+				FMeshBatch& WireMesh = Collector.AllocateMesh();
+				WireMesh = MeshTools;
+				auto WireMaterialInstance = new FColoredMaterialRenderProxy(GEngine->LevelColorationUnlitMaterial->GetRenderProxy(false), FLinearColor(0, 0, 1));
+				WireMesh.MaterialRenderProxy = WireMaterialInstance;
+				Collector.RegisterOneFrameMaterialProxy(WireMaterialInstance);
+				WireMesh.bCanApplyViewModeOverrides = false;
+				WireMesh.bWireframe = true;
+				Collector.AddMesh(ViewIndex, WireMesh);
+				NumPasses++;
+				NumTriangles += WireMesh.GetNumPrimitives();
+				NumDrawCalls++;
+			}
+			break;
+
 			default:
 
 #else
@@ -1509,7 +1534,7 @@ void FLandscapeComponentSceneProxy::GetDynamicMeshElements(const TArray<const FS
 					NumPasses++;
 					NumTriangles += Mesh.GetNumPrimitives();
 					NumDrawCalls += Mesh.Elements.Num();
-			}
+				}
 			}
 #if WITH_EDITOR
 
diff --git a/Engine/Source/Runtime/Landscape/Public/LandscapeRender.h b/Engine/Source/Runtime/Landscape/Public/LandscapeRender.h
--- a/Engine/Source/Runtime/Landscape/Public/LandscapeRender.h
+++ b/Engine/Source/Runtime/Landscape/Public/LandscapeRender.h
@@ -41,6 +41,7 @@ namespace ELandscapeViewMode
 		DebugLayer,
 		LayerDensity,
 		LOD,
+		WireframeOnTop,
 	};
 }
 
