Fix "CurrentTool && CurrentTool" in landscape editor code

diff --git a/Engine/Source/Editor/LandscapeEditor/Private/LandscapeEdMode.cpp b/Engine/Source/Editor/LandscapeEditor/Private/LandscapeEdMode.cpp
--- a/Engine/Source/Editor/LandscapeEditor/Private/LandscapeEdMode.cpp
+++ b/Engine/Source/Editor/LandscapeEditor/Private/LandscapeEdMode.cpp
@@ -550,7 +550,7 @@ void FEdModeLandscape::Tick(FEditorViewportClient* ViewportClient, float DeltaTi
 			}
 		}
 
-		if (CurrentTool && CurrentTool)
+		if (CurrentTool)
 		{
 			CurrentTool->Tick(ViewportClient, DeltaTime);
 		}
@@ -571,7 +571,7 @@ bool FEdModeLandscape::MouseMove(FEditorViewportClient* ViewportClient, FViewpor
 {
 	if (bToolActive && !Viewport->KeyState(EKeys::LeftMouseButton))
 	{
-		if (CurrentTool && CurrentTool)
+		if (CurrentTool)
 		{
 			CurrentTool->EndTool(ViewportClient);
 			bToolActive = false;
@@ -586,9 +586,9 @@ bool FEdModeLandscape::MouseMove(FEditorViewportClient* ViewportClient, FViewpor
 	bool Result = false;
 	if (NewLandscapePreviewMode == ENewLandscapePreviewMode::None)
 	{
-		if (CurrentTool && CurrentTool)
+		if (CurrentTool)
 		{
-			Result = CurrentTool && CurrentTool->MouseMove(ViewportClient, Viewport, MouseX, MouseY);
+			Result = CurrentTool->MouseMove(ViewportClient, Viewport, MouseX, MouseY);
 			//ViewportClient->Invalidate( false, false );
 		}
 	}
@@ -2054,7 +2054,7 @@ void FEdModeLandscape::Render(const FSceneView* View, FViewport* Viewport, FPrim
 	}
 
 	// Override Rendering for Splines Tool
-	if (CurrentTool && CurrentTool)
+	if (CurrentTool)
 	{
 		CurrentTool->Render(View, Viewport, PDI);
 	}
@@ -2092,7 +2092,7 @@ EAxisList::Type FEdModeLandscape::GetWidgetAxisToDraw(FWidget::EWidgetMode InWid
 	if (NewLandscapePreviewMode == ENewLandscapePreviewMode::None)
 	{
 		// Override Widget for Splines Tool
-		if (CurrentTool && CurrentTool)
+		if (CurrentTool)
 		{
 			return CurrentTool->GetWidgetAxisToDraw(InWidgetMode);
 		}
@@ -2131,7 +2131,7 @@ FVector FEdModeLandscape::GetWidgetLocation() const
 	}
 
 	// Override Widget for Splines Tool
-	if (CurrentTool && CurrentTool)
+	if (CurrentTool)
 	{
 		return CurrentTool->GetWidgetLocation();
 	}
@@ -2148,7 +2148,7 @@ bool FEdModeLandscape::GetCustomDrawingCoordinateSystem(FMatrix& InMatrix, void*
 	}
 
 	// Override Widget for Splines Tool
-	if (CurrentTool && CurrentTool)
+	if (CurrentTool)
 	{
 		InMatrix = CurrentTool->GetWidgetRotation();
 		return true;
diff --git a/Engine/Source/Editor/LandscapeEditor/Private/LandscapeEdModeBrushes.cpp b/Engine/Source/Editor/LandscapeEditor/Private/LandscapeEdModeBrushes.cpp
--- a/Engine/Source/Editor/LandscapeEditor/Private/LandscapeEdModeBrushes.cpp
+++ b/Engine/Source/Editor/LandscapeEditor/Private/LandscapeEdModeBrushes.cpp
@@ -260,7 +260,7 @@ public:
 
 							if (PaintAmount > 0.f)
 							{
-								if (EdMode->CurrentTool && EdMode->CurrentTool && EdMode->CurrentTool->GetToolType() != FLandscapeTool::TT_Mask
+								if (EdMode->CurrentTool && EdMode->CurrentTool->GetToolType() != FLandscapeTool::TT_Mask
 									&& EdMode->UISettings->bUseSelectedRegion)
 								{
 									float MaskValue = LandscapeInfo->SelectedRegion.FindRef(VertexKey);
@@ -462,7 +462,7 @@ public:
 				FIntPoint VertexKey = ALandscape::MakeKey(X, Y);
 
 				float PaintAmount = 1.0f;
-				if (EdMode->CurrentTool && EdMode->CurrentTool && EdMode->CurrentTool->GetToolType() != FLandscapeTool::TT_Mask
+				if (EdMode->CurrentTool && EdMode->CurrentTool->GetToolType() != FLandscapeTool::TT_Mask
 					&& EdMode->UISettings->bUseSelectedRegion)
 				{
 					float MaskValue = LandscapeInfo->SelectedRegion.FindRef(VertexKey);
@@ -681,7 +681,7 @@ public:
 
 						if (PaintAmount)
 						{
-							if (EdMode->CurrentTool && EdMode->CurrentTool && EdMode->CurrentTool->GetToolType() != FLandscapeTool::TT_Mask
+							if (EdMode->CurrentTool && EdMode->CurrentTool->GetToolType() != FLandscapeTool::TT_Mask
 								&& EdMode->UISettings->bUseSelectedRegion)
 							{
 								float MaskValue = LandscapeInfo->SelectedRegion.FindRef(VertexKey);
@@ -995,7 +995,7 @@ public:
 
 				if (PaintAmount > 0.f)
 				{
-					if (EdMode->CurrentTool && EdMode->CurrentTool && EdMode->CurrentTool->GetToolType() != FLandscapeTool::TT_Mask
+					if (EdMode->CurrentTool && EdMode->CurrentTool->GetToolType() != FLandscapeTool::TT_Mask
 						&& EdMode->UISettings->bUseSelectedRegion)
 					{
 						float MaskValue = LandscapeInfo->SelectedRegion.FindRef(VertexKey);
@@ -1131,7 +1131,7 @@ public:
 							// Set the brush value for this vertex
 							FIntPoint VertexKey = ALandscape::MakeKey(X, Y);
 
-							if (EdMode->CurrentTool && EdMode->CurrentTool && EdMode->CurrentTool->GetToolType() != FLandscapeTool::TT_Mask
+							if (EdMode->CurrentTool && EdMode->CurrentTool->GetToolType() != FLandscapeTool::TT_Mask
 								&& EdMode->UISettings->bUseSelectedRegion)
 							{
 								float MaskValue = LandscapeInfo->SelectedRegion.FindRef(VertexKey);
