Fix "error: '&&' within '||' [-Werror,-Wlogical-op-parentheses]"

diff --git a/Engine/Source/Editor/LandscapeEditor/Private/LandscapeEdMode.cpp b/Engine/Source/Editor/LandscapeEditor/Private/LandscapeEdMode.cpp
--- a/Engine/Source/Editor/LandscapeEditor/Private/LandscapeEdMode.cpp
+++ b/Engine/Source/Editor/LandscapeEditor/Private/LandscapeEdMode.cpp
@@ -1176,9 +1176,9 @@ bool FEdModeLandscape::InputKey(FEditorViewportClient* ViewportClient, FViewport
 
 			if ((Viewport->IsPenActive() && Viewport->GetTabletPressure() > 0.f) ||
 				(!bMovingCamera && ViewportClient->GetCurrentWidgetAxis() == EAxisList::None &&
-					(LandscapeEditorControlType == ELandscapeFoliageEditorControlType::IgnoreCtrl ||
-					 LandscapeEditorControlType == ELandscapeFoliageEditorControlType::RequireCtrl   && IsCtrlDown(Viewport) ||
-					 LandscapeEditorControlType == ELandscapeFoliageEditorControlType::RequireNoCtrl && !IsCtrlDown(Viewport))))
+					((LandscapeEditorControlType == ELandscapeFoliageEditorControlType::IgnoreCtrl) ||
+					 (LandscapeEditorControlType == ELandscapeFoliageEditorControlType::RequireCtrl   && IsCtrlDown(Viewport)) ||
+					 (LandscapeEditorControlType == ELandscapeFoliageEditorControlType::RequireNoCtrl && !IsCtrlDown(Viewport)))))
 			{
 				if (CurrentTool && (CurrentTool->GetSupportedTargetTypes() == ELandscapeToolTargetTypeMask::NA || CurrentToolTarget.TargetType != ELandscapeToolTargetType::Invalid))
 				{
