Optimized landscape editing significantly by removing the TMap of brush points and replacing it with a 2D array of data (FLandscapeBrushData class)
Most code accessing brush/landscape data by calculating indexes (e.g. Data + (Y-Y1)*XSize + (X-X1)) replaced with scanline[X] style for speed
LandscapeInfo->SelectedRegion is now only queried if it contains anything, to avoid hashing the key unneccessarily (the hash function was showing on profiles :( )
Smoothing tool's SmoothFilterKernel*Scale* replaced by the more intuitive Size
Component brush fixed to be properly centered on the cursor at odd sizes and work correctly with the add component tool when cursor is over the landscape (but part of the brush is off the landscape) and size is >1
Copy tool fixed to copy the selected area and not a circle when using ctrl+c
Also a few style cleanups. for (auto it = xyz.CreateIterator...) replaced with ranged for, enums replaced with enum classes, 0.f -> 0.0f, NULL->nullptr, etc.

diff --git a/Engine/Source/Editor/LandscapeEditor/Classes/LandscapeEditorObject.h b/Engine/Source/Editor/LandscapeEditor/Classes/LandscapeEditorObject.h
--- a/Engine/Source/Editor/LandscapeEditor/Classes/LandscapeEditorObject.h
+++ b/Engine/Source/Editor/LandscapeEditor/Classes/LandscapeEditorObject.h
@@ -251,9 +251,10 @@ class ULandscapeEditorObject : public UObject
 
 	// Smooth Tool:
 
-	// Scale multiplier for the smoothing filter kernel
-	UPROPERTY(Category="Tool Settings", EditAnywhere, NonTransactional, meta=(DisplayName="Filter Kernel Scale", ShowForTools="Smooth", ClampMin="0", ClampMax="1", UIMin="0", UIMax="1"))
-	float SmoothFilterKernelScale;
+	// The radius smoothing is performed over
+	// Higher values smooth out bigger details, lower values only smooth out smaller details
+	UPROPERTY(Category="Tool Settings", EditAnywhere, NonTransactional, meta=(DisplayName="Filter Kernel Radius", ShowForTools="Smooth", ClampMin="1", ClampMax="31", UIMin="0", UIMax="7"))
+	int32 SmoothFilterKernelSize;
 
 	// If checked, performs a detail preserving smooth using the specified detail smoothing value
 	UPROPERTY(NonTransactional)
diff --git a/Engine/Source/Editor/LandscapeEditor/Private/LandscapeEdMode.cpp b/Engine/Source/Editor/LandscapeEditor/Private/LandscapeEdMode.cpp
--- a/Engine/Source/Editor/LandscapeEditor/Private/LandscapeEdMode.cpp
+++ b/Engine/Source/Editor/LandscapeEditor/Private/LandscapeEdMode.cpp
@@ -118,7 +118,7 @@ void ALandscape::SplitHeightmap(ULandscapeComponent* Comp, bool bMoveToCurrentLe
 			MipSubsectionSizeQuads = ((MipSubsectionSizeQuads + 1) >> 1) - 1;
 		}
 
-		Comp->HeightmapScaleBias = FVector4(1.f / (float)HeightmapSizeU, 1.f / (float)HeightmapSizeV, 0.f, 0.f);
+		Comp->HeightmapScaleBias = FVector4(1.0f / (float)HeightmapSizeU, 1.0f / (float)HeightmapSizeV, 0.0f, 0.0f);
 		Comp->HeightmapTexture = HeightmapTexture;
 
 		if (Comp->MaterialInstance)
@@ -354,7 +354,7 @@ void FEdModeLandscape::Enter()
 			uint8* TexData = CurrentGizmoActor->GizmoTexture->Source.LockMip(0);
 			FMemory::Memset(TexData, 0, SquaredDataTex*sizeof(uint8));
 			// Restore Sampled Data if exist...
-			if (CurrentGizmoActor->CachedScaleXY > 0.f)
+			if (CurrentGizmoActor->CachedScaleXY > 0.0f)
 			{
 				int32 SizeX = FMath::CeilToInt(CurrentGizmoActor->CachedWidth / CurrentGizmoActor->CachedScaleXY);
 				int32 SizeY = FMath::CeilToInt(CurrentGizmoActor->CachedHeight / CurrentGizmoActor->CachedScaleXY);
@@ -660,7 +660,7 @@ namespace
 		}
 
 		FVector BaryCentric = FMath::ComputeBaryCentric2D(IntersectPoint, A, B, C);
-		if (BaryCentric.X > 0.f && BaryCentric.Y > 0.f && BaryCentric.Z > 0.f)
+		if (BaryCentric.X > 0.0f && BaryCentric.Y > 0.0f && BaryCentric.Z > 0.0f)
 		{
 			return true;
 		}
@@ -741,9 +741,9 @@ bool FEdModeLandscape::LandscapeMouseTrace(FEditorViewportClient* ViewportClient
 		FVector IntersectPoint;
 		LandscapeRenderAddCollision = NULL;
 		// Need to optimize collision for AddLandscapeComponent...?
-		for (auto It = CurrentToolTarget.LandscapeInfo->XYtoAddCollisionMap.CreateIterator(); It; ++It)
+		for (auto& XYToAddCollisionPair : CurrentToolTarget.LandscapeInfo->XYtoAddCollisionMap)
 		{
-			FLandscapeAddCollision& AddCollision = It.Value();
+			FLandscapeAddCollision& AddCollision = XYToAddCollisionPair.Value;
 			// Triangle 1
 			bCollided = RayIntersectTriangle(Start, End, AddCollision.Corners[0], AddCollision.Corners[3], AddCollision.Corners[1], IntersectPoint);
 			if (bCollided)
@@ -1163,7 +1163,7 @@ bool FEdModeLandscape::InputKey(FEditorViewportClient* ViewportClient, FViewport
 		ELandscapeFoliageEditorControlType LandscapeEditorControlType = GetDefault<ULevelEditorViewportSettings>()->LandscapeEditorControlType;
 
 		// HACK - Splines tool has not yet been updated to support not using ctrl
-		if (CurrentBrush->GetBrushType() == FLandscapeBrush::BT_Splines)
+		if (CurrentBrush->GetBrushType() == ELandscapeBrushType::Splines)
 		{
 			LandscapeEditorControlType = ELandscapeFoliageEditorControlType::RequireCtrl;
 		}
@@ -1174,7 +1174,7 @@ bool FEdModeLandscape::InputKey(FEditorViewportClient* ViewportClient, FViewport
 			// Not using "if (!ViewportClient->IsMovingCamera())" because it's wrong in ortho viewports :D
 			bool bMovingCamera = Viewport->KeyState(EKeys::MiddleMouseButton) || Viewport->KeyState(EKeys::RightMouseButton) || IsAltDown(Viewport);
 
-			if ((Viewport->IsPenActive() && Viewport->GetTabletPressure() > 0.f) ||
+			if ((Viewport->IsPenActive() && Viewport->GetTabletPressure() > 0.0f) ||
 				(!bMovingCamera && ViewportClient->GetCurrentWidgetAxis() == EAxisList::None &&
 					((LandscapeEditorControlType == ELandscapeFoliageEditorControlType::IgnoreCtrl) ||
 					 (LandscapeEditorControlType == ELandscapeFoliageEditorControlType::RequireCtrl   && IsCtrlDown(Viewport)) ||
@@ -1222,7 +1222,7 @@ bool FEdModeLandscape::InputKey(FEditorViewportClient* ViewportClient, FViewport
 		// Change Brush Size
 		if ((Event == IE_Pressed || Event == IE_Repeat) && (Key == EKeys::LeftBracket || Key == EKeys::RightBracket))
 		{
-			if (CurrentBrush->GetBrushType() == FLandscapeBrush::BT_Component)
+			if (CurrentBrush->GetBrushType() == ELandscapeBrushType::Component)
 			{
 				int32 Radius = UISettings->BrushComponentSize;
 				if (Key == EKeys::LeftBracket)
@@ -1239,31 +1239,31 @@ bool FEdModeLandscape::InputKey(FEditorViewportClient* ViewportClient, FViewport
 			else
 			{
 				float Radius = UISettings->BrushRadius;
-				float SliderMin = 0.f;
-				float SliderMax = 8192.f;
+				float SliderMin = 0.0f;
+				float SliderMax = 8192.0f;
 				float LogPosition = FMath::Clamp(Radius / SliderMax, 0.0f, 1.0f);
-				float Diff = 0.05f; //6.f / SliderMax;
+				float Diff = 0.05f; //6.0f / SliderMax;
 				if (Key == EKeys::LeftBracket)
 				{
 					Diff = -Diff;
 				}
 
-				float NewValue = Radius*(1.f + Diff);
+				float NewValue = Radius*(1.0f + Diff);
 
 				if (Key == EKeys::LeftBracket)
 				{
-					NewValue = FMath::Min(NewValue, Radius - 1.f);
+					NewValue = FMath::Min(NewValue, Radius - 1.0f);
 				}
 				else
 				{
-					NewValue = FMath::Max(NewValue, Radius + 1.f);
+					NewValue = FMath::Max(NewValue, Radius + 1.0f);
 				}
 
 				NewValue = (int32)FMath::Clamp(NewValue, SliderMin, SliderMax);
 				// convert from Exp scale to linear scale
 				//float LinearPosition = 1.0f - FMath::Pow(1.0f - LogPosition, 1.0f / 3.0f);
-				//LinearPosition = FMath::Clamp(LinearPosition + Diff, 0.f, 1.f);
-				//float NewValue = FMath::Clamp((1.0f - FMath::Pow(1.0f - LinearPosition, 3.f)) * SliderMax, SliderMin, SliderMax);
+				//LinearPosition = FMath::Clamp(LinearPosition + Diff, 0.0f, 1.0f);
+				//float NewValue = FMath::Clamp((1.0f - FMath::Pow(1.0f - LinearPosition, 3.0f)) * SliderMax, SliderMin, SliderMax);
 				//float NewValue = FMath::Clamp((SliderMax - SliderMin) * LinearPosition + SliderMin, SliderMin, SliderMax);
 
 				UISettings->BrushRadius = NewValue;
@@ -1849,7 +1849,7 @@ void FEdModeLandscape::Render(const FSceneView* View, FViewport* Viewport, FPrim
 			const TArray<uint16>& ImportHeights = UISettings->GetImportLandscapeData();
 			if (ImportHeights.Num() != 0)
 			{
-				const float InvQuadsPerComponent = 1.f / (float)QuadsPerComponent;
+				const float InvQuadsPerComponent = 1.0f / (float)QuadsPerComponent;
 				const int32 SizeX = ComponentCountX * QuadsPerComponent + 1;
 				const int32 SizeY = ComponentCountY * QuadsPerComponent + 1;
 				const int32 ImportSizeX = UISettings->ImportLandscape_Width;
@@ -1871,10 +1871,10 @@ void FEdModeLandscape::Render(const FSceneView* View, FViewport* Viewport, FPrim
 						const int32 X1 = (ComponentX + 1) * QuadsPerComponent;
 						const int32 ImportX0 = FMath::Clamp<int32>(X0 - OffsetX, 0, ImportSizeX - 1);
 						const int32 ImportX1 = FMath::Clamp<int32>(X1 - OffsetX, 0, ImportSizeX - 1);
-						const float Z00 = ((float)ImportHeights[ImportX0 + ImportY0 * ImportSizeX] - 32768.f) * LANDSCAPE_ZSCALE;
-						const float Z01 = ((float)ImportHeights[ImportX0 + ImportY1 * ImportSizeX] - 32768.f) * LANDSCAPE_ZSCALE;
-						const float Z10 = ((float)ImportHeights[ImportX1 + ImportY0 * ImportSizeX] - 32768.f) * LANDSCAPE_ZSCALE;
-						const float Z11 = ((float)ImportHeights[ImportX1 + ImportY1 * ImportSizeX] - 32768.f) * LANDSCAPE_ZSCALE;
+						const float Z00 = ((float)ImportHeights[ImportX0 + ImportY0 * ImportSizeX] - 32768.0f) * LANDSCAPE_ZSCALE;
+						const float Z01 = ((float)ImportHeights[ImportX0 + ImportY1 * ImportSizeX] - 32768.0f) * LANDSCAPE_ZSCALE;
+						const float Z10 = ((float)ImportHeights[ImportX1 + ImportY0 * ImportSizeX] - 32768.0f) * LANDSCAPE_ZSCALE;
+						const float Z11 = ((float)ImportHeights[ImportX1 + ImportY1 * ImportSizeX] - 32768.0f) * LANDSCAPE_ZSCALE;
 
 						if (ComponentX == 0)
 						{
diff --git a/Engine/Source/Editor/LandscapeEditor/Private/LandscapeEdModeBrushes.cpp b/Engine/Source/Editor/LandscapeEditor/Private/LandscapeEdModeBrushes.cpp
--- a/Engine/Source/Editor/LandscapeEditor/Private/LandscapeEdModeBrushes.cpp
+++ b/Engine/Source/Editor/LandscapeEditor/Private/LandscapeEdModeBrushes.cpp
@@ -55,10 +55,10 @@ protected:
 	/** Protected so that only subclasses can create instances of this class. */
 	FLandscapeBrushCircle(FEdModeLandscape* InEdMode, UMaterialInterface* InBrushMaterial)
 		: LastMousePosition(0, 0)
-		, BrushMaterial(NULL)
+		, BrushMaterial(nullptr)
 		, EdMode(InEdMode)
 	{
-		BrushMaterial = InBrushMaterial; //UMaterialInstanceDynamic::Create(InBrushMaterial, NULL);
+		BrushMaterial = InBrushMaterial; //UMaterialInstanceDynamic::Create(InBrushMaterial, nullptr);
 	}
 
 public:
@@ -79,9 +79,9 @@ public:
 	{
 		for (TSet<ULandscapeComponent*>::TIterator It(BrushMaterialComponents); It; ++It)
 		{
-			if ((*It)->EditToolRenderData != NULL)
+			if ((*It)->EditToolRenderData != nullptr)
 			{
-				(*It)->EditToolRenderData->Update(NULL);
+				(*It)->EditToolRenderData->Update(nullptr);
 			}
 		}
 		TArray<UMaterialInstanceDynamic*> BrushMaterialInstances;
@@ -100,101 +100,103 @@ public:
 	virtual void Tick(FEditorViewportClient* ViewportClient, float DeltaTime) override
 	{
 		ULandscapeInfo* LandscapeInfo = EdMode->CurrentToolTarget.LandscapeInfo.Get();
-		ALandscapeProxy* Proxy = LandscapeInfo ? LandscapeInfo->GetLandscapeProxy() : NULL;
-		if (Proxy)
-		{
-			const float ScaleXY = LandscapeInfo->DrawScale.X;
-			float Radius = (1.f - EdMode->UISettings->BrushFalloff) * EdMode->UISettings->BrushRadius / ScaleXY;
-			float Falloff = EdMode->UISettings->BrushFalloff * EdMode->UISettings->BrushRadius / ScaleXY;
+		ALandscapeProxy* Proxy = LandscapeInfo->GetLandscapeProxy();
+
+		const float ScaleXY = LandscapeInfo->DrawScale.X;
+		const float TotalRadius = EdMode->UISettings->BrushRadius / ScaleXY;
+		const float Radius = (1.0f - EdMode->UISettings->BrushFalloff) * TotalRadius;
+		const float Falloff = EdMode->UISettings->BrushFalloff * TotalRadius;
 
-			// Set brush material.
-			int32 X1 = FMath::FloorToInt(LastMousePosition.X - (Radius + Falloff));
-			int32 Y1 = FMath::FloorToInt(LastMousePosition.Y - (Radius + Falloff));
-			int32 X2 = FMath::CeilToInt(LastMousePosition.X + (Radius + Falloff));
-			int32 Y2 = FMath::CeilToInt(LastMousePosition.Y + (Radius + Falloff));
+		FIntRect Bounds;
+		Bounds.Min.X = FMath::FloorToInt(LastMousePosition.X - TotalRadius);
+		Bounds.Min.Y = FMath::FloorToInt(LastMousePosition.Y - TotalRadius);
+		Bounds.Max.X = FMath::CeilToInt( LastMousePosition.X + TotalRadius);
+		Bounds.Max.Y = FMath::CeilToInt( LastMousePosition.Y + TotalRadius);
+
+		TSet<ULandscapeComponent*> NewComponents;
+
+		if (!ViewportClient->IsMovingCamera())
+		{
+			// GetComponentsInRegion expects an inclusive max
+			LandscapeInfo->GetComponentsInRegion(Bounds.Min.X, Bounds.Min.Y, Bounds.Max.X - 1, Bounds.Max.Y - 1, NewComponents);
+		}
 
-			TSet<ULandscapeComponent*> NewComponents;
+		// Remove the material from any old components that are no longer in the region
+		TSet<ULandscapeComponent*> RemovedComponents = BrushMaterialComponents.Difference(NewComponents);
+		for (ULandscapeComponent* RemovedComponent : RemovedComponents)
+		{
+			BrushMaterialFreeInstances.Push(BrushMaterialInstanceMap.FindAndRemoveChecked(RemovedComponent));
 
-			if (!ViewportClient->IsMovingCamera())
+			if (ensure(RemovedComponent->EditToolRenderData != nullptr))
 			{
-				LandscapeInfo->GetComponentsInRegion(X1, Y1, X2, Y2, NewComponents);
+				RemovedComponent->EditToolRenderData->Update(nullptr);
 			}
+		}
 
-			// Remove the material from any old components that are no longer in the region
-			TSet<ULandscapeComponent*> RemovedComponents = BrushMaterialComponents.Difference(NewComponents);
-			for (ULandscapeComponent* RemovedComponent : RemovedComponents)
+		// Set brush material for components in new region
+		TSet<ULandscapeComponent*> AddedComponents = NewComponents.Difference(BrushMaterialComponents);
+		for (ULandscapeComponent* AddedComponent : AddedComponents)
+		{
+			if (ensure(AddedComponent->EditToolRenderData != nullptr))
 			{
-				BrushMaterialFreeInstances.Push(BrushMaterialInstanceMap.FindAndRemoveChecked(RemovedComponent));
-
-				if (ensure(RemovedComponent->EditToolRenderData != NULL))
+				UMaterialInstanceDynamic* BrushMaterialInstance = nullptr;
+				if (BrushMaterialFreeInstances.Num() > 0)
 				{
-					RemovedComponent->EditToolRenderData->Update(NULL);
+					BrushMaterialInstance = BrushMaterialFreeInstances.Pop();
 				}
-			}
-
-			// Set brush material for components in new region
-			TSet<ULandscapeComponent*> AddedComponents = NewComponents.Difference(BrushMaterialComponents);
-			for (ULandscapeComponent* AddedComponent : AddedComponents)
-			{
-				if (ensure(AddedComponent->EditToolRenderData != NULL))
+				else
 				{
-					UMaterialInstanceDynamic* BrushMaterialInstance = NULL;
-					if (BrushMaterialFreeInstances.Num() > 0)
-					{
-						BrushMaterialInstance = BrushMaterialFreeInstances.Pop();
-					}
-					else
-					{
-						BrushMaterialInstance = UMaterialInstanceDynamic::Create(BrushMaterial, NULL);
-					}
-					BrushMaterialInstanceMap.Add(AddedComponent, BrushMaterialInstance);
-					AddedComponent->EditToolRenderData->Update(BrushMaterialInstance);
+					BrushMaterialInstance = UMaterialInstanceDynamic::Create(BrushMaterial, nullptr);
 				}
+				BrushMaterialInstanceMap.Add(AddedComponent, BrushMaterialInstance);
+				AddedComponent->EditToolRenderData->Update(BrushMaterialInstance);
 			}
+		}
 
-			BrushMaterialComponents = MoveTemp(NewComponents);
+		BrushMaterialComponents = MoveTemp(NewComponents);
 
-			// Set params for brush material.
-			FVector WorldLocation = Proxy->LandscapeActorToWorld().TransformPosition(FVector(LastMousePosition.X, LastMousePosition.Y, 0));
+		// Set params for brush material.
+		FVector WorldLocation = Proxy->LandscapeActorToWorld().TransformPosition(FVector(LastMousePosition.X, LastMousePosition.Y, 0));
+
+		for (const auto& BrushMaterialInstancePair : BrushMaterialInstanceMap)
+		{
+			ULandscapeComponent* const Component = BrushMaterialInstancePair.Key;
+			UMaterialInstanceDynamic* const MaterialInstance = BrushMaterialInstancePair.Value;
 
-			for (auto It = BrushMaterialInstanceMap.CreateConstIterator(); It; ++It)
-			{
 				// Painting can cause the EditToolRenderData to be destructed, so update it if necessary
-				if (!AddedComponents.Contains(It.Key()) && ensure(It.Key()->EditToolRenderData != NULL))
+			if (!AddedComponents.Contains(Component) && ensure(Component->EditToolRenderData != nullptr))
+			{
+				if (Component->EditToolRenderData->ToolMaterial == nullptr)
 				{
-					if (It.Key()->EditToolRenderData->ToolMaterial == NULL)
-					{
-						It.Key()->EditToolRenderData->Update(It.Value());
-					}
+					Component->EditToolRenderData->Update(MaterialInstance);
 				}
+			}
 
-				It.Value()->SetScalarParameterValue(FName(TEXT("LocalRadius")), Radius);
-				It.Value()->SetScalarParameterValue(FName(TEXT("LocalFalloff")), Falloff);
-				It.Value()->SetVectorParameterValue(FName(TEXT("WorldPosition")), FLinearColor(WorldLocation.X, WorldLocation.Y, WorldLocation.Z, ScaleXY));
+			MaterialInstance->SetScalarParameterValue(FName(TEXT("LocalRadius")), Radius);
+			MaterialInstance->SetScalarParameterValue(FName(TEXT("LocalFalloff")), Falloff);
+			MaterialInstance->SetVectorParameterValue(FName(TEXT("WorldPosition")), FLinearColor(WorldLocation.X, WorldLocation.Y, WorldLocation.Z, ScaleXY));
 
-				bool bCanPaint = true;
+			bool bCanPaint = true;
 
-				const ULandscapeComponent* Component = It.Key();
-				const ALandscapeProxy* LandscapeProxy = Component->GetLandscapeProxy();
-				const ULandscapeLayerInfoObject* LayerInfo = EdMode->CurrentToolTarget.LayerInfo.Get();
+			const ALandscapeProxy* LandscapeProxy = Component->GetLandscapeProxy();
+			const ULandscapeLayerInfoObject* LayerInfo = EdMode->CurrentToolTarget.LayerInfo.Get();
 
-				if (EdMode->CurrentToolTarget.TargetType == ELandscapeToolTargetType::Weightmap &&
-					EdMode->UISettings->PaintingRestriction != ELandscapeLayerPaintingRestriction::None)
+			if (EdMode->CurrentToolTarget.TargetType == ELandscapeToolTargetType::Weightmap &&
+				EdMode->UISettings->PaintingRestriction != ELandscapeLayerPaintingRestriction::None)
+			{
+				bool bExisting = Component->WeightmapLayerAllocations.ContainsByPredicate([LayerInfo](const FWeightmapLayerAllocationInfo& Allocation){ return Allocation.LayerInfo == LayerInfo; });
+				if (!bExisting)
 				{
-					bool bExisting = Component->WeightmapLayerAllocations.ContainsByPredicate([LayerInfo](const FWeightmapLayerAllocationInfo& Allocation){ return Allocation.LayerInfo == LayerInfo; });
-					if (!bExisting)
+					if (EdMode->UISettings->PaintingRestriction == ELandscapeLayerPaintingRestriction::ExistingOnly ||
+						(EdMode->UISettings->PaintingRestriction == ELandscapeLayerPaintingRestriction::UseMaxLayers &&
+						LandscapeProxy->MaxPaintedLayersPerComponent > 0 && Component->WeightmapLayerAllocations.Num() >= LandscapeProxy->MaxPaintedLayersPerComponent))
 					{
-						if (EdMode->UISettings->PaintingRestriction == ELandscapeLayerPaintingRestriction::ExistingOnly ||
-							(EdMode->UISettings->PaintingRestriction == ELandscapeLayerPaintingRestriction::UseMaxLayers &&
-							LandscapeProxy->MaxPaintedLayersPerComponent > 0 && Component->WeightmapLayerAllocations.Num() >= LandscapeProxy->MaxPaintedLayersPerComponent))
-						{
-							bCanPaint = false;
-						}
+						bCanPaint = false;
 					}
 				}
-
-				It.Value()->SetScalarParameterValue("CanPaint", bCanPaint ? 1.0f : 0.0f);
 			}
+
+			MaterialInstance->SetScalarParameterValue("CanPaint", bCanPaint ? 1.0f : 0.0f);
 		}
 	}
 
@@ -209,83 +211,107 @@ public:
 		return bUpdate;
 	}
 
-	virtual bool ApplyBrush(const TArray<FLandscapeToolMousePosition>& MousePositions, TMap<FIntPoint, float>& OutBrush, int32& X1, int32& Y1, int32& X2, int32& Y2) override
+	virtual FLandscapeBrushData ApplyBrush(const TArray<FLandscapeToolMousePosition>& InMousePositions) override
 	{
 		ULandscapeInfo* LandscapeInfo = EdMode->CurrentToolTarget.LandscapeInfo.Get();
-		if (LandscapeInfo)
+		const float ScaleXY = LandscapeInfo->DrawScale.X;
+		const float TotalRadius = EdMode->UISettings->BrushRadius / ScaleXY;
+		const float Radius = (1.0f - EdMode->UISettings->BrushFalloff) * TotalRadius;
+		const float Falloff = EdMode->UISettings->BrushFalloff * TotalRadius;
+
+		// Cap number of mouse positions to a sensible number
+		TArray<FLandscapeToolMousePosition> MousePositions;
+		if (InMousePositions.Num() > 10)
 		{
-			float ScaleXY = LandscapeInfo->DrawScale.X;
-			float Radius = (1.f - EdMode->UISettings->BrushFalloff) * EdMode->UISettings->BrushRadius / ScaleXY;
-			float Falloff = EdMode->UISettings->BrushFalloff * EdMode->UISettings->BrushRadius / ScaleXY;
+			for (int32 i = 0; i < 10; ++i)
+			{
+				// Scale so we include the first and last of the input positions
+				MousePositions.Add(InMousePositions[(i * (InMousePositions.Num() - 1)) / 9]);
+			}
+		}
+		else
+		{
+			MousePositions = InMousePositions;
+		}
+
+		FIntRect Bounds;
+		for (const FLandscapeToolMousePosition& MousePosition : MousePositions)
+		{
+			FIntRect SpotBounds;
+			SpotBounds.Min.X = FMath::FloorToInt(MousePosition.Position.X - TotalRadius);
+			SpotBounds.Min.Y = FMath::FloorToInt(MousePosition.Position.Y - TotalRadius);
+			SpotBounds.Max.X = FMath::CeilToInt( MousePosition.Position.X + TotalRadius);
+			SpotBounds.Max.Y = FMath::CeilToInt( MousePosition.Position.Y + TotalRadius);
 
-			bool bHasOutput = false;
-			for (int32 MoveIdx = 0; MoveIdx < MousePositions.Num(); MoveIdx++)
+			if (Bounds.IsEmpty())
 			{
-				float MousePositionX = MousePositions[MoveIdx].PositionX;
-				float MousePositionY = MousePositions[MoveIdx].PositionY;
+				Bounds = SpotBounds;
+			}
+			else
+			{
+				Bounds.Min = Bounds.Min.ComponentMin(SpotBounds.Min);
+				Bounds.Max = Bounds.Max.ComponentMax(SpotBounds.Max);
+			}
+		}
 
-				int32 ThisMoveX1 = FMath::FloorToInt(MousePositionX - (Radius + Falloff));
-				int32 ThisMoveY1 = FMath::FloorToInt(MousePositionY - (Radius + Falloff));
-				int32 ThisMoveX2 = FMath::CeilToInt(MousePositionX + (Radius + Falloff));
-				int32 ThisMoveY2 = FMath::CeilToInt(MousePositionY + (Radius + Falloff));
+		// Clamp to landscape bounds
+		int32 MinX, MaxX, MinY, MaxY;
+		if (!ensure(LandscapeInfo->GetLandscapeExtent(MinX, MinY, MaxX, MaxY)))
+		{
+			// Landscape has no components somehow
+			return FLandscapeBrushData();
+		}
+		Bounds.Clip(FIntRect(MinX, MinY, MaxX + 1, MaxY + 1));
 
-				if (MoveIdx == 0)
-				{
-					X1 = ThisMoveX1;
-					Y1 = ThisMoveY1;
-					X2 = ThisMoveX2;
-					Y2 = ThisMoveY2;
-				}
-				else
-				{
-					X1 = FMath::Min<int32>(ThisMoveX1, X1);
-					Y1 = FMath::Min<int32>(ThisMoveY1, Y1);
-					X2 = FMath::Max<int32>(ThisMoveX2, X2);
-					Y2 = FMath::Max<int32>(ThisMoveY2, Y2);
-				}
+		FLandscapeBrushData BrushData(Bounds);
+
+		for (const FLandscapeToolMousePosition& MousePosition : MousePositions)
+		{
+			FIntRect SpotBounds;
+			SpotBounds.Min.X = FMath::Max(FMath::FloorToInt(MousePosition.Position.X - TotalRadius), Bounds.Min.X);
+			SpotBounds.Min.Y = FMath::Max(FMath::FloorToInt(MousePosition.Position.Y - TotalRadius), Bounds.Min.Y);
+			SpotBounds.Max.X = FMath::Min(FMath::CeilToInt( MousePosition.Position.X + TotalRadius), Bounds.Max.X);
+			SpotBounds.Max.Y = FMath::Min(FMath::CeilToInt( MousePosition.Position.Y + TotalRadius), Bounds.Max.Y);
 
-				for (int32 Y = ThisMoveY1; Y <= ThisMoveY2; Y++)
+			for (int32 Y = SpotBounds.Min.Y; Y < SpotBounds.Max.Y; Y++)
+			{
+				float* Scanline = BrushData.GetDataPtr(FIntPoint(0, Y));
+
+				for (int32 X = SpotBounds.Min.X; X < SpotBounds.Max.X; X++)
 				{
-					for (int32 X = ThisMoveX1; X <= ThisMoveX2; X++)
+					float PrevAmount = Scanline[X];
+					if (PrevAmount < 1.0f)
 					{
-						FIntPoint VertexKey(X, Y);
-
-						float* PrevAmount = OutBrush.Find(VertexKey);
-						if (!PrevAmount || *PrevAmount < 1.f)
-						{
-							// Distance from mouse
-							float MouseDist = FMath::Sqrt(FMath::Square(MousePositionX - (float)X) + FMath::Square(MousePositionY - (float)Y));
+						// Distance from mouse
+						float MouseDist = FMath::Sqrt(FMath::Square(MousePosition.Position.X - (float)X) + FMath::Square(MousePosition.Position.Y - (float)Y));
 
-							float PaintAmount = CalculateFalloff(MouseDist, Radius, Falloff);
+						float PaintAmount = CalculateFalloff(MouseDist, Radius, Falloff);
 
-							if (PaintAmount > 0.f)
+						if (PaintAmount > 0.0f)
+						{
+							if (EdMode->CurrentTool && EdMode->CurrentTool->GetToolType() != ELandscapeToolType::Mask
+								&& EdMode->UISettings->bUseSelectedRegion && LandscapeInfo->SelectedRegion.Num() > 0)
 							{
-								if (EdMode->CurrentTool && EdMode->CurrentTool->GetToolType() != FLandscapeTool::TT_Mask
-									&& EdMode->UISettings->bUseSelectedRegion)
+								float MaskValue = LandscapeInfo->SelectedRegion.FindRef(ALandscape::MakeKey(X, Y));
+								if (EdMode->UISettings->bUseNegativeMask)
 								{
-									float MaskValue = LandscapeInfo->SelectedRegion.FindRef(VertexKey);
-									if (EdMode->UISettings->bUseNegativeMask)
-									{
-										MaskValue = 1.f - MaskValue;
-									}
-									PaintAmount *= MaskValue;
+									MaskValue = 1.0f - MaskValue;
 								}
+								PaintAmount *= MaskValue;
+							}
 
-								if (!PrevAmount || *PrevAmount < PaintAmount)
-								{
-									// Set the brush value for this vertex
-									OutBrush.Add(VertexKey, PaintAmount);
-									bHasOutput = true;
-								}
+							if (PaintAmount > PrevAmount)
+							{
+								// Set the brush value for this vertex
+								Scanline[X] = PaintAmount;
 							}
 						}
 					}
 				}
 			}
-
-			return bHasOutput;
 		}
-		return false;
+
+		return BrushData;
 	}
 };
 
@@ -307,10 +333,10 @@ public:
 	FEdModeLandscape* EdMode;
 
 	FLandscapeBrushComponent(FEdModeLandscape* InEdMode)
-		: BrushMaterial(NULL)
+		: BrushMaterial(nullptr)
 		, EdMode(InEdMode)
 	{
-		BrushMaterial = LoadObject<UMaterial>(NULL, TEXT("/Engine/EditorLandscapeResources/SelectBrushMaterial.SelectBrushMaterial"), NULL, LOAD_None, NULL);
+		BrushMaterial = LoadObject<UMaterial>(nullptr, TEXT("/Engine/EditorLandscapeResources/SelectBrushMaterial.SelectBrushMaterial"), nullptr, LOAD_None, nullptr);
 	}
 
 	// FGCObject interface
@@ -319,15 +345,15 @@ public:
 		Collector.AddReferencedObject(BrushMaterial);
 	}
 
-	virtual EBrushType GetBrushType() override { return BT_Component; }
+	virtual ELandscapeBrushType GetBrushType() override { return ELandscapeBrushType::Component; }
 
 	virtual void LeaveBrush() override
 	{
 		for (TSet<ULandscapeComponent*>::TIterator It(BrushMaterialComponents); It; ++It)
 		{
-			if ((*It)->EditToolRenderData != NULL)
+			if ((*It)->EditToolRenderData != nullptr)
 			{
-				(*It)->EditToolRenderData->Update(NULL);
+				(*It)->EditToolRenderData->Update(nullptr);
 			}
 		}
 		BrushMaterialComponents.Empty();
@@ -348,15 +374,16 @@ public:
 			ULandscapeInfo* LandscapeInfo = EdMode->CurrentToolTarget.LandscapeInfo.Get();
 			if (LandscapeInfo && LandscapeInfo->ComponentSizeQuads > 0)
 			{
-				int32 X = FMath::FloorToInt(LastMousePosition.X);
-				int32 Y = FMath::FloorToInt(LastMousePosition.Y);
+				const int32 BrushSize = FMath::Max(EdMode->UISettings->BrushComponentSize, 0);
 
-				int32 ComponentIndexX = (X >= 0.f) ? X / LandscapeInfo->ComponentSizeQuads : (X + 1) / LandscapeInfo->ComponentSizeQuads - 1;
-				int32 ComponentIndexY = (Y >= 0.f) ? Y / LandscapeInfo->ComponentSizeQuads : (Y + 1) / LandscapeInfo->ComponentSizeQuads - 1;
-				int32 BrushSize = FMath::Max(EdMode->UISettings->BrushComponentSize - 1, 0);
-				for (int32 YIndex = -(BrushSize >> 1); YIndex <= (BrushSize >> 1) + (BrushSize % 2); ++YIndex)
+				const float BrushOriginX = LastMousePosition.X / LandscapeInfo->ComponentSizeQuads - (BrushSize - 1) / 2.0f;
+				const float BrushOriginY = LastMousePosition.Y / LandscapeInfo->ComponentSizeQuads - (BrushSize - 1) / 2.0f;
+				const int32 ComponentIndexX = FMath::FloorToInt(BrushOriginX);
+				const int32 ComponentIndexY = FMath::FloorToInt(BrushOriginY);
+
+				for (int32 YIndex = 0; YIndex < BrushSize; ++YIndex)
 				{
-					for (int32 XIndex = -(BrushSize >> 1); XIndex <= (BrushSize >> 1) + (BrushSize % 2); ++XIndex)
+					for (int32 XIndex = 0; XIndex < BrushSize; ++XIndex)
 					{
 						ULandscapeComponent* Component = LandscapeInfo->XYtoComponentMap.FindRef(FIntPoint((ComponentIndexX + XIndex), (ComponentIndexY + YIndex)));
 						if (Component && FLevelUtils::IsLevelVisible(Component->GetLandscapeProxy()->GetLevel()))
@@ -380,7 +407,7 @@ public:
 				// Set brush material for components in new region
 				for (ULandscapeComponent* NewComponent : NewComponents)
 				{
-					if (NewComponent->EditToolRenderData != NULL)
+					if (NewComponent->EditToolRenderData != nullptr)
 					{
 						NewComponent->EditToolRenderData->Update(BrushMaterial);
 					}
@@ -392,9 +419,9 @@ public:
 		TSet<ULandscapeComponent*> RemovedComponents = BrushMaterialComponents.Difference(NewComponents);
 		for (ULandscapeComponent* RemovedComponent : RemovedComponents)
 		{
-			if (RemovedComponent->EditToolRenderData != NULL)
+			if (RemovedComponent->EditToolRenderData != nullptr)
 			{
-				RemovedComponent->EditToolRenderData->Update(NULL);
+				RemovedComponent->EditToolRenderData->Update(nullptr);
 			}
 		}
 
@@ -410,75 +437,82 @@ public:
 	{
 		bool bUpdate = false;
 
-
 		return bUpdate;
 	}
 
-	virtual bool ApplyBrush(const TArray<FLandscapeToolMousePosition>& MousePositions, TMap<FIntPoint, float>& OutBrush, int32& X1, int32& Y1, int32& X2, int32& Y2) override
+	virtual FLandscapeBrushData ApplyBrush(const TArray<FLandscapeToolMousePosition>& MousePositions) override
 	{
 		// Selection Brush only works for 
 		ULandscapeInfo* LandscapeInfo = EdMode->CurrentToolTarget.LandscapeInfo.Get();
-		X1 = INT_MAX;
-		Y1 = INT_MAX;
-		X2 = INT_MIN;
-		Y2 = INT_MIN;
 
-		// Check for add component...
-		if (EdMode->LandscapeRenderAddCollision)
-		{
-			// Apply Brush size..
-			int32 X = FMath::FloorToInt(LastMousePosition.X);
-			int32 Y = FMath::FloorToInt(LastMousePosition.Y);
+		FIntRect Bounds;
 
-			int32 ComponentIndexX = (X >= 0.f) ? X / LandscapeInfo->ComponentSizeQuads : (X + 1) / LandscapeInfo->ComponentSizeQuads - 1;
-			int32 ComponentIndexY = (Y >= 0.f) ? Y / LandscapeInfo->ComponentSizeQuads : (Y + 1) / LandscapeInfo->ComponentSizeQuads - 1;
+		// The add component tool needs the raw bounds of the brush rather than the bounds of the actually existing components under the brush
+		if (EdMode->CurrentTool->GetToolName() == FName("AddComponent"))
+		{
+			const int32 BrushSize = FMath::Max(EdMode->UISettings->BrushComponentSize, 0);
 
-			int32 BrushSize = FMath::Max(EdMode->UISettings->BrushComponentSize - 1, 0);
+			const float BrushOriginX = LastMousePosition.X / LandscapeInfo->ComponentSizeQuads - (BrushSize - 1) / 2.0f;
+			const float BrushOriginY = LastMousePosition.Y / LandscapeInfo->ComponentSizeQuads - (BrushSize - 1) / 2.0f;
+			const int32 ComponentIndexX = FMath::FloorToInt(BrushOriginX);
+			const int32 ComponentIndexY = FMath::FloorToInt(BrushOriginY);
 
-			X1 = (ComponentIndexX - (BrushSize >> 1)) * LandscapeInfo->ComponentSizeQuads;
-			X2 = (ComponentIndexX + (BrushSize >> 1) + (BrushSize % 2) + 1) * LandscapeInfo->ComponentSizeQuads;
-			Y1 = (ComponentIndexY - (BrushSize >> 1)) * LandscapeInfo->ComponentSizeQuads;
-			Y2 = (ComponentIndexY + (BrushSize >> 1) + (BrushSize % 2) + 1) * LandscapeInfo->ComponentSizeQuads;
+			Bounds.Min.X = (ComponentIndexX) * LandscapeInfo->ComponentSizeQuads;
+			Bounds.Min.Y = (ComponentIndexY) * LandscapeInfo->ComponentSizeQuads;
+			Bounds.Max.X = (ComponentIndexX + BrushSize) * LandscapeInfo->ComponentSizeQuads + 1;
+			Bounds.Max.Y = (ComponentIndexY + BrushSize) * LandscapeInfo->ComponentSizeQuads + 1;
 		}
 		else
 		{
+			if (BrushMaterialComponents.Num() == 0)
+			{
+				return FLandscapeBrushData();
+			}
+
 			// Get extent for all components
-			for (TSet<ULandscapeComponent*>::TIterator It(BrushMaterialComponents); It; ++It)
+			Bounds.Min.X = INT_MAX;
+			Bounds.Min.Y = INT_MAX;
+			Bounds.Max.X = INT_MIN;
+			Bounds.Max.Y = INT_MIN;
+
+			for (ULandscapeComponent* Component : BrushMaterialComponents)
 			{
-				if (*It)
+				if (ensure(Component))
 				{
-					(*It)->GetComponentExtent(X1, Y1, X2, Y2);
+					Component->GetComponentExtent(Bounds.Min.X, Bounds.Min.Y, Bounds.Max.X, Bounds.Max.Y);
 				}
 			}
+
+			// GetComponentExtent returns an inclusive max bound
+			Bounds.Max += FIntPoint(1, 1);
 		}
 
-		// Should not be possible...
-		//check(X1 <= X2 && Y1 <= Y2);
+		FLandscapeBrushData BrushData(Bounds);
 
-		for (int32 Y = Y1; Y <= Y2; Y++)
+		for (int32 Y = Bounds.Min.Y; Y < Bounds.Max.Y; Y++)
 		{
-			for (int32 X = X1; X <= X2; X++)
-			{
-				FIntPoint VertexKey = ALandscape::MakeKey(X, Y);
+			float* Scanline = BrushData.GetDataPtr(FIntPoint(0, Y));
 
+			for (int32 X = Bounds.Min.X; X < Bounds.Max.X; X++)
+			{
 				float PaintAmount = 1.0f;
-				if (EdMode->CurrentTool && EdMode->CurrentTool->GetToolType() != FLandscapeTool::TT_Mask
-					&& EdMode->UISettings->bUseSelectedRegion)
+				if (EdMode->CurrentTool && EdMode->CurrentTool->GetToolType() != ELandscapeToolType::Mask
+					&& EdMode->UISettings->bUseSelectedRegion && LandscapeInfo->SelectedRegion.Num() > 0)
 				{
-					float MaskValue = LandscapeInfo->SelectedRegion.FindRef(VertexKey);
+					float MaskValue = LandscapeInfo->SelectedRegion.FindRef(ALandscape::MakeKey(X, Y));
 					if (EdMode->UISettings->bUseNegativeMask)
 					{
-						MaskValue = 1.f - MaskValue;
+						MaskValue = 1.0f - MaskValue;
 					}
 					PaintAmount *= MaskValue;
 				}
 
 				// Set the brush value for this vertex
-				OutBrush.Add(VertexKey, PaintAmount);
+				Scanline[X] = PaintAmount;
 			}
 		}
 
-		return (X1 <= X2 && Y1 <= Y2);
+		return BrushData;
 	}
 };
 
@@ -494,17 +528,16 @@ class FLandscapeBrushGizmo : public FLandscapeBrush
 	virtual FText GetDisplayName() override { return NSLOCTEXT("UnrealEd", "LandscapeMode_Brush_Gizmo", "Gizmo"); };
 
 protected:
-	FVector2D LastMousePosition;
 	UMaterialInstanceDynamic* BrushMaterial;
 public:
 	FEdModeLandscape* EdMode;
 
 	FLandscapeBrushGizmo(FEdModeLandscape* InEdMode)
-		: BrushMaterial(NULL)
+		: BrushMaterial(nullptr)
 		, EdMode(InEdMode)
 	{
-		UMaterialInstanceConstant* GizmoMaterial = LoadObject<UMaterialInstanceConstant>(NULL, TEXT("/Engine/EditorLandscapeResources/MaskBrushMaterial_Gizmo.MaskBrushMaterial_Gizmo"), NULL, LOAD_None, NULL);
-		BrushMaterial = UMaterialInstanceDynamic::Create(GizmoMaterial, NULL);
+		UMaterialInstanceConstant* GizmoMaterial = LoadObject<UMaterialInstanceConstant>(nullptr, TEXT("/Engine/EditorLandscapeResources/MaskBrushMaterial_Gizmo.MaskBrushMaterial_Gizmo"), nullptr, LOAD_None, nullptr);
+		BrushMaterial = UMaterialInstanceDynamic::Create(GizmoMaterial, nullptr);
 	}
 
 	// FGCObject interface
@@ -513,26 +546,20 @@ public:
 		Collector.AddReferencedObject(BrushMaterial);
 	}
 
-	virtual EBrushType GetBrushType() override { return BT_Gizmo; }
+	virtual ELandscapeBrushType GetBrushType() override { return ELandscapeBrushType::Gizmo; }
 
 	virtual void LeaveBrush() override
 	{
 		for (TSet<ULandscapeComponent*>::TIterator It(BrushMaterialComponents); It; ++It)
 		{
-			if ((*It)->EditToolRenderData != NULL)
+			if ((*It)->EditToolRenderData != nullptr)
 			{
-				(*It)->EditToolRenderData->Update(NULL);
+				(*It)->EditToolRenderData->Update(nullptr);
 			}
 		}
 		BrushMaterialComponents.Empty();
 	}
 
-	virtual void BeginStroke(float LandscapeX, float LandscapeY, FLandscapeTool* CurrentTool) override
-	{
-		FLandscapeBrush::BeginStroke(LandscapeX, LandscapeY, CurrentTool);
-		LastMousePosition = FVector2D(LandscapeX, LandscapeY);
-	}
-
 	virtual void Tick(FEditorViewportClient* ViewportClient, float DeltaTime) override
 	{
 		if (GLandscapeEditRenderMode & ELandscapeEditRenderMode::Gizmo || GLandscapeEditRenderMode & ELandscapeEditRenderMode::Select)
@@ -549,21 +576,22 @@ public:
 					FMatrix WToL = LToW.InverseFast();
 
 					UTexture2D* DataTexture = Gizmo->GizmoTexture;
-					int32 MinX = MAX_int32, MaxX = MIN_int32, MinY = MAX_int32, MaxY = MIN_int32;
+					FIntRect Bounds(MAX_int32, MAX_int32, MIN_int32, MIN_int32);
 					FVector LocalPos[4];
 					//FMatrix WorldToLocal = Proxy->LocalToWorld().Inverse();
 					for (int32 i = 0; i < 4; ++i)
 					{
 						//LocalPos[i] = WorldToLocal.TransformPosition(Gizmo->FrustumVerts[i]);
 						LocalPos[i] = WToL.TransformPosition(Gizmo->FrustumVerts[i]);
-						MinX = FMath::Min(MinX, (int32)LocalPos[i].X);
-						MinY = FMath::Min(MinY, (int32)LocalPos[i].Y);
-						MaxX = FMath::Max(MaxX, (int32)LocalPos[i].X);
-						MaxY = FMath::Max(MaxY, (int32)LocalPos[i].Y);
+						Bounds.Min.X = FMath::Min(Bounds.Min.X, (int32)LocalPos[i].X);
+						Bounds.Min.Y = FMath::Min(Bounds.Min.Y, (int32)LocalPos[i].Y);
+						Bounds.Max.X = FMath::Max(Bounds.Max.X, (int32)LocalPos[i].X);
+						Bounds.Max.Y = FMath::Max(Bounds.Max.Y, (int32)LocalPos[i].Y);
 					}
 
+					// GetComponentsInRegion expects an inclusive max
 					TSet<ULandscapeComponent*> NewComponents;
-					LandscapeInfo->GetComponentsInRegion(MinX, MinY, MaxX, MaxY, NewComponents);
+					LandscapeInfo->GetComponentsInRegion(Bounds.Min.X, Bounds.Min.Y, Bounds.Max.X - 1, Bounds.Max.Y - 1, NewComponents);
 
 					float SquaredScaleXY = FMath::Square(ScaleXY);
 					FLinearColor AlphaScaleBias(
@@ -574,31 +602,31 @@ public:
 						);
 					BrushMaterial->SetVectorParameterValue(FName(TEXT("AlphaScaleBias")), AlphaScaleBias);
 
-					float Angle = (-EdMode->CurrentGizmoActor->GetActorRotation().Euler().Z) * PI / 180.f;
+					float Angle = (-EdMode->CurrentGizmoActor->GetActorRotation().Euler().Z) * PI / 180.0f;
 					FLinearColor LandscapeLocation(EdMode->CurrentGizmoActor->GetActorLocation().X, EdMode->CurrentGizmoActor->GetActorLocation().Y, EdMode->CurrentGizmoActor->GetActorLocation().Z, Angle);
 					BrushMaterial->SetVectorParameterValue(FName(TEXT("LandscapeLocation")), LandscapeLocation);
 					BrushMaterial->SetTextureParameterValue(FName(TEXT("AlphaTexture")), DataTexture);
 
 					// Set brush material for components in new region
-					for (TSet<ULandscapeComponent*>::TIterator It(NewComponents); It; ++It)
+					for (ULandscapeComponent* NewComponent : NewComponents)
 					{
-						if ((*It)->EditToolRenderData != NULL)
+						if (NewComponent->EditToolRenderData != nullptr)
 						{
-							(*It)->EditToolRenderData->UpdateGizmo((Gizmo->DataType != LGT_None) && (GLandscapeEditRenderMode & ELandscapeEditRenderMode::Gizmo) ? BrushMaterial : NULL);
+							NewComponent->EditToolRenderData->UpdateGizmo((Gizmo->DataType != LGT_None) && (GLandscapeEditRenderMode & ELandscapeEditRenderMode::Gizmo) ? BrushMaterial : nullptr);
 						}
 					}
 
 					// Remove the material from any old components that are no longer in the region
 					TSet<ULandscapeComponent*> RemovedComponents = BrushMaterialComponents.Difference(NewComponents);
-					for (TSet<ULandscapeComponent*>::TIterator It(RemovedComponents); It; ++It)
+					for (ULandscapeComponent* RemovedComponent : RemovedComponents)
 					{
-						if ((*It)->EditToolRenderData != NULL)
+						if (RemovedComponent->EditToolRenderData != nullptr)
 						{
-							(*It)->EditToolRenderData->UpdateGizmo(NULL);
+							RemovedComponent->EditToolRenderData->UpdateGizmo(nullptr);
 						}
 					}
 
-					BrushMaterialComponents = NewComponents;
+					BrushMaterialComponents = MoveTemp(NewComponents);
 				}
 			}
 		}
@@ -606,7 +634,6 @@ public:
 
 	virtual void MouseMove(float LandscapeX, float LandscapeY) override
 	{
-		LastMousePosition = FVector2D(LandscapeX, LandscapeY);
 	}
 
 	virtual bool InputKey(FEditorViewportClient* InViewportClient, FViewport* InViewport, FKey InKey, EInputEvent InEvent) override
@@ -615,91 +642,99 @@ public:
 		return bUpdate;
 	}
 
-	virtual bool ApplyBrush(const TArray<FLandscapeToolMousePosition>& MousePositions, TMap<FIntPoint, float>& OutBrush, int32& X1, int32& Y1, int32& X2, int32& Y2) override
+	virtual FLandscapeBrushData ApplyBrush(const TArray<FLandscapeToolMousePosition>& MousePositions) override
 	{
 		// Selection Brush only works for 
 		ALandscapeGizmoActiveActor* Gizmo = EdMode->CurrentGizmoActor.Get();
 		ULandscapeInfo* LandscapeInfo = EdMode->CurrentToolTarget.LandscapeInfo.Get();
 
-		if (Gizmo && Gizmo->GetRootComponent())
+		if (!Gizmo || !Gizmo->GetRootComponent())
 		{
-			Gizmo->TargetLandscapeInfo = LandscapeInfo;
-			float ScaleXY = LandscapeInfo->DrawScale.X;
+			return FLandscapeBrushData();
+		}
 
-			X1 = INT_MAX;
-			Y1 = INT_MAX;
-			X2 = INT_MIN;
-			Y2 = INT_MIN;
+		if (BrushMaterialComponents.Num() == 0)
+		{
+			return FLandscapeBrushData();
+		}
 
-			// Get extent for all components
-			for (TSet<ULandscapeComponent*>::TIterator It(BrushMaterialComponents); It; ++It)
+		Gizmo->TargetLandscapeInfo = LandscapeInfo;
+		float ScaleXY = LandscapeInfo->DrawScale.X;
+
+		// Get extent for all components
+		FIntRect Bounds;
+		Bounds.Min.X = INT_MAX;
+		Bounds.Min.Y = INT_MAX;
+		Bounds.Max.X = INT_MIN;
+		Bounds.Max.Y = INT_MIN;
+
+		for (ULandscapeComponent* Component : BrushMaterialComponents)
+		{
+			if (ensure(Component))
 			{
-				if (*It)
-				{
-					(*It)->GetComponentExtent(X1, Y1, X2, Y2);
-				}
+				Component->GetComponentExtent(Bounds.Min.X, Bounds.Min.Y, Bounds.Max.X, Bounds.Max.Y);
 			}
+		}
 
-			// Should not be possible...
-			//check(X1 <= X2 && Y1 <= Y2);
+		FLandscapeBrushData BrushData(Bounds);
 
-			//FMatrix LandscapeToGizmoLocal = Landscape->LocalToWorld() * Gizmo->WorldToLocal();
-			const float LW = Gizmo->GetWidth() / (2 * ScaleXY);
-			const float LH = Gizmo->GetHeight() / (2 * ScaleXY);
+		//FMatrix LandscapeToGizmoLocal = Landscape->LocalToWorld() * Gizmo->WorldToLocal();
+		const float LW = Gizmo->GetWidth() / (2 * ScaleXY);
+		const float LH = Gizmo->GetHeight() / (2 * ScaleXY);
 
-			FMatrix WToL = LandscapeInfo->GetLandscapeProxy()->LandscapeActorToWorld().ToMatrixWithScale().InverseFast();
-			FVector BaseLocation = WToL.TransformPosition(Gizmo->GetActorLocation());
-			FMatrix LandscapeToGizmoLocal =
-				(FTranslationMatrix(FVector(-LW + 0.5, -LH + 0.5, 0)) * FRotationTranslationMatrix(FRotator(0, Gizmo->GetActorRotation().Yaw, 0), FVector(BaseLocation.X, BaseLocation.Y, 0))).InverseFast();
+		FMatrix WToL = LandscapeInfo->GetLandscapeProxy()->LandscapeActorToWorld().ToMatrixWithScale().InverseFast();
+		FVector BaseLocation = WToL.TransformPosition(Gizmo->GetActorLocation());
+		FMatrix LandscapeToGizmoLocal =
+			(FTranslationMatrix(FVector(-LW + 0.5, -LH + 0.5, 0)) * FRotationTranslationMatrix(FRotator(0, Gizmo->GetActorRotation().Yaw, 0), FVector(BaseLocation.X, BaseLocation.Y, 0))).InverseFast();
 
-			float W = Gizmo->GetWidth() / ScaleXY; //Gizmo->GetWidth() / (Gizmo->DrawScale * Gizmo->DrawScale3D.X);
-			float H = Gizmo->GetHeight() / ScaleXY; //Gizmo->GetHeight() / (Gizmo->DrawScale * Gizmo->DrawScale3D.Y);
+		float W = Gizmo->GetWidth() / ScaleXY; //Gizmo->GetWidth() / (Gizmo->DrawScale * Gizmo->DrawScale3D.X);
+		float H = Gizmo->GetHeight() / ScaleXY; //Gizmo->GetHeight() / (Gizmo->DrawScale * Gizmo->DrawScale3D.Y);
 
-			for (int32 Y = Y1; Y <= Y2; Y++)
+		for (int32 Y = Bounds.Min.Y; Y < Bounds.Max.Y; Y++)
+		{
+			float* Scanline = BrushData.GetDataPtr(FIntPoint(0, Y));
+
+			for (int32 X = Bounds.Min.X; X < Bounds.Max.X; X++)
 			{
-				for (int32 X = X1; X <= X2; X++)
+				FVector GizmoLocal = LandscapeToGizmoLocal.TransformPosition(FVector(X, Y, 0));
+				if (GizmoLocal.X < W && GizmoLocal.X > 0 && GizmoLocal.Y < H && GizmoLocal.Y > 0)
 				{
-					FIntPoint VertexKey = ALandscape::MakeKey(X, Y);
-
-					FVector GizmoLocal = LandscapeToGizmoLocal.TransformPosition(FVector(X, Y, 0));
-					if (GizmoLocal.X < W && GizmoLocal.X > 0 && GizmoLocal.Y < H && GizmoLocal.Y > 0)
+					float PaintAmount = 1.0f;
+					// Transform in 0,0 origin LW radius
+					if (EdMode->UISettings->bSmoothGizmoBrush)
 					{
-						float PaintAmount = 1.f;
-						// Transform in 0,0 origin LW radius
-						if (EdMode->UISettings->bSmoothGizmoBrush)
-						{
-							FVector TransformedLocal(FMath::Abs(GizmoLocal.X - LW), FMath::Abs(GizmoLocal.Y - LH) * (W / H), 0);
-							float FalloffRadius = LW * EdMode->UISettings->BrushFalloff;
-							float SquareRadius = LW - FalloffRadius;
-							float Cos = FMath::Abs(TransformedLocal.X) / TransformedLocal.Size2D();
-							float Sin = FMath::Abs(TransformedLocal.Y) / TransformedLocal.Size2D();
-							float RatioX = FalloffRadius > 0.f ? 1.f - FMath::Clamp((FMath::Abs(TransformedLocal.X) - Cos*SquareRadius) / FalloffRadius, 0.f, 1.f) : 1.f;
-							float RatioY = FalloffRadius > 0.f ? 1.f - FMath::Clamp((FMath::Abs(TransformedLocal.Y) - Sin*SquareRadius) / FalloffRadius, 0.f, 1.f) : 1.f;
-							float Ratio = TransformedLocal.Size2D() > SquareRadius ? RatioX * RatioY : 1.f; //TransformedLocal.X / LW * TransformedLocal.Y / LW;
-							PaintAmount = Ratio*Ratio*(3 - 2 * Ratio); //FMath::Lerp(SquareFalloff, RectFalloff*RectFalloff, Ratio);
-						}
+						FVector TransformedLocal(FMath::Abs(GizmoLocal.X - LW), FMath::Abs(GizmoLocal.Y - LH) * (W / H), 0);
+						float FalloffRadius = LW * EdMode->UISettings->BrushFalloff;
+						float SquareRadius = LW - FalloffRadius;
+						float Cos = FMath::Abs(TransformedLocal.X) / TransformedLocal.Size2D();
+						float Sin = FMath::Abs(TransformedLocal.Y) / TransformedLocal.Size2D();
+						float RatioX = FalloffRadius > 0.0f ? 1.0f - FMath::Clamp((FMath::Abs(TransformedLocal.X) - Cos*SquareRadius) / FalloffRadius, 0.0f, 1.0f) : 1.0f;
+						float RatioY = FalloffRadius > 0.0f ? 1.0f - FMath::Clamp((FMath::Abs(TransformedLocal.Y) - Sin*SquareRadius) / FalloffRadius, 0.0f, 1.0f) : 1.0f;
+						float Ratio = TransformedLocal.Size2D() > SquareRadius ? RatioX * RatioY : 1.0f; //TransformedLocal.X / LW * TransformedLocal.Y / LW;
+						PaintAmount = Ratio*Ratio*(3 - 2 * Ratio); //FMath::Lerp(SquareFalloff, RectFalloff*RectFalloff, Ratio);
+					}
 
-						if (PaintAmount)
+					if (PaintAmount)
+					{
+						if (EdMode->CurrentTool && EdMode->CurrentTool->GetToolType() != ELandscapeToolType::Mask
+							&& EdMode->UISettings->bUseSelectedRegion && LandscapeInfo->SelectedRegion.Num() > 0)
 						{
-							if (EdMode->CurrentTool && EdMode->CurrentTool->GetToolType() != FLandscapeTool::TT_Mask
-								&& EdMode->UISettings->bUseSelectedRegion)
+							float MaskValue = LandscapeInfo->SelectedRegion.FindRef(ALandscape::MakeKey(X, Y));
+							if (EdMode->UISettings->bUseNegativeMask)
 							{
-								float MaskValue = LandscapeInfo->SelectedRegion.FindRef(VertexKey);
-								if (EdMode->UISettings->bUseNegativeMask)
-								{
-									MaskValue = 1.f - MaskValue;
-								}
-								PaintAmount *= MaskValue;
+								MaskValue = 1.0f - MaskValue;
 							}
-
-							// Set the brush value for this vertex
-							OutBrush.Add(VertexKey, PaintAmount);
+							PaintAmount *= MaskValue;
 						}
+
+						// Set the brush value for this vertex
+						Scanline[X] = PaintAmount;
 					}
 				}
 			}
 		}
-		return (X1 <= X2 && Y1 <= Y2);
+
+		return BrushData;
 	}
 };
 
@@ -723,15 +758,15 @@ public:
 	{
 	}
 
-	virtual EBrushType GetBrushType() override { return BT_Splines; }
+	virtual ELandscapeBrushType GetBrushType() override { return ELandscapeBrushType::Splines; }
 
 	virtual void MouseMove(float LandscapeX, float LandscapeY) override
 	{
 	}
 
-	virtual bool ApplyBrush(const TArray<FLandscapeToolMousePosition>& MousePositions, TMap<FIntPoint, float>& OutBrush, int32& X1, int32& Y1, int32& X2, int32& Y2) override
+	virtual FLandscapeBrushData ApplyBrush(const TArray<FLandscapeToolMousePosition>& MousePositions) override
 	{
-		return false;
+		return FLandscapeBrushData();
 	}
 };
 
@@ -755,15 +790,15 @@ public:
 	{
 	}
 
-	virtual EBrushType GetBrushType() override { return BT_Normal; }
+	virtual ELandscapeBrushType GetBrushType() override { return ELandscapeBrushType::Normal; }
 
 	virtual void MouseMove(float LandscapeX, float LandscapeY) override
 	{
 	}
 
-	virtual bool ApplyBrush(const TArray<FLandscapeToolMousePosition>& MousePositions, TMap<FIntPoint, float>& OutBrush, int32& X1, int32& Y1, int32& X2, int32& Y2) override
+	virtual FLandscapeBrushData ApplyBrush(const TArray<FLandscapeToolMousePosition>& MousePositions) override
 	{
-		return false;
+		return FLandscapeBrushData();
 	}
 };
 
@@ -778,15 +813,15 @@ protected:
 
 	virtual float CalculateFalloff(float Distance, float Radius, float Falloff) override
 	{
-		return Distance < Radius ? 1.f :
-			Falloff > 0.f ? FMath::Max<float>(0.f, 1.f - (Distance - Radius) / Falloff) :
-			0.f;
+		return Distance < Radius ? 1.0f :
+			Falloff > 0.0f ? FMath::Max<float>(0.0f, 1.0f - (Distance - Radius) / Falloff) :
+			0.0f;
 	}
 
 public:
 	static FLandscapeBrushCircle_Linear* Create(FEdModeLandscape* InEdMode)
 	{
-		UMaterialInstanceConstant* CircleBrushMaterial_Linear = LoadObject<UMaterialInstanceConstant>(NULL, TEXT("/Engine/EditorLandscapeResources/CircleBrushMaterial_Linear.CircleBrushMaterial_Linear"), NULL, LOAD_None, NULL);
+		UMaterialInstanceConstant* CircleBrushMaterial_Linear = LoadObject<UMaterialInstanceConstant>(nullptr, TEXT("/Engine/EditorLandscapeResources/CircleBrushMaterial_Linear.CircleBrushMaterial_Linear"), nullptr, LOAD_None, nullptr);
 		return new FLandscapeBrushCircle_Linear(InEdMode, CircleBrushMaterial_Linear);
 	}
 	virtual const TCHAR* GetBrushName() override { return TEXT("Circle_Linear"); }
@@ -812,7 +847,7 @@ protected:
 public:
 	static FLandscapeBrushCircle_Smooth* Create(FEdModeLandscape* InEdMode)
 	{
-		UMaterialInstanceConstant* CircleBrushMaterial_Smooth = LoadObject<UMaterialInstanceConstant>(NULL, TEXT("/Engine/EditorLandscapeResources/CircleBrushMaterial_Smooth.CircleBrushMaterial_Smooth"), NULL, LOAD_None, NULL);
+		UMaterialInstanceConstant* CircleBrushMaterial_Smooth = LoadObject<UMaterialInstanceConstant>(nullptr, TEXT("/Engine/EditorLandscapeResources/CircleBrushMaterial_Smooth.CircleBrushMaterial_Smooth"), nullptr, LOAD_None, nullptr);
 		return new FLandscapeBrushCircle_Smooth(InEdMode, CircleBrushMaterial_Smooth);
 	}
 	virtual const TCHAR* GetBrushName() override { return TEXT("Circle_Smooth"); }
@@ -832,22 +867,22 @@ protected:
 	{
 		if (Distance <= Radius)
 		{
-			return 1.f;
+			return 1.0f;
 		}
 
 		if (Distance > Radius + Falloff)
 		{
-			return 0.f;
+			return 0.0f;
 		}
 
 		// Elliptical falloff
-		return FMath::Sqrt(1.f - FMath::Square((Distance - Radius) / Falloff));
+		return FMath::Sqrt(1.0f - FMath::Square((Distance - Radius) / Falloff));
 	}
 
 public:
 	static FLandscapeBrushCircle_Spherical* Create(FEdModeLandscape* InEdMode)
 	{
-		UMaterialInstanceConstant* CircleBrushMaterial_Spherical = LoadObject<UMaterialInstanceConstant>(NULL, TEXT("/Engine/EditorLandscapeResources/CircleBrushMaterial_Spherical.CircleBrushMaterial_Spherical"), NULL, LOAD_None, NULL);
+		UMaterialInstanceConstant* CircleBrushMaterial_Spherical = LoadObject<UMaterialInstanceConstant>(nullptr, TEXT("/Engine/EditorLandscapeResources/CircleBrushMaterial_Spherical.CircleBrushMaterial_Spherical"), nullptr, LOAD_None, nullptr);
 		return new FLandscapeBrushCircle_Spherical(InEdMode, CircleBrushMaterial_Spherical);
 	}
 	virtual const TCHAR* GetBrushName() override { return TEXT("Circle_Spherical"); }
@@ -866,22 +901,22 @@ protected:
 	{
 		if (Distance <= Radius)
 		{
-			return 1.f;
+			return 1.0f;
 		}
 
 		if (Distance > Radius + Falloff)
 		{
-			return 0.f;
+			return 0.0f;
 		}
 
 		// inverse elliptical falloff
-		return 1.f - FMath::Sqrt(1.f - FMath::Square((Falloff + Radius - Distance) / Falloff));
+		return 1.0f - FMath::Sqrt(1.0f - FMath::Square((Falloff + Radius - Distance) / Falloff));
 	}
 
 public:
 	static FLandscapeBrushCircle_Tip* Create(FEdModeLandscape* InEdMode)
 	{
-		UMaterialInstanceConstant* CircleBrushMaterial_Tip = LoadObject<UMaterialInstanceConstant>(NULL, TEXT("/Engine/EditorLandscapeResources/CircleBrushMaterial_Tip.CircleBrushMaterial_Tip"), NULL, LOAD_None, NULL);
+		UMaterialInstanceConstant* CircleBrushMaterial_Tip = LoadObject<UMaterialInstanceConstant>(nullptr, TEXT("/Engine/EditorLandscapeResources/CircleBrushMaterial_Tip.CircleBrushMaterial_Tip"), nullptr, LOAD_None, nullptr);
 		return new FLandscapeBrushCircle_Tip(InEdMode, CircleBrushMaterial_Tip);
 	}
 	virtual const TCHAR* GetBrushName() override { return TEXT("Circle_Tip"); }
@@ -911,10 +946,10 @@ public:
 
 		const uint8* AlphaData = EdMode->UISettings->AlphaTextureData.GetData();
 
-		float Alpha00 = (float)AlphaData[SampleX0 + SampleY0 * SizeX] / 255.f;
-		float Alpha01 = (float)AlphaData[SampleX0 + SampleY1 * SizeX] / 255.f;
-		float Alpha10 = (float)AlphaData[SampleX1 + SampleY0 * SizeX] / 255.f;
-		float Alpha11 = (float)AlphaData[SampleX1 + SampleY1 * SizeX] / 255.f;
+		float Alpha00 = (float)AlphaData[SampleX0 + SampleY0 * SizeX] / 255.0f;
+		float Alpha01 = (float)AlphaData[SampleX0 + SampleY1 * SizeX] / 255.0f;
+		float Alpha10 = (float)AlphaData[SampleX1 + SampleY0 * SizeX] / 255.0f;
+		float Alpha11 = (float)AlphaData[SampleX1 + SampleY1 * SizeX] / 255.0f;
 
 		return FMath::Lerp(
 			FMath::Lerp(Alpha00, Alpha01, FMath::Fractional(SampleX)),
@@ -939,48 +974,61 @@ protected:
 public:
 	static FLandscapeBrushAlphaPattern* Create(FEdModeLandscape* InEdMode)
 	{
-		UMaterialInstanceConstant* PatternBrushMaterial = LoadObject<UMaterialInstanceConstant>(NULL, TEXT("/Engine/EditorLandscapeResources/PatternBrushMaterial_Smooth.PatternBrushMaterial_Smooth"), NULL, LOAD_None, NULL);
+		UMaterialInstanceConstant* PatternBrushMaterial = LoadObject<UMaterialInstanceConstant>(nullptr, TEXT("/Engine/EditorLandscapeResources/PatternBrushMaterial_Smooth.PatternBrushMaterial_Smooth"), nullptr, LOAD_None, nullptr);
 		return new FLandscapeBrushAlphaPattern(InEdMode, PatternBrushMaterial);
 	}
 
-	virtual EBrushType GetBrushType() override { return BT_Alpha; }
+	virtual ELandscapeBrushType GetBrushType() override { return ELandscapeBrushType::Alpha; }
 
-	virtual bool ApplyBrush(const TArray<FLandscapeToolMousePosition>& MousePositions, TMap<FIntPoint, float>& OutBrush, int32& X1, int32& Y1, int32& X2, int32& Y2) override
+	virtual FLandscapeBrushData ApplyBrush(const TArray<FLandscapeToolMousePosition>& MousePositions) override
 	{
 		ULandscapeInfo* LandscapeInfo = EdMode->CurrentToolTarget.LandscapeInfo.Get();
-		float ScaleXY = LandscapeInfo->DrawScale.X;
-
-		float Radius = (1.f - EdMode->UISettings->BrushFalloff) * EdMode->UISettings->BrushRadius / ScaleXY;
-		float Falloff = EdMode->UISettings->BrushFalloff * EdMode->UISettings->BrushRadius / ScaleXY;
+		const float ScaleXY = LandscapeInfo->DrawScale.X;
+		const float TotalRadius = EdMode->UISettings->BrushRadius / ScaleXY;
+		const float Radius = (1.0f - EdMode->UISettings->BrushFalloff) * TotalRadius;
+		const float Falloff = EdMode->UISettings->BrushFalloff * TotalRadius;
 
 		int32 SizeX = EdMode->UISettings->AlphaTextureSizeX;
 		int32 SizeY = EdMode->UISettings->AlphaTextureSizeY;
 
-		X1 = FMath::FloorToInt(LastMousePosition.X - (Radius + Falloff));
-		Y1 = FMath::FloorToInt(LastMousePosition.Y - (Radius + Falloff));
-		X2 = FMath::CeilToInt(LastMousePosition.X + (Radius + Falloff));
-		Y2 = FMath::CeilToInt(LastMousePosition.Y + (Radius + Falloff));
+		FIntRect Bounds;
+		Bounds.Min.X = FMath::FloorToInt(LastMousePosition.X - TotalRadius);
+		Bounds.Min.Y = FMath::FloorToInt(LastMousePosition.Y - TotalRadius);
+		Bounds.Max.X = FMath::CeilToInt( LastMousePosition.X + TotalRadius);
+		Bounds.Max.Y = FMath::CeilToInt( LastMousePosition.Y + TotalRadius);
+
 
-		for (int32 Y = Y1; Y <= Y2; Y++)
+		// Clamp to landscape bounds
+		int32 MinX, MaxX, MinY, MaxY;
+		if (!ensure(LandscapeInfo->GetLandscapeExtent(MinX, MinY, MaxX, MaxY)))
 		{
-			for (int32 X = X1; X <= X2; X++)
-			{
-				FIntPoint VertexKey = ALandscape::MakeKey(X, Y);
+			// Landscape has no components somehow
+			return FLandscapeBrushData();
+		}
+		Bounds.Clip(FIntRect(MinX, MinY, MaxX + 1, MaxY + 1));
 
+		FLandscapeBrushData BrushData(Bounds);
+
+		for (int32 Y = Bounds.Min.Y; Y < Bounds.Max.Y; Y++)
+		{
+			float* Scanline = BrushData.GetDataPtr(FIntPoint(0, Y));
+
+			for (int32 X = Bounds.Min.X; X < Bounds.Max.X; X++)
+			{
 				// Find alphamap sample location
 				float SampleX = (float)X / EdMode->UISettings->AlphaBrushScale + (float)SizeX * EdMode->UISettings->AlphaBrushPanU;
 				float SampleY = (float)Y / EdMode->UISettings->AlphaBrushScale + (float)SizeY * EdMode->UISettings->AlphaBrushPanV;
 
-				float Angle = PI * EdMode->UISettings->AlphaBrushRotation / 180.f;
+				float Angle = PI * EdMode->UISettings->AlphaBrushRotation / 180.0f;
 
 				float ModSampleX = FMath::Fmod(SampleX * FMath::Cos(Angle) - SampleY * FMath::Sin(Angle), (float)SizeX);
 				float ModSampleY = FMath::Fmod(SampleY * FMath::Cos(Angle) + SampleX * FMath::Sin(Angle), (float)SizeY);
 
-				if (ModSampleX < 0.f)
+				if (ModSampleX < 0.0f)
 				{
 					ModSampleX += (float)SizeX;
 				}
-				if (ModSampleY < 0.f)
+				if (ModSampleY < 0.0f)
 				{
 					ModSampleY += (float)SizeY;
 				}
@@ -993,31 +1041,31 @@ public:
 
 				float PaintAmount = CalculateFalloff(MouseDist, Radius, Falloff) * Alpha;
 
-				if (PaintAmount > 0.f)
+				if (PaintAmount > 0.0f)
 				{
-					if (EdMode->CurrentTool && EdMode->CurrentTool->GetToolType() != FLandscapeTool::TT_Mask
-						&& EdMode->UISettings->bUseSelectedRegion)
+					if (EdMode->CurrentTool && EdMode->CurrentTool->GetToolType() != ELandscapeToolType::Mask
+						&& EdMode->UISettings->bUseSelectedRegion && LandscapeInfo->SelectedRegion.Num() > 0)
 					{
-						float MaskValue = LandscapeInfo->SelectedRegion.FindRef(VertexKey);
+						float MaskValue = LandscapeInfo->SelectedRegion.FindRef(ALandscape::MakeKey(X, Y));
 						if (EdMode->UISettings->bUseNegativeMask)
 						{
-							MaskValue = 1.f - MaskValue;
+							MaskValue = 1.0f - MaskValue;
 						}
 						PaintAmount *= MaskValue;
 					}
 					// Set the brush value for this vertex
-					OutBrush.Add(VertexKey, PaintAmount);
+					Scanline[X] = PaintAmount;
 				}
 			}
 		}
-		return (X1 <= X2 && Y1 <= Y2);
+		return BrushData;
 	}
 
 	virtual void Tick(FEditorViewportClient* ViewportClient, float DeltaTime) override
 	{
 		FLandscapeBrushCircle::Tick(ViewportClient, DeltaTime);
 
-		ALandscapeProxy* Proxy = EdMode->CurrentToolTarget.LandscapeInfo.IsValid() ? EdMode->CurrentToolTarget.LandscapeInfo->GetLandscapeProxy() : NULL;
+		ALandscapeProxy* Proxy = EdMode->CurrentToolTarget.LandscapeInfo.IsValid() ? EdMode->CurrentToolTarget.LandscapeInfo->GetLandscapeProxy() : nullptr;
 		if (Proxy)
 		{
 			const float ScaleXY = EdMode->CurrentToolTarget.LandscapeInfo->DrawScale.X;
@@ -1037,12 +1085,13 @@ public:
 			int32 Channel = EdMode->UISettings->AlphaTextureChannel;
 			FLinearColor AlphaTextureMask(Channel == 0 ? 1 : 0, Channel == 1 ? 1 : 0, Channel == 2 ? 1 : 0, Channel == 3 ? 1 : 0);
 
-			for (auto It = BrushMaterialInstanceMap.CreateConstIterator(); It; ++It)
+			for (const auto& BrushMaterialInstancePair : BrushMaterialInstanceMap)
 			{
-				It.Value()->SetVectorParameterValue(FName(TEXT("AlphaScaleBias")), AlphaScaleBias);
-				It.Value()->SetVectorParameterValue(FName(TEXT("LandscapeLocation")), LandscapeLocation);
-				It.Value()->SetVectorParameterValue(FName(TEXT("AlphaTextureMask")), AlphaTextureMask);
-				It.Value()->SetTextureParameterValue(FName(TEXT("AlphaTexture")), EdMode->UISettings->AlphaTexture);
+				UMaterialInstanceDynamic* const MaterialInstance = BrushMaterialInstancePair.Value;
+				MaterialInstance->SetVectorParameterValue(FName(TEXT("AlphaScaleBias")),    AlphaScaleBias);
+				MaterialInstance->SetVectorParameterValue(FName(TEXT("LandscapeLocation")), LandscapeLocation);
+				MaterialInstance->SetVectorParameterValue(FName(TEXT("AlphaTextureMask")),  AlphaTextureMask);
+				MaterialInstance->SetTextureParameterValue(FName(TEXT("AlphaTexture")),     EdMode->UISettings->AlphaTexture);
 			}
 		}
 	}
@@ -1064,8 +1113,8 @@ class FLandscapeBrushAlpha : public FLandscapeBrushAlphaBase
 protected:
 	FLandscapeBrushAlpha(FEdModeLandscape* InEdMode, UMaterialInterface* InBrushMaterial)
 		: FLandscapeBrushAlphaBase(InEdMode, InBrushMaterial)
-		, LastMouseAngle(0.f)
-		, OldMousePosition(0.f, 0.f)
+		, LastMouseAngle(0.0f)
+		, OldMousePosition(0.0f, 0.0f)
 		, LastMouseSampleTime(FPlatformTime::Seconds())
 	{
 	}
@@ -1073,22 +1122,19 @@ protected:
 public:
 	static FLandscapeBrushAlpha* Create(FEdModeLandscape* InEdMode)
 	{
-		UMaterialInstanceConstant* AlphaBrushMaterial = LoadObject<UMaterialInstanceConstant>(NULL, TEXT("/Engine/EditorLandscapeResources/AlphaBrushMaterial_Smooth.AlphaBrushMaterial_Smooth"), NULL, LOAD_None, NULL);
+		UMaterialInstanceConstant* AlphaBrushMaterial = LoadObject<UMaterialInstanceConstant>(nullptr, TEXT("/Engine/EditorLandscapeResources/AlphaBrushMaterial_Smooth.AlphaBrushMaterial_Smooth"), nullptr, LOAD_None, nullptr);
 		return new FLandscapeBrushAlpha(InEdMode, AlphaBrushMaterial);
 	}
 
-	virtual bool ApplyBrush(const TArray<FLandscapeToolMousePosition>& MousePositions, TMap<FIntPoint, float>& OutBrush, int32& X1, int32& Y1, int32& X2, int32& Y2) override
+	virtual FLandscapeBrushData ApplyBrush(const TArray<FLandscapeToolMousePosition>& MousePositions) override
 	{
 		ULandscapeInfo* LandscapeInfo = EdMode->CurrentToolTarget.LandscapeInfo.Get();
 		if (EdMode->UISettings->bAlphaBrushAutoRotate && OldMousePosition.IsZero())
 		{
-			X1 = FMath::FloorToInt(LastMousePosition.X);
-			Y1 = FMath::FloorToInt(LastMousePosition.Y);
-			X2 = FMath::CeilToInt(LastMousePosition.X);
-			Y2 = FMath::CeilToInt(LastMousePosition.Y);
 			OldMousePosition = LastMousePosition;
-			LastMouseAngle = 0.f;
+			LastMouseAngle = 0.0f;
 			LastMouseSampleTime = FPlatformTime::Seconds();
+			return FLandscapeBrushData();
 		}
 		else
 		{
@@ -1096,18 +1142,33 @@ public:
 			float Radius = EdMode->UISettings->BrushRadius / ScaleXY;
 			int32 SizeX = EdMode->UISettings->AlphaTextureSizeX;
 			int32 SizeY = EdMode->UISettings->AlphaTextureSizeY;
-			float MaxSize = 2.f * FMath::Sqrt(FMath::Square(Radius) / 2.f);
+			float MaxSize = 2.0f * FMath::Sqrt(FMath::Square(Radius) / 2.0f);
 			float AlphaBrushScale = MaxSize / (float)FMath::Max<int32>(SizeX, SizeY);
 			const float BrushAngle = EdMode->UISettings->bAlphaBrushAutoRotate ? LastMouseAngle : FMath::DegreesToRadians(EdMode->UISettings->AlphaBrushRotation);
 
-			X1 = FMath::FloorToInt(LastMousePosition.X - Radius);
-			Y1 = FMath::FloorToInt(LastMousePosition.Y - Radius);
-			X2 = FMath::CeilToInt(LastMousePosition.X + Radius);
-			Y2 = FMath::CeilToInt(LastMousePosition.Y + Radius);
+			FIntRect Bounds;
+			Bounds.Min.X = FMath::FloorToInt(LastMousePosition.X - Radius);
+			Bounds.Min.Y = FMath::FloorToInt(LastMousePosition.Y - Radius);
+			Bounds.Max.X = FMath::CeilToInt( LastMousePosition.X + Radius);
+			Bounds.Max.Y = FMath::CeilToInt( LastMousePosition.Y + Radius);
 
-			for (int32 Y = Y1; Y <= Y2; Y++)
+
+			// Clamp to landscape bounds
+			int32 MinX, MaxX, MinY, MaxY;
+			if (!ensure(LandscapeInfo->GetLandscapeExtent(MinX, MinY, MaxX, MaxY)))
+			{
+				// Landscape has no components somehow
+				return FLandscapeBrushData();
+			}
+			Bounds.Clip(FIntRect(MinX, MinY, MaxX + 1, MaxY + 1));
+
+			FLandscapeBrushData BrushData(Bounds);
+
+			for (int32 Y = Bounds.Min.Y; Y < Bounds.Max.Y; Y++)
 			{
-				for (int32 X = X1; X <= X2; X++)
+				float* Scanline = BrushData.GetDataPtr(FIntPoint(0, Y));
+
+				for (int32 X = Bounds.Min.X; X < Bounds.Max.X; X++)
 				{
 					// Find alphamap sample location
 					float ScaleSampleX = ((float)X - LastMousePosition.X) / AlphaBrushScale;
@@ -1126,30 +1187,30 @@ public:
 						// Sample the alpha texture
 						float Alpha = GetAlphaSample(SampleX, SampleY);
 
-						if (Alpha > 0.f)
+						if (Alpha > 0.0f)
 						{
 							// Set the brush value for this vertex
 							FIntPoint VertexKey = ALandscape::MakeKey(X, Y);
 
-							if (EdMode->CurrentTool && EdMode->CurrentTool->GetToolType() != FLandscapeTool::TT_Mask
-								&& EdMode->UISettings->bUseSelectedRegion)
+							if (EdMode->CurrentTool && EdMode->CurrentTool->GetToolType() != ELandscapeToolType::Mask
+								&& EdMode->UISettings->bUseSelectedRegion && LandscapeInfo->SelectedRegion.Num() > 0)
 							{
-								float MaskValue = LandscapeInfo->SelectedRegion.FindRef(VertexKey);
+								float MaskValue = LandscapeInfo->SelectedRegion.FindRef(ALandscape::MakeKey(X, Y));
 								if (EdMode->UISettings->bUseNegativeMask)
 								{
-									MaskValue = 1.f - MaskValue;
+									MaskValue = 1.0f - MaskValue;
 								}
 								Alpha *= MaskValue;
 							}
 
-							OutBrush.Add(VertexKey, Alpha);
+							Scanline[X] = Alpha;
 						}
 					}
 				}
 			}
-		}
 
-		return (X1 <= X2 && Y1 <= Y2);
+			return BrushData;
+		}
 	}
 
 	virtual void MouseMove(float LandscapeX, float LandscapeY) override
@@ -1165,7 +1226,7 @@ public:
 				double SampleTime = FPlatformTime::Seconds();
 				float DeltaTime = (float)(SampleTime - LastMouseSampleTime);
 				FVector2D MouseDirection = MouseDelta.GetSafeNormal();
-				float MouseAngle = FMath::Lerp(LastMouseAngle, FMath::Atan2(-MouseDirection.Y, MouseDirection.X), FMath::Min<float>(10.f * DeltaTime, 1.f));		// lerp over 100ms
+				float MouseAngle = FMath::Lerp(LastMouseAngle, FMath::Atan2(-MouseDirection.Y, MouseDirection.X), FMath::Min<float>(10.0f * DeltaTime, 1.0f));		// lerp over 100ms
 				LastMouseAngle = MouseAngle;
 				LastMouseSampleTime = SampleTime;
 				OldMousePosition = LastMousePosition;
@@ -1185,24 +1246,25 @@ public:
 			int32 SizeX = EdMode->UISettings->AlphaTextureSizeX;
 			int32 SizeY = EdMode->UISettings->AlphaTextureSizeY;
 			float Radius = EdMode->UISettings->BrushRadius / ScaleXY;
-			float MaxSize = 2.f * FMath::Sqrt(FMath::Square(Radius) / 2.f);
+			float MaxSize = 2.0f * FMath::Sqrt(FMath::Square(Radius) / 2.0f);
 			float AlphaBrushScale = MaxSize / (float)FMath::Max<int32>(SizeX, SizeY);
 
 			FLinearColor BrushScaleRot(
 				1.0f / (AlphaBrushScale * SizeX),
 				1.0f / (AlphaBrushScale * SizeY),
-				0.f,
+				0.0f,
 				EdMode->UISettings->bAlphaBrushAutoRotate ? LastMouseAngle : FMath::DegreesToRadians(EdMode->UISettings->AlphaBrushRotation)
 				);
 
 			int32 Channel = EdMode->UISettings->AlphaTextureChannel;
 			FLinearColor AlphaTextureMask(Channel == 0 ? 1 : 0, Channel == 1 ? 1 : 0, Channel == 2 ? 1 : 0, Channel == 3 ? 1 : 0);
 
-			for (auto It = BrushMaterialInstanceMap.CreateConstIterator(); It; ++It)
+			for (const auto& BrushMaterialInstancePair : BrushMaterialInstanceMap)
 			{
-				It.Value()->SetVectorParameterValue(FName(TEXT("BrushScaleRot")), BrushScaleRot);
-				It.Value()->SetVectorParameterValue(FName(TEXT("AlphaTextureMask")), AlphaTextureMask);
-				It.Value()->SetTextureParameterValue(FName(TEXT("AlphaTexture")), EdMode->UISettings->AlphaTexture);
+				UMaterialInstanceDynamic* const MaterialInstance = BrushMaterialInstancePair.Value;
+				MaterialInstance->SetVectorParameterValue(FName(TEXT("BrushScaleRot")), BrushScaleRot);
+				MaterialInstance->SetVectorParameterValue(FName(TEXT("AlphaTextureMask")), AlphaTextureMask);
+				MaterialInstance->SetTextureParameterValue(FName(TEXT("AlphaTexture")), EdMode->UISettings->AlphaTexture);
 			}
 		}
 	}
diff --git a/Engine/Source/Editor/LandscapeEditor/Private/LandscapeEdModeComponentTools.cpp b/Engine/Source/Editor/LandscapeEditor/Private/LandscapeEdModeComponentTools.cpp
--- a/Engine/Source/Editor/LandscapeEditor/Private/LandscapeEdModeComponentTools.cpp
+++ b/Engine/Source/Editor/LandscapeEditor/Private/LandscapeEdModeComponentTools.cpp
@@ -37,33 +37,36 @@ public:
 		if (LandscapeInfo)
 		{
 			LandscapeInfo->Modify();
-			// Get list of verts to update
-			TMap<FIntPoint, float> BrushInfo;
-			int32 X1, Y1, X2, Y2;
-			if (!Brush->ApplyBrush(MousePositions, BrushInfo, X1, Y1, X2, Y2))
-			{
-				return;
-			}
 
 			// Invert when holding Shift
 			bool bInvert = MousePositions[MousePositions.Num() - 1].bShiftDown;
 
-			if (Brush->GetBrushType() == FLandscapeBrush::BT_Component)
+			if (Brush->GetBrushType() == ELandscapeBrushType::Component)
 			{
-				// Todo hold selection... static?
+				// TODO - only retrieve bounds as we don't need the data
+				const FLandscapeBrushData BrushInfo = Brush->ApplyBrush(MousePositions);
+				if (!BrushInfo)
+				{
+					return;
+				}
+
+				int32 X1, Y1, X2, Y2;
+				BrushInfo.GetInclusiveBounds(X1, Y1, X2, Y2);
+
+				// Shrink bounds by 1,1 to avoid GetComponentsInRegion picking up extra components on all sides due to the overlap between components
 				TSet<ULandscapeComponent*> NewComponents;
 				LandscapeInfo->GetComponentsInRegion(X1 + 1, Y1 + 1, X2 - 1, Y2 - 1, NewComponents);
 
 				if (!bInitializedComponentInvert)
 				{
 					// Get the component under the mouse location. Copied from FLandscapeBrushComponent::ApplyBrush()
-					const int32 MouseX = FMath::RoundToInt(MousePositions[0].PositionX);
-					const int32 MouseY = FMath::RoundToInt(MousePositions[0].PositionY);
-					const int32 MouseComponentIndexX = (MouseX >= 0.f) ? MouseX / LandscapeInfo->ComponentSizeQuads : (MouseX + 1) / LandscapeInfo->ComponentSizeQuads - 1;
-					const int32 MouseComponentIndexY = (MouseY >= 0.f) ? MouseY / LandscapeInfo->ComponentSizeQuads : (MouseY + 1) / LandscapeInfo->ComponentSizeQuads - 1;
+					const float MouseX = MousePositions[0].Position.X;
+					const float MouseY = MousePositions[0].Position.Y;
+					const int32 MouseComponentIndexX = (MouseX >= 0.0f) ? FMath::FloorToInt(MouseX / LandscapeInfo->ComponentSizeQuads) : FMath::CeilToInt(MouseX / LandscapeInfo->ComponentSizeQuads);
+					const int32 MouseComponentIndexY = (MouseY >= 0.0f) ? FMath::FloorToInt(MouseY / LandscapeInfo->ComponentSizeQuads) : FMath::CeilToInt(MouseY / LandscapeInfo->ComponentSizeQuads);
 					ULandscapeComponent* MouseComponent = LandscapeInfo->XYtoComponentMap.FindRef(FIntPoint(MouseComponentIndexX, MouseComponentIndexY));
 
-					if (MouseComponent != NULL)
+					if (MouseComponent != nullptr)
 					{
 						bComponentInvert = LandscapeInfo->GetSelectedComponents().Contains(MouseComponent);
 					}
@@ -102,48 +105,60 @@ public:
 			}
 			else // Select various shape regions
 			{
-				X1 -= 1;
-				Y1 -= 1;
-				X2 += 1;
-				Y2 += 1;
+				const FLandscapeBrushData BrushInfo = Brush->ApplyBrush(MousePositions);
+				if (!BrushInfo)
+				{
+					return;
+				}
+
+				int32 X1, Y1, X2, Y2;
+				BrushInfo.GetInclusiveBounds(X1, Y1, X2, Y2);
 
 				// Tablet pressure
-				float Pressure = ViewportClient->Viewport->IsPenActive() ? ViewportClient->Viewport->GetTabletPressure() : 1.f;
+				float Pressure = ViewportClient->Viewport->IsPenActive() ? ViewportClient->Viewport->GetTabletPressure() : 1.0f;
 
 				Cache.CacheData(X1, Y1, X2, Y2);
 				TArray<uint8> Data;
 				Cache.GetCachedData(X1, Y1, X2, Y2, Data);
 
 				TSet<ULandscapeComponent*> NewComponents;
-				// Remove invalid regions
 				LandscapeInfo->GetComponentsInRegion(X1, Y1, X2, Y2, NewComponents);
 				LandscapeInfo->UpdateSelectedComponents(NewComponents, false);
 
-				for (auto It = BrushInfo.CreateIterator(); It; ++It)
+				for (int32 Y = BrushInfo.GetBounds().Min.Y; Y < BrushInfo.GetBounds().Max.Y; Y++)
 				{
-					FIntPoint Key = It.Key();
-					if (It.Value() > 0.f && LandscapeInfo->IsValidPosition(Key.X, Key.Y))
+					const float* BrushScanline = BrushInfo.GetDataPtr(FIntPoint(0, Y));
+					uint8* DataScanline = Data.GetData() + (Y - Y1) * (X2 - X1 + 1) + (0 - X1);
+
+					for (int32 X = BrushInfo.GetBounds().Min.X; X < BrushInfo.GetBounds().Max.X; X++)
 					{
-						float PaintValue = It.Value() * UISettings->ToolStrength * Pressure;
-						float Value = LandscapeInfo->SelectedRegion.FindRef(Key);
-						if (bInvert)
+						const FIntPoint Key = ALandscape::MakeKey(X, Y);
+						const float BrushValue = BrushScanline[X];
+
+						if (BrushValue > 0.0f && LandscapeInfo->IsValidPosition(X, Y))
 						{
-							if (Value - PaintValue > 0.f)
+							float PaintValue = BrushValue * UISettings->ToolStrength * Pressure;
+							float Value = DataScanline[X];
+							checkSlow(Value == LandscapeInfo->SelectedRegion.FindRef(Key));
+							if (bInvert)
+							{
+								Value = FMath::Max(Value - PaintValue, 0.0f);
+							}
+							else
+							{
+								Value = FMath::Min(Value + PaintValue, 1.0f);
+							}
+							if (Value > 0.0f)
 							{
-								LandscapeInfo->SelectedRegion.Add(Key, FMath::Max(Value - PaintValue, 0.f));
+								LandscapeInfo->SelectedRegion.Add(Key, Value);
 							}
 							else
 							{
 								LandscapeInfo->SelectedRegion.Remove(Key);
 							}
 
+							DataScanline[X] = FMath::Clamp<int32>(FMath::RoundToInt(Value * 255), 0, 255);
 						}
-						else
-						{
-							LandscapeInfo->SelectedRegion.Add(Key, FMath::Min(Value + PaintValue, 1.0f));
-						}
-
-						Data[(Key.X - X1) + (Key.Y - Y1)*(1 + X2 - X1)] = FMath::Clamp<int32>(FMath::RoundToInt(LandscapeInfo->SelectedRegion.FindRef(Key) * 255), 0, 255);
 					}
 				}
 
@@ -175,7 +190,7 @@ public:
 	virtual void SetEditRenderType() override { GLandscapeEditRenderMode = ELandscapeEditRenderMode::SelectComponent | (GLandscapeEditRenderMode & ELandscapeEditRenderMode::BitMaskForMask); }
 	virtual bool SupportsMask() override { return false; }
 
-	virtual FLandscapeTool::EToolType GetToolType() override { return this->TT_Mask; }
+	virtual ELandscapeToolType GetToolType() override { return ELandscapeToolType::Mask; }
 };
 
 template<class TStrokeClass>
@@ -210,46 +225,39 @@ public:
 		{
 			LandscapeInfo->Modify();
 			// Get list of verts to update
-			TMap<FIntPoint, float> BrushInfo;
-			int32 X1, Y1, X2, Y2;
-			if (!Brush->ApplyBrush(MousePositions, BrushInfo, X1, Y1, X2, Y2))
+			FLandscapeBrushData BrushInfo = Brush->ApplyBrush(MousePositions);
+			if (!BrushInfo)
 			{
 				return;
 			}
 
+			int32 X1, Y1, X2, Y2;
+			BrushInfo.GetInclusiveBounds(X1, Y1, X2, Y2);
+
 			// Invert when holding Shift
 			bool bInvert = MousePositions[MousePositions.Num() - 1].bShiftDown;
 
-			X1 -= 1;
-			Y1 -= 1;
-			X2 += 1;
-			Y2 += 1;
-
 			// Tablet pressure
-			float Pressure = ViewportClient->Viewport->IsPenActive() ? ViewportClient->Viewport->GetTabletPressure() : 1.f;
+			float Pressure = ViewportClient->Viewport->IsPenActive() ? ViewportClient->Viewport->GetTabletPressure() : 1.0f;
 
 			Cache.CacheData(X1, Y1, X2, Y2);
 			TArray<uint8> Data;
 			Cache.GetCachedData(X1, Y1, X2, Y2, Data);
 
-			for (auto It = BrushInfo.CreateConstIterator(); It; ++It)
+			for (int32 Y = BrushInfo.GetBounds().Min.Y; Y < BrushInfo.GetBounds().Max.Y; Y++)
 			{
-				int32 X, Y;
-				ALandscape::UnpackKey(It.Key(), X, Y);
+				const float* BrushScanline = BrushInfo.GetDataPtr(FIntPoint(0, Y));
+				uint8* DataScanline = Data.GetData() + (Y - Y1) * (X2 - X1 + 1) + (0 - X1);
 
-				if (It.Value() > 0.f)
+				for (int32 X = BrushInfo.GetBounds().Min.X; X < BrushInfo.GetBounds().Max.X; X++)
 				{
-					float Value = Data[(X - X1) + (Y - Y1)*(1 + X2 - X1)];
-					if (bInvert)
-					{
-						Value = 0;
-					}
-					else
+					const float BrushValue = BrushScanline[X];
+
+					if (BrushValue > 0.0f)
 					{
-						Value = 255; // Just on and off for visibility, for masking...
+						uint8 Value = bInvert ? 0 : 255; // Just on and off for visibility, for masking...
+						DataScanline[X] = Value;
 					}
-
-					Data[(X - X1) + (Y - Y1)*(1 + X2 - X1)] = FMath::Clamp<int32>(FMath::RoundToInt(Value), 0, 255);
 				}
 			}
 
@@ -298,7 +306,7 @@ public:
 
 	void Apply(FEditorViewportClient* ViewportClient, FLandscapeBrush* Brush, const ULandscapeEditorObject* UISettings, const TArray<FLandscapeToolMousePosition>& MousePositions)
 	{
-		ALandscape* Landscape = LandscapeInfo ? LandscapeInfo->LandscapeActor.Get() : NULL;
+		ALandscape* Landscape = LandscapeInfo ? LandscapeInfo->LandscapeActor.Get() : nullptr;
 
 		if (Landscape)
 		{
@@ -334,12 +342,17 @@ public:
 			if (!SelectedComponents.Num())
 			{
 				// Get list of verts to update
-				TMap<FIntPoint, float> BrushInfo;
-				int32 X1, Y1, X2, Y2;
-				if (!Brush->ApplyBrush(MousePositions, BrushInfo, X1, Y1, X2, Y2))
+				// TODO - only retrieve bounds as we don't need the data
+				FLandscapeBrushData BrushInfo = Brush->ApplyBrush(MousePositions);
+				if (!BrushInfo)
 				{
 					return;
 				}
+
+				int32 X1, Y1, X2, Y2;
+				BrushInfo.GetInclusiveBounds(X1, Y1, X2, Y2);
+
+				// Shrink bounds by 1,1 to avoid GetComponentsInRegion picking up extra components on all sides due to the overlap between components
 				LandscapeInfo->GetComponentsInRegion(X1 + 1, Y1 + 1, X2 - 1, Y2 - 1, SelectedComponents);
 				bBrush = true;
 			}
@@ -350,9 +363,9 @@ public:
 			if (SelectedComponents.Num())
 			{
 				bool bIsAllCurrentLevel = true;
-				for (TSet<ULandscapeComponent*>::TIterator It(SelectedComponents); It; ++It)
+				for (ULandscapeComponent* Component : SelectedComponents)
 				{
-					if ((*It)->GetLandscapeProxy()->GetLevel() != World->GetCurrentLevel())
+					if (Component->GetLandscapeProxy()->GetLevel() != World->GetCurrentLevel())
 					{
 						bIsAllCurrentLevel = false;
 					}
@@ -369,14 +382,13 @@ public:
 					return;
 				}
 
-				for (TSet<ULandscapeComponent*>::TIterator It(SelectedComponents); It; ++It)
+				for (ULandscapeComponent* Component : SelectedComponents)
 				{
-					UMaterialInterface* LandscapeMaterial = (*It)->GetLandscapeMaterial();
-					if (LandscapeMaterial && LandscapeMaterial->GetOutermost() == (*It)->GetOutermost())
+					UMaterialInterface* LandscapeMaterial = Component->GetLandscapeMaterial();
+					if (LandscapeMaterial && LandscapeMaterial->GetOutermost() == Component->GetOutermost())
 					{
-						ULandscapeComponent* Comp = *It;
 						RenameObjects.AddUnique(LandscapeMaterial);
-						MsgBoxList += Comp->GetName() + TEXT("'s ") + LandscapeMaterial->GetPathName();
+						MsgBoxList += Component->GetName() + TEXT("'s ") + LandscapeMaterial->GetPathName();
 						MsgBoxList += FString::Printf(TEXT("\n"));
 						//It.RemoveCurrent();
 					}
@@ -416,16 +428,15 @@ public:
 				int32 ComponentSizeVerts = Landscape->NumSubsections * (Landscape->SubsectionSizeQuads + 1);
 				int32 NeedHeightmapSize = 1 << FMath::CeilLogTwo(ComponentSizeVerts);
 
-				for (TSet<ULandscapeComponent*>::TConstIterator It(SelectedComponents); It; ++It)
+				for (ULandscapeComponent* Component : SelectedComponents)
 				{
-					ULandscapeComponent* Comp = *It;
-					SelectProxies.Add(Comp->GetLandscapeProxy());
-					if (Comp->GetLandscapeProxy()->GetOuter() != World->GetCurrentLevel())
+					SelectProxies.Add(Component->GetLandscapeProxy());
+					if (Component->GetLandscapeProxy()->GetOuter() != World->GetCurrentLevel())
 					{
-						TargetSelectedComponents.Add(Comp);
+						TargetSelectedComponents.Add(Component);
 					}
 
-					ULandscapeHeightfieldCollisionComponent* CollisionComp = Comp->CollisionComponent.Get();
+					ULandscapeHeightfieldCollisionComponent* CollisionComp = Component->CollisionComponent.Get();
 					SelectProxies.Add(CollisionComp->GetLandscapeProxy());
 					if (CollisionComp->GetLandscapeProxy()->GetOuter() != World->GetCurrentLevel())
 					{
@@ -436,21 +447,19 @@ public:
 				int32 TotalProgress = TargetSelectedComponents.Num() * TargetSelectedCollisionComponents.Num();
 
 				// Check which ones are need for height map change
-				for (TSet<ULandscapeComponent*>::TConstIterator It(TargetSelectedComponents); It; ++It)
+				for (ULandscapeComponent* Component : TargetSelectedComponents)
 				{
-					ULandscapeComponent* Comp = *It;
-					Comp->Modify();
-					OldTextureSet.Add(Comp->HeightmapTexture);
+					Component->Modify();
+					OldTextureSet.Add(Component->HeightmapTexture);
 				}
 
 				// Need to split all the component which share Heightmap with selected components
 				// Search neighbor only
-				for (TSet<ULandscapeComponent*>::TConstIterator It(TargetSelectedComponents); It; ++It)
+				for (ULandscapeComponent* Component : TargetSelectedComponents)
 				{
-					ULandscapeComponent* Comp = *It;
-					int32 SearchX = Comp->HeightmapTexture->Source.GetSizeX() / NeedHeightmapSize;
-					int32 SearchY = Comp->HeightmapTexture->Source.GetSizeY() / NeedHeightmapSize;
-					FIntPoint ComponentBase = Comp->GetSectionBase() / Comp->ComponentSizeQuads;
+					int32 SearchX = Component->HeightmapTexture->Source.GetSizeX() / NeedHeightmapSize;
+					int32 SearchY = Component->HeightmapTexture->Source.GetSizeY() / NeedHeightmapSize;
+					FIntPoint ComponentBase = Component->GetSectionBase() / Component->ComponentSizeQuads;
 
 					for (int32 Y = 0; Y < SearchY; ++Y)
 					{
@@ -462,12 +471,12 @@ public:
 								int32 XDir = (Dir >> 1) ? 1 : -1;
 								int32 YDir = (Dir % 2) ? 1 : -1;
 								ULandscapeComponent* Neighbor = LandscapeInfo->XYtoComponentMap.FindRef(ComponentBase + FIntPoint(XDir*X, YDir*Y));
-								if (Neighbor && Neighbor->HeightmapTexture == Comp->HeightmapTexture && !HeightmapUpdateComponents.Contains(Neighbor))
+								if (Neighbor && Neighbor->HeightmapTexture == Component->HeightmapTexture && !HeightmapUpdateComponents.Contains(Neighbor))
 								{
 									Neighbor->Modify();
 									if (!TargetSelectedComponents.Contains(Neighbor))
 									{
-										Neighbor->HeightmapScaleBias.X = -1.f; // just mark this component is for original level, not current level
+										Neighbor->HeightmapScaleBias.X = -1.0f; // just mark this component is for original level, not current level
 									}
 									HeightmapUpdateComponents.Add(Neighbor);
 								}
@@ -477,19 +486,18 @@ public:
 				}
 
 				// Changing Heightmap format for selected components
-				for (TSet<ULandscapeComponent*>::TConstIterator It(HeightmapUpdateComponents); It; ++It)
+				for (ULandscapeComponent* Component : HeightmapUpdateComponents)
 				{
-					ULandscapeComponent* Comp = *It;
-					ALandscape::SplitHeightmap(Comp, (Comp->HeightmapScaleBias.X > 0.f));
+					ALandscape::SplitHeightmap(Component, (Component->HeightmapScaleBias.X > 0.0f));
 				}
 
 				// Delete if it is no referenced textures...
-				for (TSet<UTexture2D*>::TIterator It(OldTextureSet); It; ++It)
+				for (UTexture2D* Texture : OldTextureSet)
 				{
-					(*It)->SetFlags(RF_Transactional);
-					(*It)->Modify();
-					(*It)->MarkPackageDirty();
-					(*It)->ClearFlags(RF_Standalone);
+					Texture->SetFlags(RF_Transactional);
+					Texture->Modify();
+					Texture->MarkPackageDirty();
+					Texture->ClearFlags(RF_Standalone);
 				}
 
 				ALandscapeProxy* LandscapeProxy = LandscapeInfo->GetCurrentLevelLandscapeProxy(false);
@@ -512,9 +520,9 @@ public:
 					}
 				}
 
-				for (TSet<ALandscapeProxy*>::TIterator It(SelectProxies); It; ++It)
+				for (ALandscapeProxy* Proxy : SelectProxies)
 				{
-					(*It)->Modify();
+					Proxy->Modify();
 				}
 
 				LandscapeProxy->Modify();
@@ -523,10 +531,9 @@ public:
 				// Change Weight maps...
 				{
 					FLandscapeEditDataInterface LandscapeEdit(LandscapeInfo);
-					for (TSet<ULandscapeComponent*>::TConstIterator ComponentIt(TargetSelectedComponents); ComponentIt; ++ComponentIt)
+					for (ULandscapeComponent* Component : TargetSelectedComponents)
 					{
-						ULandscapeComponent* Comp = *ComponentIt;
-						int32 TotalNeededChannels = Comp->WeightmapLayerAllocations.Num();
+						int32 TotalNeededChannels = Component->WeightmapLayerAllocations.Num();
 						int32 CurrentLayer = 0;
 						TArray<UTexture2D*> NewWeightmapTextures;
 
@@ -536,8 +543,8 @@ public:
 						{
 							// UE_LOG(LogLandscape, Log, TEXT("Still need %d channels"), TotalNeededChannels);
 
-							UTexture2D* CurrentWeightmapTexture = NULL;
-							FLandscapeWeightmapUsage* CurrentWeightmapUsage = NULL;
+							UTexture2D* CurrentWeightmapTexture = nullptr;
+							FLandscapeWeightmapUsage* CurrentWeightmapUsage = nullptr;
 
 							if (TotalNeededChannels < 4)
 							{
@@ -545,20 +552,20 @@ public:
 
 								// see if we can find a suitable existing weightmap texture with sufficient channels
 								int32 BestDistanceSquared = MAX_int32;
-								for (TMap<UTexture2D*, struct FLandscapeWeightmapUsage>::TIterator WeightmapIt(LandscapeProxy->WeightmapUsageMap); WeightmapIt; ++WeightmapIt)
+								for (auto& WeightmapUsagePair : LandscapeProxy->WeightmapUsageMap)
 								{
-									FLandscapeWeightmapUsage* TryWeightmapUsage = &WeightmapIt.Value();
+									FLandscapeWeightmapUsage* TryWeightmapUsage = &WeightmapUsagePair.Value;
 									if (TryWeightmapUsage->FreeChannelCount() >= TotalNeededChannels)
 									{
 										// See if this candidate is closer than any others we've found
 										for (int32 ChanIdx = 0; ChanIdx < 4; ChanIdx++)
 										{
-											if (TryWeightmapUsage->ChannelUsage[ChanIdx] != NULL)
+											if (TryWeightmapUsage->ChannelUsage[ChanIdx] != nullptr)
 											{
-												int32 TryDistanceSquared = (TryWeightmapUsage->ChannelUsage[ChanIdx]->GetSectionBase() - Comp->GetSectionBase()).SizeSquared();
+												int32 TryDistanceSquared = (TryWeightmapUsage->ChannelUsage[ChanIdx]->GetSectionBase() - Component->GetSectionBase()).SizeSquared();
 												if (TryDistanceSquared < BestDistanceSquared)
 												{
-													CurrentWeightmapTexture = WeightmapIt.Key();
+													CurrentWeightmapTexture = WeightmapUsagePair.Key;
 													CurrentWeightmapUsage = TryWeightmapUsage;
 													BestDistanceSquared = TryDistanceSquared;
 												}
@@ -570,17 +577,17 @@ public:
 
 							bool NeedsUpdateResource = false;
 							// No suitable weightmap texture
-							if (CurrentWeightmapTexture == NULL)
+							if (CurrentWeightmapTexture == nullptr)
 							{
-								Comp->MarkPackageDirty();
+								Component->MarkPackageDirty();
 
 								// Weightmap is sized the same as the component
-								int32 WeightmapSize = (Comp->SubsectionSizeQuads + 1) * Comp->NumSubsections;
+								int32 WeightmapSize = (Component->SubsectionSizeQuads + 1) * Component->NumSubsections;
 
 								// We need a new weightmap texture
 								CurrentWeightmapTexture = LandscapeProxy->CreateLandscapeTexture(WeightmapSize, WeightmapSize, TEXTUREGROUP_Terrain_Weightmap, TSF_BGRA8);
 								// Alloc dummy mips
-								Comp->CreateEmptyTextureMips(CurrentWeightmapTexture);
+								Component->CreateEmptyTextureMips(CurrentWeightmapTexture);
 								CurrentWeightmapTexture->PostEditChange();
 
 								// Store it in the usage map
@@ -595,10 +602,10 @@ public:
 							{
 								// UE_LOG(LogLandscape, Log, TEXT("Finding allocation for layer %d"), CurrentLayer);
 
-								if (CurrentWeightmapUsage->ChannelUsage[ChanIdx] == NULL)
+								if (CurrentWeightmapUsage->ChannelUsage[ChanIdx] == nullptr)
 								{
 									// Use this allocation
-									FWeightmapLayerAllocationInfo& AllocInfo = Comp->WeightmapLayerAllocations[CurrentLayer];
+									FWeightmapLayerAllocationInfo& AllocInfo = Component->WeightmapLayerAllocations[CurrentLayer];
 
 									if (AllocInfo.WeightmapTextureIndex == 255)
 									{
@@ -607,19 +614,19 @@ public:
 									}
 									else
 									{
-										UTexture2D* OldWeightmapTexture = Comp->WeightmapTextures[AllocInfo.WeightmapTextureIndex];
+										UTexture2D* OldWeightmapTexture = Component->WeightmapTextures[AllocInfo.WeightmapTextureIndex];
 
 										// Copy the data
 										LandscapeEdit.CopyTextureChannel(CurrentWeightmapTexture, ChanIdx, OldWeightmapTexture, AllocInfo.WeightmapTextureChannel);
 										LandscapeEdit.ZeroTextureChannel(OldWeightmapTexture, AllocInfo.WeightmapTextureChannel);
 
 										// Remove the old allocation
-										FLandscapeWeightmapUsage* OldWeightmapUsage = Comp->GetLandscapeProxy()->WeightmapUsageMap.Find(OldWeightmapTexture);
-										OldWeightmapUsage->ChannelUsage[AllocInfo.WeightmapTextureChannel] = NULL;
+										FLandscapeWeightmapUsage* OldWeightmapUsage = Component->GetLandscapeProxy()->WeightmapUsageMap.Find(OldWeightmapTexture);
+										OldWeightmapUsage->ChannelUsage[AllocInfo.WeightmapTextureChannel] = nullptr;
 									}
 
 									// Assign the new allocation
-									CurrentWeightmapUsage->ChannelUsage[ChanIdx] = Comp;
+									CurrentWeightmapUsage->ChannelUsage[ChanIdx] = Component;
 									AllocInfo.WeightmapTextureIndex = NewWeightmapTextures.Num() - 1;
 									AllocInfo.WeightmapTextureChannel = ChanIdx;
 									CurrentLayer++;
@@ -629,12 +636,11 @@ public:
 						}
 
 						// Replace the weightmap textures
-						Comp->WeightmapTextures = NewWeightmapTextures;
+						Component->WeightmapTextures = NewWeightmapTextures;
 
 						// Update the mipmaps for the textures we edited
-						for (int32 Idx = 0; Idx < Comp->WeightmapTextures.Num(); Idx++)
+						for (UTexture2D* WeightmapTexture : Component->WeightmapTextures)
 						{
-							UTexture2D* WeightmapTexture = Comp->WeightmapTextures[Idx];
 							FLandscapeTextureDataInfo* WeightmapDataInfo = LandscapeEdit.GetTextureDataInfo(WeightmapTexture);
 
 							int32 NumMips = WeightmapTexture->Source.GetNumMips();
@@ -645,7 +651,7 @@ public:
 								WeightmapTextureMipData[MipIdx] = (FColor*)WeightmapDataInfo->GetMipData(MipIdx);
 							}
 
-							ULandscapeComponent::UpdateWeightmapMips(Comp->NumSubsections, Comp->SubsectionSizeQuads, WeightmapTexture, WeightmapTextureMipData, 0, 0, MAX_int32, MAX_int32, WeightmapDataInfo);
+							ULandscapeComponent::UpdateWeightmapMips(Component->NumSubsections, Component->SubsectionSizeQuads, WeightmapTexture, WeightmapTextureMipData, 0, 0, MAX_int32, MAX_int32, WeightmapDataInfo);
 						}
 					}
 					// Need to Repacking all the Weight map (to make it packed well...)
@@ -654,45 +660,43 @@ public:
 
 				// Move the components to the Proxy actor
 				// This does not use the MoveSelectedActorsToCurrentLevel path as there is no support to only move certain components.
-				for (TSet<ULandscapeComponent*>::TConstIterator It(TargetSelectedComponents); It; ++It)
+				for (ULandscapeComponent* Component :TargetSelectedComponents)
 				{
 					// Need to move or recreate all related data (Height map, Weight map, maybe collision components, allocation info)
-					ULandscapeComponent* Comp = *It;
-					Comp->GetLandscapeProxy()->LandscapeComponents.Remove(Comp);
-					Comp->UnregisterComponent();
-					Comp->DetachFromParent(true);
-					Comp->InvalidateLightingCache();
-					Comp->Rename(NULL, LandscapeProxy);
-					LandscapeProxy->LandscapeComponents.Add(Comp);
-					Comp->AttachTo(LandscapeProxy->GetRootComponent(), NAME_None, EAttachLocation::KeepWorldPosition);
-					Comp->UpdateMaterialInstances();
+					Component->GetLandscapeProxy()->LandscapeComponents.Remove(Component);
+					Component->UnregisterComponent();
+					Component->DetachFromParent(true);
+					Component->InvalidateLightingCache();
+					Component->Rename(nullptr, LandscapeProxy);
+					LandscapeProxy->LandscapeComponents.Add(Component);
+					Component->AttachTo(LandscapeProxy->GetRootComponent(), NAME_None, EAttachLocation::KeepWorldPosition);
+					Component->UpdateMaterialInstances();
 
 					FFormatNamedArguments Args;
-					Args.Add(TEXT("ComponentName"), FText::FromString(Comp->GetName()));
+					Args.Add(TEXT("ComponentName"), FText::FromString(Component->GetName()));
 					GWarn->StatusUpdate(Progress++, TotalProgress, FText::Format(LOCTEXT("MovingComponentStatus", "Moving Component: {ComponentName}"), Args));
 				}
 
-				for (TArray<ULandscapeHeightfieldCollisionComponent*>::TConstIterator It(TargetSelectedCollisionComponents); It; ++It)
+				for (ULandscapeHeightfieldCollisionComponent* Component : TargetSelectedCollisionComponents)
 				{
 					// Need to move or recreate all related data (Height map, Weight map, maybe collision components, allocation info)
-					ULandscapeHeightfieldCollisionComponent* Comp = *It;
 
 					// Move any foliage associated
-					AInstancedFoliageActor* OldIFA = AInstancedFoliageActor::GetInstancedFoliageActorForLevel(Comp->GetLandscapeProxy()->GetLevel());
+					AInstancedFoliageActor* OldIFA = AInstancedFoliageActor::GetInstancedFoliageActorForLevel(Component->GetLandscapeProxy()->GetLevel());
 					if (OldIFA)
 					{
-						OldIFA->MoveInstancesForComponentToCurrentLevel(Comp);
+						OldIFA->MoveInstancesForComponentToCurrentLevel(Component);
 					}
 
-					Comp->GetLandscapeProxy()->CollisionComponents.Remove(*It);
-					Comp->UnregisterComponent();
-					Comp->DetachFromParent(true);
-					Comp->Rename(NULL, LandscapeProxy);
-					LandscapeProxy->CollisionComponents.Add(*It);
-					Comp->AttachTo(LandscapeProxy->GetRootComponent(), NAME_None, EAttachLocation::KeepWorldPosition);
+					Component->GetLandscapeProxy()->CollisionComponents.Remove(Component);
+					Component->UnregisterComponent();
+					Component->DetachFromParent(true);
+					Component->Rename(nullptr, LandscapeProxy);
+					LandscapeProxy->CollisionComponents.Add(Component);
+					Component->AttachTo(LandscapeProxy->GetRootComponent(), NAME_None, EAttachLocation::KeepWorldPosition);
 
 					FFormatNamedArguments Args;
-					Args.Add(TEXT("ComponentName"), FText::FromString(Comp->GetName()));
+					Args.Add(TEXT("ComponentName"), FText::FromString(Component->GetName()));
 					GWarn->StatusUpdate(Progress++, TotalProgress, FText::Format(LOCTEXT("MovingComponentStatus", "Moving Component: {ComponentName}"), Args));
 				}
 
@@ -707,11 +711,11 @@ public:
 					LandscapeProxy->RegisterAllComponents();
 				}
 
-				for (TSet<ALandscapeProxy*>::TIterator It(SelectProxies); It; ++It)
+				for (ALandscapeProxy* Proxy : SelectProxies)
 				{
-					if ((*It)->GetRootComponent()->IsRegistered())
+					if (Proxy->GetRootComponent()->IsRegistered())
 					{
-						(*It)->RegisterAllComponents();
+						Proxy->RegisterAllComponents();
 					}
 				}
 
@@ -754,10 +758,10 @@ class FLandscapeToolStrokeAddComponent : public FLandscapeToolStrokeBase
 {
 public:
 	FLandscapeToolStrokeAddComponent(FEdModeLandscape* InEdMode, const FLandscapeToolTarget& InTarget)
-		:	EdMode(InEdMode)
-		,	LandscapeInfo(InTarget.LandscapeInfo.Get())
-		,	HeightCache(InTarget)
-		,	XYOffsetCache(InTarget)
+		: EdMode(InEdMode)
+		, LandscapeInfo(InTarget.LandscapeInfo.Get())
+		, HeightCache(InTarget)
+		, XYOffsetCache(InTarget)
 	{
 	}
 
@@ -770,36 +774,39 @@ public:
 
 	virtual void Apply(FEditorViewportClient* ViewportClient, FLandscapeBrush* Brush, const ULandscapeEditorObject* UISettings, const TArray<FLandscapeToolMousePosition>& MousePositions)
 	{
-		ALandscapeProxy* Landscape = LandscapeInfo ? LandscapeInfo->GetCurrentLevelLandscapeProxy(true) : NULL;
+		ALandscapeProxy* Landscape = LandscapeInfo ? LandscapeInfo->GetCurrentLevelLandscapeProxy(true) : nullptr;
 		if (Landscape && EdMode->LandscapeRenderAddCollision)
 		{
+			check(Brush->GetBrushType() == ELandscapeBrushType::Component);
+
 			// Get list of verts to update
-			TMap<FIntPoint, float> BrushInfo;
-			int32 X1, Y1, X2, Y2;
-			if (!Brush->ApplyBrush(MousePositions, BrushInfo, X1, Y1, X2, Y2))
+			// TODO - only retrieve bounds as we don't need the data
+			FLandscapeBrushData BrushInfo = Brush->ApplyBrush(MousePositions);
+			if (!BrushInfo)
 			{
 				return;
 			}
 
-			// expand the area to get valid data from regions...
+			int32 X1, Y1, X2, Y2;
+			BrushInfo.GetInclusiveBounds(X1, Y1, X2, Y2);
+
+			// Find component range for this block of data, non shared vertices
+			int32 ComponentIndexX1, ComponentIndexY1, ComponentIndexX2, ComponentIndexY2;
+			ALandscape::CalcComponentIndicesNoOverlap(X1, Y1, X2, Y2, Landscape->ComponentSizeQuads, ComponentIndexX1, ComponentIndexY1, ComponentIndexX2, ComponentIndexY2);
+
+			// expand the area by one vertex in each direction to ensure normals are calculated correctly
 			X1 -= 1;
 			Y1 -= 1;
 			X2 += 1;
 			Y2 += 1;
 
-			HeightCache.CacheData(X1, Y1, X2, Y2);
 			TArray<uint16> Data;
-			HeightCache.GetCachedData(X1,Y1,X2,Y2,Data);
-			XYOffsetCache.CacheData(X1, Y1, X2, Y2);
 			TArray<FVector> XYOffsetData;
+			HeightCache.CacheData(X1, Y1, X2, Y2);
+			XYOffsetCache.CacheData(X1, Y1, X2, Y2);
+			HeightCache.GetCachedData(X1, Y1, X2, Y2, Data);
 			bool bHasXYOffset = XYOffsetCache.GetCachedData(X1, Y1, X2, Y2, XYOffsetData);
 
-			// Find component range for this block of data, non shared vertices
-			int32 ComponentIndexX1 = (X1 + 1 >= 0) ? (X1 + 1) / Landscape->ComponentSizeQuads : (X1 + 2) / Landscape->ComponentSizeQuads - 1;
-			int32 ComponentIndexY1 = (Y1 + 1 >= 0) ? (Y1 + 1) / Landscape->ComponentSizeQuads : (Y1 + 2) / Landscape->ComponentSizeQuads - 1;
-			int32 ComponentIndexX2 = (X2 - 2 >= 0) ? (X2 - 2) / Landscape->ComponentSizeQuads : (X2 - 1) / Landscape->ComponentSizeQuads - 1;
-			int32 ComponentIndexY2 = (Y2 - 2 >= 0) ? (Y2 - 2) / Landscape->ComponentSizeQuads : (Y2 - 1) / Landscape->ComponentSizeQuads - 1;
-
 			TArray<ULandscapeComponent*> NewComponents;
 			Landscape->Modify();
 			LandscapeInfo->Modify();
@@ -829,7 +836,7 @@ public:
 
 						int32 ComponentVerts = (Landscape->SubsectionSizeQuads + 1) * Landscape->NumSubsections;
 						// Update Weightmap Scale Bias
-						LandscapeComponent->WeightmapScaleBias = FVector4(1.f / (float)ComponentVerts, 1.f / (float)ComponentVerts, 0.5f / (float)ComponentVerts, 0.5f / (float)ComponentVerts);
+						LandscapeComponent->WeightmapScaleBias = FVector4(1.0f / (float)ComponentVerts, 1.0f / (float)ComponentVerts, 0.5f / (float)ComponentVerts, 0.5f / (float)ComponentVerts);
 						LandscapeComponent->WeightmapSubsectionOffset = (float)(LandscapeComponent->SubsectionSizeQuads + 1) / (float)ComponentVerts;
 
 						TArray<FColor> HeightData;
@@ -853,8 +860,8 @@ public:
 				XYOffsetCache.Flush();
 			}
 
-				HeightCache.SetCachedData(X1, Y1, X2, Y2, Data);
-				HeightCache.Flush();
+			HeightCache.SetCachedData(X1, Y1, X2, Y2, Data);
+			HeightCache.Flush();
 
 			for (int32 Idx = 0; Idx < NewComponents.Num(); Idx++)
 			{
@@ -874,11 +881,11 @@ public:
 
 				for (int32 NeighborIdx = 0; NeighborIdx < 8; ++NeighborIdx)
 				{
-					ULandscapeComponent* Comp = LandscapeInfo->XYtoComponentMap.FindRef(LandscapeKey[NeighborIdx]);
-					if (Comp)
+					ULandscapeComponent* Component = LandscapeInfo->XYtoComponentMap.FindRef(LandscapeKey[NeighborIdx]);
+					if (Component)
 					{
-						NewComponents[Idx]->NeighborLOD[NeighborIdx] = Comp->ForcedLOD >= 0 ? Comp->ForcedLOD : 255;
-						NewComponents[Idx]->NeighborLODBias[NeighborIdx] = Comp->LODBias + 128;
+						NewComponents[Idx]->NeighborLOD[NeighborIdx] = Component->ForcedLOD >= 0 ? Component->ForcedLOD : 255;
+						NewComponents[Idx]->NeighborLODBias[NeighborIdx] = Component->LODBias + 128;
 					}
 				}
 
@@ -894,7 +901,7 @@ public:
 				}
 			}
 
-			EdMode->LandscapeRenderAddCollision = NULL;
+			EdMode->LandscapeRenderAddCollision = nullptr;
 
 			GEngine->BroadcastOnActorMoved(Landscape);
 		}
@@ -950,12 +957,17 @@ public:
 			if (!SelectedComponents.Num())
 			{
 				// Get list of verts to update
-				TMap<FIntPoint, float> BrushInfo;
-				int32 X1, Y1, X2, Y2;
-				if (!Brush->ApplyBrush(MousePositions, BrushInfo, X1, Y1, X2, Y2))
+				// TODO - only retrieve bounds as we don't need the data
+				FLandscapeBrushData BrushInfo = Brush->ApplyBrush(MousePositions);
+				if (!BrushInfo)
 				{
 					return;
 				}
+
+				int32 X1, Y1, X2, Y2;
+				BrushInfo.GetInclusiveBounds(X1, Y1, X2, Y2);
+
+				// Shrink bounds by 1,1 to avoid GetComponentsInRegion picking up extra components on all sides due to the overlap between components
 				LandscapeInfo->GetComponentsInRegion(X1 + 1, Y1 + 1, X2 - 1, Y2 - 1, SelectedComponents);
 			}
 
@@ -965,12 +977,11 @@ public:
 			TSet<ULandscapeComponent*> HeightmapUpdateComponents;
 			// Need to split all the component which share Heightmap with selected components
 			// Search neighbor only
-			for (TSet<ULandscapeComponent*>::TConstIterator It(SelectedComponents); It; ++It)
+			for (ULandscapeComponent* Component : SelectedComponents)
 			{
-				ULandscapeComponent* Comp = *It;
-				int32 SearchX = Comp->HeightmapTexture->Source.GetSizeX() / NeedHeightmapSize;
-				int32 SearchY = Comp->HeightmapTexture->Source.GetSizeY() / NeedHeightmapSize;
-				FIntPoint ComponentBase = Comp->GetSectionBase() / Comp->ComponentSizeQuads;
+				int32 SearchX = Component->HeightmapTexture->Source.GetSizeX() / NeedHeightmapSize;
+				int32 SearchY = Component->HeightmapTexture->Source.GetSizeY() / NeedHeightmapSize;
+				FIntPoint ComponentBase = Component->GetSectionBase() / Component->ComponentSizeQuads;
 
 				for (int32 Y = 0; Y < SearchY; ++Y)
 				{
@@ -982,7 +993,7 @@ public:
 							int32 XDir = (Dir >> 1) ? 1 : -1;
 							int32 YDir = (Dir % 2) ? 1 : -1;
 							ULandscapeComponent* Neighbor = LandscapeInfo->XYtoComponentMap.FindRef(ComponentBase + FIntPoint(XDir*X, YDir*Y));
-							if (Neighbor && Neighbor->HeightmapTexture == Comp->HeightmapTexture && !HeightmapUpdateComponents.Contains(Neighbor))
+							if (Neighbor && Neighbor->HeightmapTexture == Component->HeightmapTexture && !HeightmapUpdateComponents.Contains(Neighbor))
 							{
 								Neighbor->Modify();
 								HeightmapUpdateComponents.Add(Neighbor);
@@ -995,21 +1006,21 @@ public:
 			// Changing Heightmap format for selected components
 			for (TSet<ULandscapeComponent*>::TConstIterator It(HeightmapUpdateComponents); It; ++It)
 			{
-				ULandscapeComponent* Comp = *It;
-				ALandscape::SplitHeightmap(Comp, false);
+				ULandscapeComponent* Component = *It;
+				ALandscape::SplitHeightmap(Component, false);
 			}
 
 			TArray<FIntPoint> DeletedNeighborKeys;
 			// Check which ones are need for height map change
 			for (TSet<ULandscapeComponent*>::TIterator It(SelectedComponents); It; ++It)
 			{
-				ULandscapeComponent* Comp = *It;
-				ALandscapeProxy* Proxy = Comp->GetLandscapeProxy();
+				ULandscapeComponent* Component = *It;
+				ALandscapeProxy* Proxy = Component->GetLandscapeProxy();
 				Proxy->Modify();
-				//Comp->Modify();
+				//Component->Modify();
 
 				// Reset neighbors LOD information
-				FIntPoint ComponentBase = Comp->GetSectionBase() / Comp->ComponentSizeQuads;
+				FIntPoint ComponentBase = Component->GetSectionBase() / Component->ComponentSizeQuads;
 				FIntPoint LandscapeKey[8] =
 				{
 					ComponentBase + FIntPoint(-1, -1),
@@ -1037,39 +1048,39 @@ public:
 				}
 
 				// Remove Selected Region in deleted Component
-				for (int32 Y = 0; Y < Comp->ComponentSizeQuads; ++Y)
+				for (int32 Y = 0; Y < Component->ComponentSizeQuads; ++Y)
 				{
-					for (int32 X = 0; X < Comp->ComponentSizeQuads; ++X)
+					for (int32 X = 0; X < Component->ComponentSizeQuads; ++X)
 					{
-						LandscapeInfo->SelectedRegion.Remove(FIntPoint(X, Y) + Comp->GetSectionBase());
+						LandscapeInfo->SelectedRegion.Remove(FIntPoint(X, Y) + Component->GetSectionBase());
 					}
 				}
 
-				if (Comp->HeightmapTexture)
+				if (Component->HeightmapTexture)
 				{
-					Comp->HeightmapTexture->SetFlags(RF_Transactional);
-					Comp->HeightmapTexture->Modify();
-					Comp->HeightmapTexture->MarkPackageDirty();
-					Comp->HeightmapTexture->ClearFlags(RF_Standalone); // Remove when there is no reference for this Heightmap...
+					Component->HeightmapTexture->SetFlags(RF_Transactional);
+					Component->HeightmapTexture->Modify();
+					Component->HeightmapTexture->MarkPackageDirty();
+					Component->HeightmapTexture->ClearFlags(RF_Standalone); // Remove when there is no reference for this Heightmap...
 				}
 
-				for (int32 i = 0; i < Comp->WeightmapTextures.Num(); ++i)
+				for (int32 i = 0; i < Component->WeightmapTextures.Num(); ++i)
 				{
-					Comp->WeightmapTextures[i]->SetFlags(RF_Transactional);
-					Comp->WeightmapTextures[i]->Modify();
-					Comp->WeightmapTextures[i]->MarkPackageDirty();
-					Comp->WeightmapTextures[i]->ClearFlags(RF_Standalone);
+					Component->WeightmapTextures[i]->SetFlags(RF_Transactional);
+					Component->WeightmapTextures[i]->Modify();
+					Component->WeightmapTextures[i]->MarkPackageDirty();
+					Component->WeightmapTextures[i]->ClearFlags(RF_Standalone);
 				}
 
-				if (Comp->XYOffsetmapTexture)
+				if (Component->XYOffsetmapTexture)
 				{
-					Comp->XYOffsetmapTexture->SetFlags(RF_Transactional);
-					Comp->XYOffsetmapTexture->Modify();
-					Comp->XYOffsetmapTexture->MarkPackageDirty();
-					Comp->XYOffsetmapTexture->ClearFlags(RF_Standalone);
+					Component->XYOffsetmapTexture->SetFlags(RF_Transactional);
+					Component->XYOffsetmapTexture->Modify();
+					Component->XYOffsetmapTexture->MarkPackageDirty();
+					Component->XYOffsetmapTexture->ClearFlags(RF_Standalone);
 				}
 
-				FIntPoint Key = Comp->GetSectionBase() / Comp->ComponentSizeQuads;
+				FIntPoint Key = Component->GetSectionBase() / Component->ComponentSizeQuads;
 				DeletedNeighborKeys.AddUnique(Key + FIntPoint(-1, -1));
 				DeletedNeighborKeys.AddUnique(Key + FIntPoint(+0, -1));
 				DeletedNeighborKeys.AddUnique(Key + FIntPoint(+1, -1));
@@ -1079,12 +1090,12 @@ public:
 				DeletedNeighborKeys.AddUnique(Key + FIntPoint(+0, +1));
 				DeletedNeighborKeys.AddUnique(Key + FIntPoint(+1, +1));
 
-				ULandscapeHeightfieldCollisionComponent* CollisionComp = Comp->CollisionComponent.Get();
+				ULandscapeHeightfieldCollisionComponent* CollisionComp = Component->CollisionComponent.Get();
 				if (CollisionComp)
 				{
 					CollisionComp->DestroyComponent();
 				}
-				Comp->DestroyComponent();
+				Component->DestroyComponent();
 			}
 
 			// Update AddCollisions...
@@ -1095,10 +1106,10 @@ public:
 
 			for (int32 i = 0; i < DeletedNeighborKeys.Num(); ++i)
 			{
-				ULandscapeComponent* Comp = LandscapeInfo->XYtoComponentMap.FindRef(DeletedNeighborKeys[i]);
-				if (Comp)
+				ULandscapeComponent* Component = LandscapeInfo->XYtoComponentMap.FindRef(DeletedNeighborKeys[i]);
+				if (Component)
 				{
-					ULandscapeHeightfieldCollisionComponent* CollisionComp = Comp->CollisionComponent.Get();
+					ULandscapeHeightfieldCollisionComponent* CollisionComp = Component->CollisionComponent.Get();
 					if (CollisionComp)
 					{
 						CollisionComp->UpdateAddCollisions();
@@ -1164,25 +1175,22 @@ public:
 			Gizmo->TargetLandscapeInfo = LandscapeInfo;
 
 			// Get list of verts to update
-			TMap<FIntPoint, float> BrushInfo;
-			int32 X1, Y1, X2, Y2;
-			if (!EdMode->GizmoBrush->ApplyBrush(MousePositions, BrushInfo, X1, Y1, X2, Y2))
+			// TODO - only retrieve bounds as we don't need the data
+			FLandscapeBrushData BrushInfo = Brush->ApplyBrush(MousePositions);
+			if (!BrushInfo)
 			{
 				return;
 			}
 
+			int32 X1, Y1, X2, Y2;
+			BrushInfo.GetInclusiveBounds(X1, Y1, X2, Y2);
+
 			//Gizmo->Modify(); // No transaction for Copied data as other tools...
 			//Gizmo->SelectedData.Empty();
 			Gizmo->ClearGizmoData();
 
 			// Tablet pressure
-			//float Pressure = ViewportClient->Viewport->IsPenActive() ? ViewportClient->Viewport->GetTabletPressure() : 1.f;
-
-			// expand the area by one vertex in each direction to ensure normals are calculated correctly
-			X1 -= 1;
-			Y1 -= 1;
-			X2 += 1;
-			Y2 += 1;
+			//float Pressure = ViewportClient->Viewport->IsPenActive() ? ViewportClient->Viewport->GetTabletPressure() : 1.0f;
 
 			bool bApplyToAll = EdMode->UISettings->bApplyToAllTargets;
 			const int32 LayerNum = LandscapeInfo->Layers.Num();
@@ -1246,7 +1254,7 @@ public:
 						for (int32 i = -1; (!bApplyToAll && i < 0) || i < LayerNum; ++i)
 						{
 							// Don't try to copy data for null layers
-							if ( (bApplyToAll && i >= 0 && !LandscapeInfo->Layers[i].LayerInfoObj) ||
+							if ((bApplyToAll && i >= 0 && !LandscapeInfo->Layers[i].LayerInfoObj) ||
 								(!bApplyToAll && !EdMode->CurrentToolTarget.LayerInfo.Get()))
 							{
 								continue;
@@ -1292,7 +1300,7 @@ public:
 							FGizmoPreData LerpedData;
 							float FracX = LandscapeLocal.X - LX;
 							float FracY = LandscapeLocal.Y - LY;
-							LerpedData.Ratio = bFullCopy ? 1.f :
+							LerpedData.Ratio = bFullCopy ? 1.0f :
 								FMath::Lerp(
 								FMath::Lerp(GizmoPreData[0].Ratio, GizmoPreData[1].Ratio, FracX),
 								FMath::Lerp(GizmoPreData[2].Ratio, GizmoPreData[3].Ratio, FracX),
@@ -1305,12 +1313,12 @@ public:
 								FracY
 								);
 
-							if (!bDidCopy && LerpedData.Ratio > 0.f)
+							if (!bDidCopy && LerpedData.Ratio > 0.0f)
 							{
 								bDidCopy = true;
 							}
 
-							if (LerpedData.Ratio > 0.f)
+							if (LerpedData.Ratio > 0.0f)
 							{
 								// Added for LayerNames
 								if (bApplyToAll)
@@ -1427,7 +1435,7 @@ public:
 			//for ( TMap<uint64, FGizmoSelectData>::TIterator It(Gizmo->SelectedData); It; ++It )
 			//{
 			//	FGizmoSelectData& Data = It.Value();
-			//	if (Data.Ratio <= 0.f)
+			//	if (Data.Ratio <= 0.0f)
 			//	{
 			//		Gizmo->SelectedData.Remove(It.Key());
 			//	}
@@ -1456,6 +1464,7 @@ class FLandscapeToolCopy : public FLandscapeToolBase<FLandscapeToolStrokeCopy<To
 public:
 	FLandscapeToolCopy(FEdModeLandscape* InEdMode)
 		: FLandscapeToolBase<FLandscapeToolStrokeCopy<ToolTarget> >(InEdMode)
+		, BackupCurrentBrush(nullptr)
 	{
 	}
 
@@ -1475,37 +1484,25 @@ public:
 
 	virtual bool BeginTool(FEditorViewportClient* ViewportClient, const FLandscapeToolTarget& InTarget, const FVector& InHitLocation) override
 	{
-		check(this->MousePositions.Num() == 0);
-		this->bToolActive = true;
-		this->ToolStroke.Emplace(this->EdMode, InTarget);
-
 		this->EdMode->GizmoBrush->Tick(ViewportClient, 0.1f);
-		this->EdMode->GizmoBrush->BeginStroke(InHitLocation.X, InHitLocation.Y, this);
 
-		new(this->MousePositions) FLandscapeToolMousePosition(InHitLocation.X, InHitLocation.Y, false);
-		this->ToolStroke->Apply(ViewportClient, this->EdMode->CurrentBrush, this->EdMode->UISettings, this->MousePositions);
-		this->MousePositions.Empty();
-		return true;
+		// horrible hack
+		// (but avoids duplicating the code from FLandscapeToolBase)
+		BackupCurrentBrush = this->EdMode->CurrentBrush;
+		this->EdMode->CurrentBrush = this->EdMode->GizmoBrush;
+
+		return FLandscapeToolBase::BeginTool(ViewportClient, InTarget, InHitLocation);
 	}
 
 	virtual void EndTool(FEditorViewportClient* ViewportClient) override
 	{
-		if (this->bToolActive && this->MousePositions.Num())
-		{
-			this->ToolStroke->Apply(ViewportClient, this->EdMode->CurrentBrush, this->EdMode->UISettings, this->MousePositions);
-			this->MousePositions.Empty();
-		}
-
-		this->ToolStroke.Reset();
-		this->bToolActive = false;
-		this->EdMode->GizmoBrush->EndStroke();
-	}
+		FLandscapeToolBase::EndTool(ViewportClient);
 
-	virtual bool MouseMove(FEditorViewportClient* ViewportClient, FViewport* Viewport, int x, int y) override
-	{
-		return true;
+		this->EdMode->CurrentBrush = BackupCurrentBrush;
 	}
 
+protected:
+	FLandscapeBrush* BackupCurrentBrush;
 };
 
 template<class ToolTarget>
@@ -1539,15 +1536,17 @@ public:
 			//LandscapeInfo->Modify();
 
 			// Get list of verts to update
-			TMap<FIntPoint, float> BrushInfo;
-			int32 X1, Y1, X2, Y2;
-			if (!Brush->ApplyBrush(MousePositions, BrushInfo, X1, Y1, X2, Y2))
+			FLandscapeBrushData BrushInfo = Brush->ApplyBrush(MousePositions);
+			if (!BrushInfo)
 			{
 				return;
 			}
 
+			int32 X1, Y1, X2, Y2;
+			BrushInfo.GetInclusiveBounds(X1, Y1, X2, Y2);
+
 			// Tablet pressure
-			float Pressure = (ViewportClient && ViewportClient->Viewport->IsPenActive()) ? ViewportClient->Viewport->GetTabletPressure() : 1.f;
+			float Pressure = (ViewportClient && ViewportClient->Viewport->IsPenActive()) ? ViewportClient->Viewport->GetTabletPressure() : 1.0f;
 
 			// expand the area by one vertex in each direction to ensure normals are calculated correctly
 			X1 -= 1;
@@ -1585,8 +1584,8 @@ public:
 			const float W = Gizmo->GetWidth() / (2 * ScaleXY);
 			const float H = Gizmo->GetHeight() / (2 * ScaleXY);
 
-			const float SignX = Gizmo->GetRootComponent()->RelativeScale3D.X > 0.f ? 1.0f : -1.0f;
-			const float SignY = Gizmo->GetRootComponent()->RelativeScale3D.Y > 0.f ? 1.0f : -1.0f;
+			const float SignX = Gizmo->GetRootComponent()->RelativeScale3D.X > 0.0f ? 1.0f : -1.0f;
+			const float SignY = Gizmo->GetRootComponent()->RelativeScale3D.Y > 0.0f ? 1.0f : -1.0f;
 
 			const float ScaleX = Gizmo->CachedWidth / Width * ScaleXY / Gizmo->CachedScaleXY;
 			const float ScaleY = Gizmo->CachedHeight / Height * ScaleXY / Gizmo->CachedScaleXY;
@@ -1596,35 +1595,39 @@ public:
 			FVector BaseLocation = WToL.TransformPosition(Gizmo->GetActorLocation());
 			//FMatrix LandscapeLocalToGizmo = FRotationTranslationMatrix(FRotator(0, Gizmo->Rotation.Yaw, 0), FVector(BaseLocation.X - W + 0.5, BaseLocation.Y - H + 0.5, 0));
 			FMatrix LandscapeToGizmoLocal =
-				(FTranslationMatrix(FVector((-W + 0.5)*SignX, (-H + 0.5)*SignY, 0)) * FScaleRotationTranslationMatrix(FVector(SignX, SignY, 1.f), FRotator(0, Gizmo->GetActorRotation().Yaw, 0), FVector(BaseLocation.X, BaseLocation.Y, 0))).InverseFast();
+				(FTranslationMatrix(FVector((-W + 0.5)*SignX, (-H + 0.5)*SignY, 0)) * FScaleRotationTranslationMatrix(FVector(SignX, SignY, 1.0f), FRotator(0, Gizmo->GetActorRotation().Yaw, 0), FVector(BaseLocation.X, BaseLocation.Y, 0))).InverseFast();
 
-			for (auto It = BrushInfo.CreateIterator(); It; ++It)
+			for (int32 Y = BrushInfo.GetBounds().Min.Y; Y < BrushInfo.GetBounds().Max.Y; Y++)
 			{
-				int32 X, Y;
-				ALandscape::UnpackKey(It.Key(), X, Y);
+				const float* BrushScanline = BrushInfo.GetDataPtr(FIntPoint(0, Y));
 
-				if (It.Value() > 0.f)
+				for (int32 X = BrushInfo.GetBounds().Min.X; X < BrushInfo.GetBounds().Max.X; X++)
 				{
-					// Value before we apply our painting
-					int32 index = (X - X1) + (Y - Y1)*(1 + X2 - X1);
-					float PaintAmount = (Brush->GetBrushType() == FLandscapeBrush::BT_Gizmo) ? It.Value() : It.Value() * EdMode->UISettings->ToolStrength * Pressure;
+					const float BrushValue = BrushScanline[X];
 
-					FVector GizmoLocal = LandscapeToGizmoLocal.TransformPosition(FVector(X, Y, 0));
-					GizmoLocal.X *= ScaleX * SignX;
-					GizmoLocal.Y *= ScaleY * SignY;
+					if (BrushValue > 0.0f)
+					{
+						// TODO: This is a mess and badly needs refactoring
 
-					int32 LX = FMath::FloorToInt(GizmoLocal.X);
-					int32 LY = FMath::FloorToInt(GizmoLocal.Y);
+						// Value before we apply our painting
+						int32 index = (X - X1) + (Y - Y1)*(1 + X2 - X1);
+						float PaintAmount = (Brush->GetBrushType() == ELandscapeBrushType::Gizmo) ? BrushValue : BrushValue * EdMode->UISettings->ToolStrength * Pressure;
 
-					float FracX = GizmoLocal.X - LX;
-					float FracY = GizmoLocal.Y - LY;
+						FVector GizmoLocal = LandscapeToGizmoLocal.TransformPosition(FVector(X, Y, 0));
+						GizmoLocal.X *= ScaleX * SignX;
+						GizmoLocal.Y *= ScaleY * SignY;
 
-					FGizmoSelectData* Data00 = Gizmo->SelectedData.Find(ALandscape::MakeKey(LX, LY));
-					FGizmoSelectData* Data10 = Gizmo->SelectedData.Find(ALandscape::MakeKey(LX + 1, LY));
-					FGizmoSelectData* Data01 = Gizmo->SelectedData.Find(ALandscape::MakeKey(LX, LY + 1));
-					FGizmoSelectData* Data11 = Gizmo->SelectedData.Find(ALandscape::MakeKey(LX + 1, LY + 1));
+						int32 LX = FMath::FloorToInt(GizmoLocal.X);
+						int32 LY = FMath::FloorToInt(GizmoLocal.Y);
+
+						float FracX = GizmoLocal.X - LX;
+						float FracY = GizmoLocal.Y - LY;
+
+						FGizmoSelectData* Data00 = Gizmo->SelectedData.Find(ALandscape::MakeKey(LX, LY));
+						FGizmoSelectData* Data10 = Gizmo->SelectedData.Find(ALandscape::MakeKey(LX + 1, LY));
+						FGizmoSelectData* Data01 = Gizmo->SelectedData.Find(ALandscape::MakeKey(LX, LY + 1));
+						FGizmoSelectData* Data11 = Gizmo->SelectedData.Find(ALandscape::MakeKey(LX + 1, LY + 1));
 
-					{
 						for (int32 i = -1; (!bApplyToAll && i < 0) || i < LayerNum; ++i)
 						{
 							if ((bApplyToAll && (i < 0)) || (!bApplyToAll && EdMode->CurrentToolTarget.TargetType == ELandscapeToolTargetType::Heightmap))
@@ -1655,10 +1658,10 @@ public:
 								switch (EdMode->UISettings->PasteMode)
 								{
 								case ELandscapeToolNoiseMode::Add:
-									PaintAmount = OriginalValue < DestValue ? PaintAmount : 0.f;
+									PaintAmount = OriginalValue < DestValue ? PaintAmount : 0.0f;
 									break;
 								case ELandscapeToolNoiseMode::Sub:
-									PaintAmount = OriginalValue > DestValue ? PaintAmount : 0.f;
+									PaintAmount = OriginalValue > DestValue ? PaintAmount : 0.0f;
 									break;
 								default:
 									break;
@@ -1753,8 +1756,9 @@ class FLandscapeToolPaste : public FLandscapeToolBase<FLandscapeToolStrokePaste<
 {
 public:
 	FLandscapeToolPaste(FEdModeLandscape* InEdMode)
-		: FLandscapeToolBase<FLandscapeToolStrokePaste<ToolTarget> >(InEdMode)
+		: FLandscapeToolBase<FLandscapeToolStrokePaste<ToolTarget>>(InEdMode)
 		, bUseGizmoRegion(false)
+		, BackupCurrentBrush(nullptr)
 	{
 	}
 
@@ -1779,43 +1783,28 @@ public:
 
 	virtual bool BeginTool(FEditorViewportClient* ViewportClient, const FLandscapeToolTarget& InTarget, const FVector& InHitLocation) override
 	{
-		this->bToolActive = true;
-		this->ToolStroke.Emplace(this->EdMode, InTarget);
-
 		this->EdMode->GizmoBrush->Tick(ViewportClient, 0.1f);
+
+		// horrible hack
+		// (but avoids duplicating the code from FLandscapeToolBase)
+		BackupCurrentBrush = this->EdMode->CurrentBrush;
 		if (bUseGizmoRegion)
 		{
-			this->EdMode->GizmoBrush->BeginStroke(InHitLocation.X, InHitLocation.Y, this);
-		}
-		else
-		{
-			this->EdMode->CurrentBrush->BeginStroke(InHitLocation.X, InHitLocation.Y, this);
+			this->EdMode->CurrentBrush = this->EdMode->GizmoBrush;
 		}
 
-		new(this->MousePositions) FLandscapeToolMousePosition(InHitLocation.X, InHitLocation.Y, false);
-		this->ToolStroke->Apply(ViewportClient, bUseGizmoRegion ? this->EdMode->GizmoBrush : this->EdMode->CurrentBrush, this->EdMode->UISettings, this->MousePositions);
-		this->MousePositions.Empty();
-		return true;
+		return FLandscapeToolBase::BeginTool(ViewportClient, InTarget, InHitLocation);
 	}
 
 	virtual void EndTool(FEditorViewportClient* ViewportClient) override
 	{
-		if (this->bToolActive && this->MousePositions.Num())
-		{
-			this->ToolStroke->Apply(ViewportClient, this->EdMode->CurrentBrush, this->EdMode->UISettings, this->MousePositions);
-			this->MousePositions.Empty();
-		}
+		FLandscapeToolBase::EndTool(ViewportClient);
 
-		this->ToolStroke.Reset();
-		this->bToolActive = false;
 		if (bUseGizmoRegion)
 		{
-			this->EdMode->GizmoBrush->EndStroke();
-		}
-		else
-		{
-			this->EdMode->CurrentBrush->EndStroke();
+			this->EdMode->CurrentBrush = BackupCurrentBrush;
 		}
+		check(this->EdMode->CurrentBrush == BackupCurrentBrush);
 	}
 
 	virtual bool MouseMove(FEditorViewportClient* ViewportClient, FViewport* Viewport, int32 x, int32 y) override
@@ -1825,11 +1814,12 @@ public:
 			return true;
 		}
 
-		FLandscapeToolBase<FLandscapeToolStrokePaste<ToolTarget> >::MouseMove(ViewportClient, Viewport, x, y);
-		return true;
+		return FLandscapeToolBase<FLandscapeToolStrokePaste<ToolTarget>>::MouseMove(ViewportClient, Viewport, x, y);
 	}
+
 protected:
 	bool bUseGizmoRegion;
+	FLandscapeBrush* BackupCurrentBrush;
 };
 
 template<class ToolTarget>
diff --git a/Engine/Source/Editor/LandscapeEditor/Private/LandscapeEdModeErosionTools.cpp b/Engine/Source/Editor/LandscapeEditor/Private/LandscapeEdModeErosionTools.cpp
--- a/Engine/Source/Editor/LandscapeEditor/Private/LandscapeEdModeErosionTools.cpp
+++ b/Engine/Source/Editor/LandscapeEditor/Private/LandscapeEdModeErosionTools.cpp
@@ -70,15 +70,17 @@ public:
 		}
 
 		// Get list of verts to update
-		TMap<FIntPoint, float> BrushInfo;
-		int32 X1, Y1, X2, Y2;
-		if (!Brush->ApplyBrush(MousePositions, BrushInfo, X1, Y1, X2, Y2) || MousePositions.Num() == 0)
+		FLandscapeBrushData BrushInfo = Brush->ApplyBrush(MousePositions);
+		if (!BrushInfo)
 		{
 			return;
 		}
 
+		int32 X1, Y1, X2, Y2;
+		BrushInfo.GetInclusiveBounds(X1, Y1, X2, Y2);
+
 		// Tablet pressure
-		float Pressure = ViewportClient->Viewport->IsPenActive() ? ViewportClient->Viewport->GetTabletPressure() : 1.f;
+		float Pressure = ViewportClient->Viewport->IsPenActive() ? ViewportClient->Viewport->GetTabletPressure() : 1.0f;
 
 		// expand the area by one vertex in each direction to ensure normals are calculated correctly
 		X1 -= 1;
@@ -114,109 +116,117 @@ public:
 		for (int32 i = 0; i < Iteration; i++)
 		{
 			bHasChanged = false;
-			for (auto It = BrushInfo.CreateIterator(); It; ++It)
+			for (int32 Y = BrushInfo.GetBounds().Min.Y; Y < BrushInfo.GetBounds().Max.Y; Y++)
 			{
-				int32 X, Y;
-				ALandscape::UnpackKey(It.Key(), X, Y);
+				const float* BrushScanline = BrushInfo.GetDataPtr(FIntPoint(0, Y));
 
-				if (It.Value() > 0.f)
+				for (int32 X = BrushInfo.GetBounds().Min.X; X < BrushInfo.GetBounds().Max.X; X++)
 				{
-					int32 Center = (X - X1) + (Y - Y1)*(1 + X2 - X1);
-					int32 Neighbor[NeighborNum] = { (X - 1 - X1) + (Y - Y1)*(1 + X2 - X1), (X + 1 - X1) + (Y - Y1)*(1 + X2 - X1), (X - X1) + (Y - 1 - Y1)*(1 + X2 - X1), (X - X1) + (Y + 1 - Y1)*(1 + X2 - X1) };
-					uint32 SlopeTotal = 0;
-					uint16 SlopeMax = Thresh;
+					const float BrushValue = BrushScanline[X];
 
-					for (int32 Idx = 0; Idx < NeighborNum; Idx++)
+					if (BrushValue > 0.0f)
 					{
-						if (HeightData[Center] > HeightData[Neighbor[Idx]])
+						int32 Center = (X - X1) + (Y - Y1)*(1 + X2 - X1);
+						int32 Neighbor[NeighborNum] = {
+							(X - 1 - X1) + (Y - Y1)*(1 + X2 - X1),   // -X
+							(X + 1 - X1) + (Y - Y1)*(1 + X2 - X1),   // +X
+							(X - X1) + (Y - 1 - Y1)*(1 + X2 - X1),   // -Y
+							(X - X1) + (Y + 1 - Y1)*(1 + X2 - X1) }; // +Y
+						uint32 SlopeTotal = 0;
+						uint16 SlopeMax = Thresh;
+
+						for (int32 Idx = 0; Idx < NeighborNum; Idx++)
 						{
-							uint16 Slope = HeightData[Center] - HeightData[Neighbor[Idx]];
-							if (Slope*It.Value() > Thresh)
+							if (HeightData[Center] > HeightData[Neighbor[Idx]])
 							{
-								SlopeTotal += Slope;
-								if (SlopeMax < Slope)
+								uint16 Slope = HeightData[Center] - HeightData[Neighbor[Idx]];
+								if (Slope * BrushValue > Thresh)
 								{
-									SlopeMax = Slope;
+									SlopeTotal += Slope;
+									if (SlopeMax < Slope)
+									{
+										SlopeMax = Slope;
+									}
 								}
 							}
 						}
-					}
 
-					if (SlopeTotal > 0)
-					{
-						float Softness = 1.f;
+						if (SlopeTotal > 0)
 						{
-							for (int32 Idx = 0; Idx < LayerNum; Idx++)
+							float Softness = 1.0f;
 							{
-								ULandscapeLayerInfoObject* LayerInfo = LandscapeInfo->Layers[Idx].LayerInfoObj;
-								if (LayerInfo)
+								for (int32 Idx = 0; Idx < LayerNum; Idx++)
 								{
-									uint8 Weight = WeightDatas[Center*LayerNum + Idx];
-									Softness -= (float)(Weight) / 255.f * LayerInfo->Hardness;
+									ULandscapeLayerInfoObject* LayerInfo = LandscapeInfo->Layers[Idx].LayerInfoObj;
+									if (LayerInfo)
+									{
+										uint8 Weight = WeightDatas[Center*LayerNum + Idx];
+										Softness -= (float)(Weight) / 255.0f * LayerInfo->Hardness;
+									}
 								}
 							}
-						}
-						if (Softness > 0.f)
-						{
-							//Softness = FMath::Clamp<float>(Softness, 0.f, 1.f);
-							float TotalHeightDiff = 0;
-							int32 WeightTransfer = FMath::Min<int32>(WeightMoveThresh, SlopeMax - Thresh);
-							for (int32 Idx = 0; Idx < NeighborNum; Idx++)
+							if (Softness > 0.0f)
 							{
-								float TotalWeight = 0.f;
-								if (HeightData[Center] > HeightData[Neighbor[Idx]])
+								//Softness = FMath::Clamp<float>(Softness, 0.0f, 1.0f);
+								float TotalHeightDiff = 0;
+								int32 WeightTransfer = FMath::Min<int32>(WeightMoveThresh, SlopeMax - Thresh);
+								for (int32 Idx = 0; Idx < NeighborNum; Idx++)
 								{
-									uint16 Slope = HeightData[Center] - HeightData[Neighbor[Idx]];
-									if (Slope > Thresh)
+									float TotalWeight = 0.0f;
+									if (HeightData[Center] > HeightData[Neighbor[Idx]])
 									{
-										float WeightDiff = Softness * UISettings->ToolStrength * Pressure * ((float)Slope / SlopeTotal) * It.Value();
-										//uint16 HeightDiff = (uint16)((SlopeMax - Thresh) * WeightDiff);
-										float HeightDiff = (SlopeMax - Thresh) * WeightDiff;
-										HeightData[Neighbor[Idx]] += HeightDiff;
-										TotalHeightDiff += HeightDiff;
-
-										if (bWeightApplied)
+										uint16 Slope = HeightData[Center] - HeightData[Neighbor[Idx]];
+										if (Slope > Thresh)
 										{
-											for (int32 LayerIdx = 0; LayerIdx < LayerNum; LayerIdx++)
-											{
-												float CenterWeight = (float)(WeightDatas[Center*LayerNum + LayerIdx]) / 255.f;
-												float Weight = (float)(WeightDatas[Neighbor[Idx] * LayerNum + LayerIdx]) / 255.f;
-												NeighborWeight[Idx*LayerNum + LayerIdx] = Weight*(float)Thickness + CenterWeight*WeightDiff*WeightTransfer; // transferred + original...
-												TotalWeight += NeighborWeight[Idx*LayerNum + LayerIdx];
-											}
-											// Need to normalize weight...
-											for (int32 LayerIdx = 0; LayerIdx < LayerNum; LayerIdx++)
+											float WeightDiff = Softness * UISettings->ToolStrength * Pressure * ((float)Slope / SlopeTotal) * BrushValue;
+											//uint16 HeightDiff = (uint16)((SlopeMax - Thresh) * WeightDiff);
+											float HeightDiff = (SlopeMax - Thresh) * WeightDiff;
+											HeightData[Neighbor[Idx]] += HeightDiff;
+											TotalHeightDiff += HeightDiff;
+
+											if (bWeightApplied)
 											{
-												WeightDatas[Neighbor[Idx] * LayerNum + LayerIdx] = (uint8)(255.f * NeighborWeight[Idx*LayerNum + LayerIdx] / TotalWeight);
+												for (int32 LayerIdx = 0; LayerIdx < LayerNum; LayerIdx++)
+												{
+													float CenterWeight = (float)(WeightDatas[Center*LayerNum + LayerIdx]) / 255.0f;
+													float Weight = (float)(WeightDatas[Neighbor[Idx] * LayerNum + LayerIdx]) / 255.0f;
+													NeighborWeight[Idx*LayerNum + LayerIdx] = Weight*(float)Thickness + CenterWeight*WeightDiff*WeightTransfer; // transferred + original...
+													TotalWeight += NeighborWeight[Idx*LayerNum + LayerIdx];
+												}
+												// Need to normalize weight...
+												for (int32 LayerIdx = 0; LayerIdx < LayerNum; LayerIdx++)
+												{
+													WeightDatas[Neighbor[Idx] * LayerNum + LayerIdx] = (uint8)(255.0f * NeighborWeight[Idx*LayerNum + LayerIdx] / TotalWeight);
+												}
 											}
 										}
 									}
 								}
-							}
 
-							HeightData[Center] -= TotalHeightDiff;
-
-							if (bWeightApplied)
-							{
-								float TotalWeight = 0.f;
-								float WeightDiff = Softness * UISettings->ToolStrength * Pressure * It.Value();
+								HeightData[Center] -= TotalHeightDiff;
 
-								for (int32 LayerIdx = 0; LayerIdx < LayerNum; LayerIdx++)
+								if (bWeightApplied)
 								{
-									float Weight = (float)(WeightDatas[Center*LayerNum + LayerIdx]) / 255.f;
-									CenterWeights[LayerIdx] = Weight*Thickness - Weight*WeightDiff*WeightTransfer;
-									TotalWeight += CenterWeights[LayerIdx];
-								}
-								// Need to normalize weight...
-								for (int32 LayerIdx = 0; LayerIdx < LayerNum; LayerIdx++)
-								{
-									WeightDatas[Center*LayerNum + LayerIdx] = (uint8)(255.f * CenterWeights[LayerIdx] / TotalWeight);
+									float TotalWeight = 0.0f;
+									float WeightDiff = Softness * UISettings->ToolStrength * Pressure * BrushValue;
+
+									for (int32 LayerIdx = 0; LayerIdx < LayerNum; LayerIdx++)
+									{
+										float Weight = (float)(WeightDatas[Center*LayerNum + LayerIdx]) / 255.0f;
+										CenterWeights[LayerIdx] = Weight*Thickness - Weight*WeightDiff*WeightTransfer;
+										TotalWeight += CenterWeights[LayerIdx];
+									}
+									// Need to normalize weight...
+									for (int32 LayerIdx = 0; LayerIdx < LayerNum; LayerIdx++)
+									{
+										WeightDatas[Center*LayerNum + LayerIdx] = (uint8)(255.0f * CenterWeights[LayerIdx] / TotalWeight);
+									}
 								}
-							}
 
-							bHasChanged = true;
-						} // if Softness > 0.f
-					} // if SlopeTotal > 0
+								bHasChanged = true;
+							} // if Softness > 0.0f
+						} // if SlopeTotal > 0
+					}
 				}
 			}
 			if (!bHasChanged)
@@ -232,16 +242,20 @@ public:
 		}
 
 		// Make some noise...
-		for (auto It = BrushInfo.CreateIterator(); It; ++It)
+		for (int32 Y = BrushInfo.GetBounds().Min.Y; Y < BrushInfo.GetBounds().Max.Y; Y++)
 		{
-			int32 X, Y;
-			ALandscape::UnpackKey(It.Key(), X, Y);
+			const float* BrushScanline = BrushInfo.GetDataPtr(FIntPoint(0, Y));
 
-			if (It.Value() > 0.f)
+			for (int32 X = BrushInfo.GetBounds().Min.X; X < BrushInfo.GetBounds().Max.X; X++)
 			{
-				FNoiseParameter NoiseParam(0, UISettings->ErosionNoiseScale, It.Value() * Thresh * UISettings->ToolStrength * BrushSizeAdjust);
-				float PaintAmount = ELandscapeToolNoiseMode::Conversion((ELandscapeToolNoiseMode::Type)UISettings->ErosionNoiseMode.GetValue(), NoiseParam.NoiseAmount, NoiseParam.Sample(X, Y));
-				HeightData[(X - X1) + (Y - Y1)*(1 + X2 - X1)] = FLandscapeHeightCache::ClampValue(HeightData[(X - X1) + (Y - Y1)*(1 + X2 - X1)] + PaintAmount);
+				const float BrushValue = BrushScanline[X];
+
+				if (BrushValue > 0.0f)
+				{
+					FNoiseParameter NoiseParam(0, UISettings->ErosionNoiseScale, BrushValue * Thresh * UISettings->ToolStrength * BrushSizeAdjust);
+					float PaintAmount = ELandscapeToolNoiseMode::Conversion((ELandscapeToolNoiseMode::Type)UISettings->ErosionNoiseMode.GetValue(), NoiseParam.NoiseAmount, NoiseParam.Sample(X, Y));
+					HeightData[(X - X1) + (Y - Y1)*(1 + X2 - X1)] = FLandscapeHeightCache::ClampValue(HeightData[(X - X1) + (Y - Y1)*(1 + X2 - X1)] + PaintAmount);
+				}
 			}
 		}
 
@@ -288,15 +302,17 @@ public:
 		}
 
 		// Get list of verts to update
-		TMap<FIntPoint, float> BrushInfo;
-		int32 X1, Y1, X2, Y2;
-		if (!Brush->ApplyBrush(MousePositions, BrushInfo, X1, Y1, X2, Y2) || MousePositions.Num() == 0)
+		FLandscapeBrushData BrushInfo = Brush->ApplyBrush(MousePositions);
+		if (!BrushInfo)
 		{
 			return;
 		}
 
+		int32 X1, Y1, X2, Y2;
+		BrushInfo.GetInclusiveBounds(X1, Y1, X2, Y2);
+
 		// Tablet pressure
-		float Pressure = ViewportClient->Viewport->IsPenActive() ? ViewportClient->Viewport->GetTabletPressure() : 1.f;
+		float Pressure = ViewportClient->Viewport->IsPenActive() ? ViewportClient->Viewport->GetTabletPressure() : 1.0f;
 
 		// expand the area by one vertex in each direction to ensure normals are calculated correctly
 		X1 -= 1;
@@ -304,7 +320,6 @@ public:
 		X2 += 1;
 		Y2 += 1;
 
-		const int32 NeighborNum = 8;
 		const int32 LayerNum = LandscapeInfo->Layers.Num();
 
 		const int32 Iteration = UISettings->HErodeIterationNum;
@@ -327,114 +342,129 @@ public:
 
 		// Only initial raining works better...
 		FNoiseParameter NoiseParam(0, UISettings->RainDistScale, RainAmount);
-		for (auto It = BrushInfo.CreateIterator(); It; ++It)
+		for (int32 Y = BrushInfo.GetBounds().Min.Y; Y < BrushInfo.GetBounds().Max.Y; Y++)
 		{
-			int32 X, Y;
-			ALandscape::UnpackKey(It.Key(), X, Y);
+			const float* BrushScanline = BrushInfo.GetDataPtr(FIntPoint(0, Y));
+			auto* WaterDataScanline = WaterData.GetData() + (Y - Y1) * (X2 - X1 + 1) + (0 - X1);
 
-			if (It.Value() >= 1.f)
+			for (int32 X = BrushInfo.GetBounds().Min.X; X < BrushInfo.GetBounds().Max.X; X++)
 			{
-				float PaintAmount = ELandscapeToolNoiseMode::Conversion((ELandscapeToolNoiseMode::Type)UISettings->RainDistMode.GetValue(), NoiseParam.NoiseAmount, NoiseParam.Sample(X, Y));
-				if (PaintAmount > 0) // Raining only for positive region...
-					WaterData[(X - X1) + (Y - Y1)*(1 + X2 - X1)] += PaintAmount;
+				const float BrushValue = BrushScanline[X];
+
+				if (BrushValue >= 1.0f)
+				{
+					float PaintAmount = ELandscapeToolNoiseMode::Conversion((ELandscapeToolNoiseMode::Type)UISettings->RainDistMode.GetValue(), NoiseParam.NoiseAmount, NoiseParam.Sample(X, Y));
+					if (PaintAmount > 0) // Raining only for positive region...
+						WaterDataScanline[X] += PaintAmount;
+				}
 			}
 		}
 
 		for (int32 i = 0; i < Iteration; i++)
 		{
 			bool bWaterExist = false;
-			for (auto It = BrushInfo.CreateIterator(); It; ++It)
+			for (int32 Y = BrushInfo.GetBounds().Min.Y; Y < BrushInfo.GetBounds().Max.Y; Y++)
 			{
-				int32 X, Y;
-				ALandscape::UnpackKey(It.Key(), X, Y);
+				const float* BrushScanline = BrushInfo.GetDataPtr(FIntPoint(0, Y));
 
-				if (It.Value() > 0.f)
+				for (int32 X = BrushInfo.GetBounds().Min.X; X < BrushInfo.GetBounds().Max.X; X++)
 				{
-					int32 Center = (X - X1) + (Y - Y1)*(1 + X2 - X1);
+					const float BrushValue = BrushScanline[X];
 
-					int32 Neighbor[NeighborNum] = {
-						(X - 1 - X1) + (Y - Y1)*(1 + X2 - X1), (X + 1 - X1) + (Y - Y1)*(1 + X2 - X1), (X - X1) + (Y - 1 - Y1)*(1 + X2 - X1), (X - X1) + (Y + 1 - Y1)*(1 + X2 - X1)
-						, (X - 1 - X1) + (Y - 1 - Y1)*(1 + X2 - X1), (X + 1 - X1) + (Y + 1 - Y1)*(1 + X2 - X1), (X + 1 - X1) + (Y - 1 - Y1)*(1 + X2 - X1), (X - 1 - X1) + (Y + 1 - Y1)*(1 + X2 - X1)
-					};
-
-					// Dissolving...				
-					float DissolvedAmount = DissolvingRatio * WaterData[Center] * It.Value();
-					if (DissolvedAmount > 0 && HeightData[Center] >= DissolvedAmount)
+					if (BrushValue > 0.0f)
 					{
-						HeightData[Center] -= DissolvedAmount;
-						SedimentData[Center] += DissolvedAmount;
-					}
+						int32 Center = (X - X1) + (Y - Y1)*(1 + X2 - X1);
+
+						const int32 NeighborNum = 8;
+						int32 Neighbor[NeighborNum] = {
+							(X - 1 - X1) + (Y - Y1)*(1 + X2 - X1),   // -X
+							(X + 1 - X1) + (Y - Y1)*(1 + X2 - X1),   // +X
+							(X - X1) + (Y - 1 - Y1)*(1 + X2 - X1),   // -Y
+							(X - X1) + (Y + 1 - Y1)*(1 + X2 - X1),   // +Y
+							(X - 1 - X1) + (Y - 1 - Y1)*(1 + X2 - X1),   // -X -Y
+							(X + 1 - X1) + (Y + 1 - Y1)*(1 + X2 - X1),   // +X +Y
+							(X + 1 - X1) + (Y - 1 - Y1)*(1 + X2 - X1),   // +X -Y
+							(X - 1 - X1) + (Y + 1 - Y1)*(1 + X2 - X1) }; // -X +Y
+
+						// Dissolving...
+						float DissolvedAmount = DissolvingRatio * WaterData[Center] * BrushValue;
+						if (DissolvedAmount > 0 && HeightData[Center] >= DissolvedAmount)
+						{
+							HeightData[Center] -= DissolvedAmount;
+							SedimentData[Center] += DissolvedAmount;
+						}
 
-					uint32 TotalHeightDiff = 0;
-					uint32 TotalAltitudeDiff = 0;
-					uint32 AltitudeDiff[NeighborNum] = { 0 };
-					uint32 TotalWaterDiff = 0;
-					uint32 TotalSedimentDiff = 0;
+						uint32 TotalHeightDiff = 0;
+						uint32 TotalAltitudeDiff = 0;
+						uint32 AltitudeDiff[NeighborNum] = { 0 };
+						uint32 TotalWaterDiff = 0;
+						uint32 TotalSedimentDiff = 0;
 
-					uint32 Altitude = HeightData[Center] + WaterData[Center];
-					float AverageAltitude = 0;
-					uint32 LowerNeighbor = 0;
-					for (int32 Idx = 0; Idx < NeighborNum; Idx++)
-					{
-						uint32 NeighborAltitude = HeightData[Neighbor[Idx]] + WaterData[Neighbor[Idx]];
-						if (Altitude > NeighborAltitude)
+						uint32 Altitude = HeightData[Center] + WaterData[Center];
+						float AverageAltitude = 0;
+						uint32 LowerNeighbor = 0;
+						for (int32 Idx = 0; Idx < NeighborNum; Idx++)
 						{
-							AltitudeDiff[Idx] = Altitude - NeighborAltitude;
-							TotalAltitudeDiff += AltitudeDiff[Idx];
-							LowerNeighbor++;
-							AverageAltitude += NeighborAltitude;
-							if (HeightData[Center] > HeightData[Neighbor[Idx]])
+							uint32 NeighborAltitude = HeightData[Neighbor[Idx]] + WaterData[Neighbor[Idx]];
+							if (Altitude > NeighborAltitude)
 							{
-								TotalHeightDiff += HeightData[Center] - HeightData[Neighbor[Idx]];
+								AltitudeDiff[Idx] = Altitude - NeighborAltitude;
+								TotalAltitudeDiff += AltitudeDiff[Idx];
+								LowerNeighbor++;
+								AverageAltitude += NeighborAltitude;
+								if (HeightData[Center] > HeightData[Neighbor[Idx]])
+								{
+									TotalHeightDiff += HeightData[Center] - HeightData[Neighbor[Idx]];
+								}
+							}
+							else
+							{
+								AltitudeDiff[Idx] = 0;
 							}
 						}
-						else
-						{
-							AltitudeDiff[Idx] = 0;
-						}
-					}
 
-					// Water Transfer
-					if (LowerNeighbor > 0)
-					{
-						AverageAltitude /= (LowerNeighbor);
-						// This is not mathematically correct, but makes good result
-						if (TotalHeightDiff)
+						// Water Transfer
+						if (LowerNeighbor > 0)
 						{
-							AverageAltitude *= (1.f - 0.1 * UISettings->ToolStrength * Pressure);
-							//AverageAltitude -= 4000.f * UISettings->ToolStrength;
-						}
+							AverageAltitude /= (LowerNeighbor);
+							// This is not mathematically correct, but makes good result
+							if (TotalHeightDiff)
+							{
+								AverageAltitude *= (1.0f - 0.1 * UISettings->ToolStrength * Pressure);
+								//AverageAltitude -= 4000.0f * UISettings->ToolStrength;
+							}
 
-						uint32 WaterTransfer = FMath::Min<uint32>(WaterData[Center], Altitude - (uint32)AverageAltitude) * It.Value();
+							uint32 WaterTransfer = FMath::Min<uint32>(WaterData[Center], Altitude - (uint32)AverageAltitude) * BrushValue;
 
-						for (int32 Idx = 0; Idx < NeighborNum; Idx++)
-						{
-							if (AltitudeDiff[Idx] > 0)
+							for (int32 Idx = 0; Idx < NeighborNum; Idx++)
 							{
-								uint32 WaterDiff = (uint32)(WaterTransfer * (float)AltitudeDiff[Idx] / TotalAltitudeDiff);
-								WaterData[Neighbor[Idx]] += WaterDiff;
-								TotalWaterDiff += WaterDiff;
-								uint32 SedimentDiff = (uint32)(SedimentData[Center] * (float)WaterDiff / WaterData[Center]);
-								SedimentData[Neighbor[Idx]] += SedimentDiff;
-								TotalSedimentDiff += SedimentDiff;
+								if (AltitudeDiff[Idx] > 0)
+								{
+									uint32 WaterDiff = (uint32)(WaterTransfer * (float)AltitudeDiff[Idx] / TotalAltitudeDiff);
+									WaterData[Neighbor[Idx]] += WaterDiff;
+									TotalWaterDiff += WaterDiff;
+									uint32 SedimentDiff = (uint32)(SedimentData[Center] * (float)WaterDiff / WaterData[Center]);
+									SedimentData[Neighbor[Idx]] += SedimentDiff;
+									TotalSedimentDiff += SedimentDiff;
+								}
 							}
-						}
 
-						WaterData[Center] -= TotalWaterDiff;
-						SedimentData[Center] -= TotalSedimentDiff;
-					}
+							WaterData[Center] -= TotalWaterDiff;
+							SedimentData[Center] -= TotalSedimentDiff;
+						}
 
-					// evaporation
-					if (WaterData[Center] > 0)
-					{
-						bWaterExist = true;
-						WaterData[Center] = (uint16)(WaterData[Center] * (1.f - EvaporateRatio));
-						float SedimentCap = SedimentCapacity*WaterData[Center];
-						float SedimentDiff = SedimentData[Center] - SedimentCap;
-						if (SedimentDiff > 0)
+						// evaporation
+						if (WaterData[Center] > 0)
 						{
-							SedimentData[Center] -= SedimentDiff;
-							HeightData[Center] = FMath::Clamp<uint16>(HeightData[Center] + SedimentDiff, 0, 65535);
+							bWaterExist = true;
+							WaterData[Center] = (uint16)(WaterData[Center] * (1.0f - EvaporateRatio));
+							float SedimentCap = SedimentCapacity*WaterData[Center];
+							float SedimentDiff = SedimentData[Center] - SedimentCap;
+							if (SedimentDiff > 0)
+							{
+								SedimentData[Center] -= SedimentDiff;
+								HeightData[Center] = FMath::Clamp<uint16>(HeightData[Center] + SedimentDiff, 0, 65535);
+							}
 						}
 					}
 				}
diff --git a/Engine/Source/Editor/LandscapeEditor/Private/LandscapeEdModePaintTools.cpp b/Engine/Source/Editor/LandscapeEditor/Private/LandscapeEdModePaintTools.cpp
--- a/Engine/Source/Editor/LandscapeEditor/Private/LandscapeEdModePaintTools.cpp
+++ b/Engine/Source/Editor/LandscapeEditor/Private/LandscapeEdModePaintTools.cpp
@@ -74,9 +74,6 @@ protected:
 // 
 // FLandscapeToolPaint
 //
-#define DATA_AT(Array, X, Y) ((Array)[(X-X1) + (Y-Y1)*(1+X2-X1)])
-
-
 template<class ToolTarget>
 class FLandscapeToolStrokePaint : public FLandscapeToolStrokePaintBase<ToolTarget>
 {
@@ -92,21 +89,26 @@ public:
 	void Apply(FEditorViewportClient* ViewportClient, FLandscapeBrush* Brush, const ULandscapeEditorObject* UISettings, const TArray<FLandscapeToolMousePosition>& MousePositions)
 	{
 		// Get list of verts to update
-		TMap<FIntPoint, float> BrushInfo;
-		int32 X1, Y1, X2, Y2;
-		if (!Brush->ApplyBrush(MousePositions, BrushInfo, X1, Y1, X2, Y2) || MousePositions.Num() == 0)
+		FLandscapeBrushData BrushInfo = Brush->ApplyBrush(MousePositions);
+		if (!BrushInfo)
 		{
 			return;
 		}
 
+		int32 X1, Y1, X2, Y2;
+		BrushInfo.GetInclusiveBounds(X1, Y1, X2, Y2);
+
 		// Tablet pressure
-		float Pressure = ViewportClient->Viewport->IsPenActive() ? ViewportClient->Viewport->GetTabletPressure() : 1.f;
+		float Pressure = ViewportClient->Viewport->IsPenActive() ? ViewportClient->Viewport->GetTabletPressure() : 1.0f;
 
 		// expand the area by one vertex in each direction to ensure normals are calculated correctly
-		X1 -= 1;
-		Y1 -= 1;
-		X2 += 1;
-		Y2 += 1;
+		if (ToolTarget::TargetType == ELandscapeToolTargetType::Heightmap)
+		{
+			X1 -= 1;
+			Y1 -= 1;
+			X2 += 1;
+			Y2 += 1;
+		}
 
 		this->Cache.CacheData(X1, Y1, X2, Y2);
 
@@ -133,14 +135,16 @@ public:
 
 			for (int32 Y = Y1; Y < Y2; Y++)
 			{
+				auto* DataScanline = Data.GetData() + (Y - Y1) * (X2 - X1 + 1) + (0 - X1);
+				auto* OriginalDataScanline = OriginalData.GetData() + (Y - Y1) * (X2 - X1 + 1) + (0 - X1);
 				for (int32 X = X1; X < X2; X++)
 				{
 					float VertexInfluence = TotalInfluenceMap.FindRef(ALandscape::MakeKey(X, Y));
 
-					typename ToolTarget::CacheClass::DataType& CurrentValue = DATA_AT(Data, X, Y);
-					typename ToolTarget::CacheClass::DataType& SourceValue = DATA_AT(OriginalData, X, Y);
+					auto& CurrentValue = DataScanline[X];
+					auto& SourceValue = OriginalDataScanline[X];
 
-					SourceValue = FMath::Lerp<typename ToolTarget::CacheClass::DataType>(SourceValue, CurrentValue, FMath::Min<float>(VertexInfluence * 0.05f, 1.f));
+					SourceValue = FMath::Lerp(SourceValue, CurrentValue, FMath::Min<float>(VertexInfluence * 0.05f, 1.0f));
 				}
 			}
 		}
@@ -150,7 +154,7 @@ public:
 
 		// Adjust strength based on brush size and drawscale, so strength 1 = one hemisphere
 		float AdjustedStrength = ToolTarget::StrengthMultiplier(this->LandscapeInfo, UISettings->BrushRadius);
-		typename ToolTarget::CacheClass::DataType DestValue = ToolTarget::CacheClass::ClampValue(255.f * UISettings->WeightTargetValue);
+		typename ToolTarget::CacheClass::DataType DestValue = ToolTarget::CacheClass::ClampValue(255.0f * UISettings->WeightTargetValue);
 
 		float DeltaTime = FMath::Min<float>(FApp::GetDeltaTime(), 0.1f); // Under 10 fps slow down paint speed
 		// * 3.0f to partially compensate for impact of DeltaTime on slowing the tools down compared to the old framerate-dependent version
@@ -162,55 +166,68 @@ public:
 		if (bUseClayBrush)
 		{
 			// Calculate normals for brush verts in data space
-			Normals.Empty((*SourceDataArrayPtr).Num());
-			Normals.AddZeroed((*SourceDataArrayPtr).Num());
+			Normals.Empty(SourceDataArrayPtr->Num());
+			Normals.AddZeroed(SourceDataArrayPtr->Num());
 
 			for (int32 Y = Y1; Y < Y2; Y++)
 			{
+				auto* SourceDataScanline_0 = SourceDataArrayPtr->GetData() + (Y - Y1) * (X2 - X1 + 1) + (0 - X1);
+				auto* SourceDataScanline_1 = SourceDataArrayPtr->GetData() + (Y - Y1) * (X2 - X1 + 1) + (0 - X1);
+				auto* NormalsScanline_0 = Normals.GetData() + (Y - Y1) * (X2 - X1 + 1) + (0 - X1);
+				auto* NormalsScanline_1 = Normals.GetData() + (Y + 1 - Y1) * (X2 - X1 + 1) + (0 - X1);
 				for (int32 X = X1; X < X2; X++)
 				{
-					FVector Vert00 = ToWorld.TransformPosition(FVector((float)X + 0.f, (float)Y + 0.f, DATA_AT(*SourceDataArrayPtr, X + 0, Y + 0)));
-					FVector Vert01 = ToWorld.TransformPosition(FVector((float)X + 0.f, (float)Y + 1.f, DATA_AT(*SourceDataArrayPtr, X + 0, Y + 1)));
-					FVector Vert10 = ToWorld.TransformPosition(FVector((float)X + 1.f, (float)Y + 0.f, DATA_AT(*SourceDataArrayPtr, X + 1, Y + 0)));
-					FVector Vert11 = ToWorld.TransformPosition(FVector((float)X + 1.f, (float)Y + 1.f, DATA_AT(*SourceDataArrayPtr, X + 1, Y + 1)));
+					FVector Vert00 = ToWorld.TransformPosition(FVector((float)X + 0.0f, (float)Y + 0.0f, SourceDataScanline_0[X + 0]));
+					FVector Vert01 = ToWorld.TransformPosition(FVector((float)X + 0.0f, (float)Y + 1.0f, SourceDataScanline_1[X + 0]));
+					FVector Vert10 = ToWorld.TransformPosition(FVector((float)X + 1.0f, (float)Y + 0.0f, SourceDataScanline_0[X + 1]));
+					FVector Vert11 = ToWorld.TransformPosition(FVector((float)X + 1.0f, (float)Y + 1.0f, SourceDataScanline_1[X + 1]));
 
 					FVector FaceNormal1 = ((Vert00 - Vert10) ^ (Vert10 - Vert11)).GetSafeNormal();
 					FVector FaceNormal2 = ((Vert11 - Vert01) ^ (Vert01 - Vert00)).GetSafeNormal();
 
 					// contribute to the vertex normals.
-					DATA_AT(Normals, X + 1, Y + 0) += FaceNormal1;
-					DATA_AT(Normals, X + 0, Y + 1) += FaceNormal2;
-					DATA_AT(Normals, X + 0, Y + 0) += FaceNormal1 + FaceNormal2;
-					DATA_AT(Normals, X + 1, Y + 1) += FaceNormal1 + FaceNormal2;
+					NormalsScanline_0[X + 1] += FaceNormal1;
+					NormalsScanline_1[X + 0] += FaceNormal2;
+					NormalsScanline_0[X + 0] += FaceNormal1 + FaceNormal2;
+					NormalsScanline_1[X + 1] += FaceNormal1 + FaceNormal2;
 				}
 			}
 			for (int32 Y = Y1; Y <= Y2; Y++)
 			{
+				auto* NormalsScanline = Normals.GetData() + (Y - Y1) * (X2 - X1 + 1) + (0 - X1);
 				for (int32 X = X1; X <= X2; X++)
 				{
-					DATA_AT(Normals, X, Y) = DATA_AT(Normals, X, Y).GetSafeNormal();
+					NormalsScanline[X] = NormalsScanline[X].GetSafeNormal();
 				}
 			}
 
 			// Find brush centroid location
-			FVector AveragePoint(0.f, 0.f, 0.f);
-			FVector AverageNormal(0.f, 0.f, 0.f);
-			float TotalWeight = 0.f;
-			for (auto It = BrushInfo.CreateIterator(); It; ++It)
+			FVector AveragePoint(0.0f, 0.0f, 0.0f);
+			FVector AverageNormal(0.0f, 0.0f, 0.0f);
+			float TotalWeight = 0.0f;
+			for (int32 Y = BrushInfo.GetBounds().Min.Y; Y < BrushInfo.GetBounds().Max.Y; Y++)
 			{
-				int32 X, Y;
-				ALandscape::UnpackKey(It.Key(), X, Y);
-				float Weight = It.Value();
+				const float* BrushScanline = BrushInfo.GetDataPtr(FIntPoint(0, Y));
+				auto* SourceDataScanline = SourceDataArrayPtr->GetData() + (Y - Y1) * (X2 - X1 + 1) + (0 - X1);
+				auto* NormalsScanline = Normals.GetData() + (Y - Y1) * (X2 - X1 + 1) + (0 - X1);
+
+				for (int32 X = BrushInfo.GetBounds().Min.X; X < BrushInfo.GetBounds().Max.X; X++)
+				{
+					const float BrushValue = BrushScanline[X];
 
-				AveragePoint += FVector((float)X * Weight, (float)Y * Weight, (float)DATA_AT(*SourceDataArrayPtr, FMath::FloorToInt(X), FMath::FloorToInt(Y)) * Weight);
+					if (BrushValue > 0.0f)
+					{
+						AveragePoint += FVector((float)X * BrushValue, (float)Y * BrushValue, (float)SourceDataScanline[X] * BrushValue);
 
-				FVector SampleNormal = DATA_AT(Normals, X, Y);
-				AverageNormal += SampleNormal * Weight;
+						FVector SampleNormal = NormalsScanline[X];
+						AverageNormal += SampleNormal * BrushValue;
 
-				TotalWeight += Weight;
+						TotalWeight += BrushValue;
+					}
+				}
 			}
 
-			if (TotalWeight > 0.f)
+			if (TotalWeight > 0.0f)
 			{
 				AveragePoint /= TotalWeight;
 				AverageNormal = AverageNormal.GetSafeNormal();
@@ -232,59 +249,66 @@ public:
 		}
 
 		// Apply the brush
-		for (auto BrushPair : BrushInfo)
+		for (int32 Y = BrushInfo.GetBounds().Min.Y; Y < BrushInfo.GetBounds().Max.Y; Y++)
 		{
-			int32 X, Y;
-			ALandscape::UnpackKey(BrushPair.Key, X, Y);
+			const float* BrushScanline = BrushInfo.GetDataPtr(FIntPoint(0, Y));
+			auto* DataScanline = Data.GetData() + (Y - Y1) * (X2 - X1 + 1) + (0 - X1);
+			auto* SourceDataScanline = SourceDataArrayPtr->GetData() + (Y - Y1) * (X2 - X1 + 1) + (0 - X1);
 
-			// Update influence map
-			float VertexInfluence = TotalInfluenceMap.FindRef(BrushPair.Key);
-			TotalInfluenceMap.Add(BrushPair.Key, VertexInfluence + BrushPair.Value);
+			for (int32 X = BrushInfo.GetBounds().Min.X; X < BrushInfo.GetBounds().Max.X; X++)
+			{
+				const FIntPoint Key = ALandscape::MakeKey(X, Y);
+				const float BrushValue = BrushScanline[X];
 
-			float PaintAmount = BrushPair.Value * PaintStrength;
-			typename ToolTarget::CacheClass::DataType& CurrentValue = DATA_AT(Data, X, Y);
-			const typename ToolTarget::CacheClass::DataType& SourceValue = DATA_AT(*SourceDataArrayPtr, X, Y);
+				// Update influence map
+				float VertexInfluence = TotalInfluenceMap.FindRef(Key);
+				TotalInfluenceMap.Add(Key, VertexInfluence + BrushValue);
 
-			if (bUseWeightTargetValue)
-			{
-				if (bInvert)
+				float PaintAmount = BrushValue * PaintStrength;
+				auto& CurrentValue = DataScanline[X];
+				const auto& SourceValue = SourceDataScanline[X];
+
+				if (bUseWeightTargetValue)
 				{
-					CurrentValue = FMath::Lerp(CurrentValue, DestValue, PaintAmount / AdjustedStrength);
+					if (bInvert)
+					{
+						CurrentValue = FMath::Lerp(CurrentValue, DestValue, PaintAmount / AdjustedStrength);
+					}
+					else
+					{
+						CurrentValue = FMath::Lerp(CurrentValue, DestValue, PaintAmount / AdjustedStrength);
+					}
 				}
-				else
+				else if (bUseClayBrush)
 				{
-					CurrentValue = FMath::Lerp(CurrentValue, DestValue, PaintAmount / AdjustedStrength);
-				}
-			}
-			else if (bUseClayBrush)
-			{
-				// Brush application starts from original world location at start of stroke
-				FVector WorldLoc = ToWorld.TransformPosition(FVector(X, Y, SourceValue));
+					// Brush application starts from original world location at start of stroke
+					FVector WorldLoc = ToWorld.TransformPosition(FVector(X, Y, SourceValue));
 
-				// Calculate new location on the brush plane
-				WorldLoc.Z = (BrushPlane.W - BrushPlane.X*WorldLoc.X - BrushPlane.Y*WorldLoc.Y) / BrushPlane.Z;
+					// Calculate new location on the brush plane
+					WorldLoc.Z = (BrushPlane.W - BrushPlane.X*WorldLoc.X - BrushPlane.Y*WorldLoc.Y) / BrushPlane.Z;
 
-				// Painted amount lerps based on brush falloff.
-				float PaintValue = FMath::Lerp<float>((float)SourceValue, FromWorld.TransformPosition(WorldLoc).Z, BrushPair.Value);
+					// Painted amount lerps based on brush falloff.
+					float PaintValue = FMath::Lerp<float>((float)SourceValue, FromWorld.TransformPosition(WorldLoc).Z, BrushValue);
 
-				if (bInvert)
-				{
-					CurrentValue = ToolTarget::CacheClass::ClampValue(FMath::Min<int32>(FMath::RoundToInt(PaintValue), CurrentValue));
-				}
-				else
-				{
-					CurrentValue = ToolTarget::CacheClass::ClampValue(FMath::Max<int32>(FMath::RoundToInt(PaintValue), CurrentValue));
-				}
-			}
-			else
-			{
-				if (bInvert)
-				{
-					CurrentValue = ToolTarget::CacheClass::ClampValue(FMath::Min<int32>(SourceValue - FMath::RoundToInt(PaintAmount), CurrentValue));
+					if (bInvert)
+					{
+						CurrentValue = ToolTarget::CacheClass::ClampValue(FMath::Min<int32>(FMath::RoundToInt(PaintValue), CurrentValue));
+					}
+					else
+					{
+						CurrentValue = ToolTarget::CacheClass::ClampValue(FMath::Max<int32>(FMath::RoundToInt(PaintValue), CurrentValue));
+					}
 				}
 				else
 				{
-					CurrentValue = ToolTarget::CacheClass::ClampValue(FMath::Max<int32>(SourceValue + FMath::RoundToInt(PaintAmount), CurrentValue));
+					if (bInvert)
+					{
+						CurrentValue = ToolTarget::CacheClass::ClampValue(FMath::Min<int32>(SourceValue - FMath::RoundToInt(PaintAmount), CurrentValue));
+					}
+					else
+					{
+						CurrentValue = ToolTarget::CacheClass::ClampValue(FMath::Max<int32>(SourceValue + FMath::RoundToInt(PaintAmount), CurrentValue));
+					}
 				}
 			}
 		}
@@ -292,7 +316,6 @@ public:
 		this->Cache.SetCachedData(X1, Y1, X2, Y2, Data, UISettings->PaintingRestriction);
 		this->Cache.Flush();
 	}
-#undef DATA_AT
 };
 
 class FLandscapeToolPaint : public FLandscapeToolPaintBase<FWeightmapToolTarget, FLandscapeToolStrokePaint<FWeightmapToolTarget>>
@@ -337,21 +360,26 @@ public:
 		if (!this->LandscapeInfo) return;
 
 		// Get list of verts to update
-		TMap<FIntPoint, float> BrushInfo;
-		int32 X1, Y1, X2, Y2;
-		if (!Brush->ApplyBrush(MousePositions, BrushInfo, X1, Y1, X2, Y2))
+		FLandscapeBrushData BrushInfo = Brush->ApplyBrush(MousePositions);
+		if (!BrushInfo)
 		{
 			return;
 		}
 
+		int32 X1, Y1, X2, Y2;
+		BrushInfo.GetInclusiveBounds(X1, Y1, X2, Y2);
+
 		// Tablet pressure
-		float Pressure = ViewportClient->Viewport->IsPenActive() ? ViewportClient->Viewport->GetTabletPressure() : 1.f;
+		float Pressure = ViewportClient->Viewport->IsPenActive() ? ViewportClient->Viewport->GetTabletPressure() : 1.0f;
 
 		// expand the area by one vertex in each direction to ensure normals are calculated correctly
-		X1 -= 1;
-		Y1 -= 1;
-		X2 += 1;
-		Y2 += 1;
+		if (ToolTarget::TargetType == ELandscapeToolTargetType::Heightmap)
+		{
+			X1 -= 1;
+			Y1 -= 1;
+			X2 += 1;
+			Y2 += 1;
+		}
 
 		this->Cache.CacheData(X1, Y1, X2, Y2);
 
@@ -365,37 +393,50 @@ public:
 		}
 		else
 		{
-			// Filter size is SmoothFilterKernelScale * BrushSize, But avoids overflow value or under 1
-			int32 HalfFilterSize = FMath::Clamp<int32>(UISettings->SmoothFilterKernelScale * FMath::Max(Y2 - Y1 + 1, X2 - X1 + 1) * 0.25f, 1, 127);
+			int32 FilterRadius = UISettings->SmoothFilterKernelSize;
 
-			for (auto It = BrushInfo.CreateIterator(); It; ++It)
+			for (int32 Y = BrushInfo.GetBounds().Min.Y; Y < BrushInfo.GetBounds().Max.Y; Y++)
 			{
-				int32 X, Y;
-				ALandscape::UnpackKey(It.Key(), X, Y);
+				const float* BrushScanline = BrushInfo.GetDataPtr(FIntPoint(0, Y));
+				auto* DataScanline = Data.GetData() + (Y - Y1) * (X2 - X1 + 1) + (0 - X1);
 
-				if (It.Value() > 0.f)
+				for (int32 X = BrushInfo.GetBounds().Min.X; X < BrushInfo.GetBounds().Max.X; X++)
 				{
-					int32 FilterValue = 0;
-					int32 FilterSamplingNumber = 0;
+					const float BrushValue = BrushScanline[X];
 
-					for (int32 y = Y - HalfFilterSize; y <= Y + HalfFilterSize; y++)
+					if (BrushValue > 0.0f)
 					{
-						int32 YY = FMath::Clamp<int32>(y, Y1, Y2);
-						for (int32 x = X - HalfFilterSize; x <= X + HalfFilterSize; x++)
+						int64 FilterValue = 0;
+						int32 FilterSamplingNumber = 0;
+
+						// The previous version of this clamped to X1/Y1/X2/Y2 *inside* the loop, which were always expanded by one from the BrushInfo
+						// So the Find on the BrushInfo always gave null for anything that would have been clamped and caused it to do nothing
+						// So now just skip iterating those by clamping first
+
+						const int32 SampleX1 = FMath::Max<int32>(X - FilterRadius, BrushInfo.GetBounds().Min.X);
+						const int32 SampleY1 = FMath::Max<int32>(Y - FilterRadius, BrushInfo.GetBounds().Min.Y);
+						const int32 SampleX2 = FMath::Min<int32>(X + FilterRadius + 1, BrushInfo.GetBounds().Max.X);
+						const int32 SampleY2 = FMath::Min<int32>(Y + FilterRadius + 1, BrushInfo.GetBounds().Max.Y);
+						for (int32 SampleY = SampleY1; SampleY < SampleY2; SampleY++)
 						{
-							int32 XX = FMath::Clamp<int32>(x, X1, X2);
-							if (BrushInfo.Find(FIntPoint(XX, YY)))
+							const float* SampleBrushScanline = BrushInfo.GetDataPtr(FIntPoint(0, SampleY));
+							auto* SampleDataScanline = Data.GetData() + (SampleY - Y1) * (X2 - X1 + 1) + (0 - X1);
+
+							for (int32 SampleX = SampleX1; SampleX < SampleX2; SampleX++)
 							{
-								FilterValue += Data[(XX - X1) + (YY - Y1)*(1 + X2 - X1)];
-								FilterSamplingNumber++;
+								const float SampleBrushValue = SampleBrushScanline[SampleX];
+								if (SampleBrushValue > 0.0f)
+								{
+									FilterValue += SampleDataScanline[SampleX];
+									FilterSamplingNumber++;
+								}
 							}
 						}
-					}
 
-					FilterValue /= FilterSamplingNumber;
+						FilterValue /= FilterSamplingNumber;
 
-					int32 HeightDataIndex = (X - X1) + (Y - Y1)*(1 + X2 - X1);
-					Data[HeightDataIndex] = FMath::Lerp(Data[HeightDataIndex], (typename ToolTarget::CacheClass::DataType)FilterValue, It.Value() * UISettings->ToolStrength * Pressure);
+						DataScanline[X] = FMath::Lerp(DataScanline[X], (typename ToolTarget::CacheClass::DataType)FilterValue, BrushValue * UISettings->ToolStrength * Pressure);
+					}
 				}
 			}
 		}
@@ -403,7 +444,6 @@ public:
 		this->Cache.SetCachedData(X1, Y1, X2, Y2, Data, UISettings->PaintingRestriction);
 		this->Cache.Flush();
 	}
-
 };
 
 template<class ToolTarget>
@@ -424,7 +464,7 @@ public:
 // FLandscapeToolFlatten
 //
 template<class ToolTarget>
-class FLandscapeToolStrokeFlatten : public FLandscapeToolStrokePaintBase < ToolTarget >
+class FLandscapeToolStrokeFlatten : public FLandscapeToolStrokePaintBase<ToolTarget>
 {
 	typename ToolTarget::CacheClass::DataType FlattenHeight;
 
@@ -455,8 +495,8 @@ public:
 		if (!bInitializedFlattenHeight || (UISettings->bPickValuePerApply && bTargetIsHeightmap))
 		{
 			bInitializedFlattenHeight = false;
-			float FlattenX = MousePositions[0].PositionX;
-			float FlattenY = MousePositions[0].PositionY;
+			float FlattenX = MousePositions[0].Position.X;
+			float FlattenY = MousePositions[0].Position.Y;
 			int32 FlattenHeightX = FMath::FloorToInt(FlattenX);
 			int32 FlattenHeightY = FMath::FloorToInt(FlattenY);
 
@@ -475,98 +515,106 @@ public:
 
 
 		// Get list of verts to update
-		TMap<FIntPoint, float> BrushInfo;
-		int32 X1, Y1, X2, Y2;
-		if (!Brush->ApplyBrush(MousePositions, BrushInfo, X1, Y1, X2, Y2))
+		FLandscapeBrushData BrushInfo = Brush->ApplyBrush(MousePositions);
+		if (!BrushInfo)
 		{
 			return;
 		}
 
+		int32 X1, Y1, X2, Y2;
+		BrushInfo.GetInclusiveBounds(X1, Y1, X2, Y2);
+
 		// Tablet pressure
-		float Pressure = ViewportClient->Viewport->IsPenActive() ? ViewportClient->Viewport->GetTabletPressure() : 1.f;
+		float Pressure = ViewportClient->Viewport->IsPenActive() ? ViewportClient->Viewport->GetTabletPressure() : 1.0f;
 
 		// expand the area by one vertex in each direction to ensure normals are calculated correctly
-		X1 -= 1;
-		Y1 -= 1;
-		X2 += 1;
-		Y2 += 1;
+		if (ToolTarget::TargetType == ELandscapeToolTargetType::Heightmap)
+		{
+			X1 -= 1;
+			Y1 -= 1;
+			X2 += 1;
+			Y2 += 1;
+		}
 
 		this->Cache.CacheData(X1, Y1, X2, Y2);
 
-		TArray<typename ToolTarget::CacheClass::DataType> HeightData;
-		this->Cache.GetCachedData(X1, Y1, X2, Y2, HeightData);
+		TArray<typename ToolTarget::CacheClass::DataType> Data;
+		this->Cache.GetCachedData(X1, Y1, X2, Y2, Data);
 
 		// Apply the brush
-		for (auto It = BrushInfo.CreateIterator(); It; ++It)
+		for (int32 Y = BrushInfo.GetBounds().Min.Y; Y < BrushInfo.GetBounds().Max.Y; Y++)
 		{
-			int32 X, Y;
-			ALandscape::UnpackKey(It.Key(), X, Y);
+			const float* BrushScanline = BrushInfo.GetDataPtr(FIntPoint(0, Y));
+			auto* DataScanline = Data.GetData() + (Y - Y1) * (X2 - X1 + 1) + (0 - X1);
 
-			if (It.Value() > 0.f)
+			for (int32 X = BrushInfo.GetBounds().Min.X; X < BrushInfo.GetBounds().Max.X; X++)
 			{
-				int32 HeightDataIndex = (X - X1) + (Y - Y1)*(1 + X2 - X1);
-			
-				float Strength = FMath::Clamp<float>(It.Value() * UISettings->ToolStrength * Pressure, 0.f, 1.f);
+				const float BrushValue = BrushScanline[X];
 
-				if (!(UISettings->bUseSlopeFlatten && bTargetIsHeightmap))
+				if (BrushValue > 0.0f)
 				{
-					int32 Delta = HeightData[HeightDataIndex] - FlattenHeight;
-					switch (UISettings->FlattenMode)
+					float Strength = FMath::Clamp<float>(BrushValue * UISettings->ToolStrength * Pressure, 0.0f, 1.0f);
+
+					if (!(UISettings->bUseSlopeFlatten && bTargetIsHeightmap))
 					{
-					case ELandscapeToolNoiseMode::Add:
-						if (Delta < 0)
+						int32 Delta = DataScanline[X] - FlattenHeight;
+						switch (UISettings->FlattenMode)
 						{
-							HeightData[HeightDataIndex] = FMath::Lerp(HeightData[HeightDataIndex], FlattenHeight, Strength);
-						}
-						break;
-					case ELandscapeToolNoiseMode::Sub:
-						if (Delta > 0)
-						{
-							HeightData[HeightDataIndex] = FMath::Lerp(HeightData[HeightDataIndex], FlattenHeight, Strength);
+						case ELandscapeToolNoiseMode::Add:
+							if (Delta < 0)
+							{
+								DataScanline[X] = FMath::Lerp(DataScanline[X], FlattenHeight, Strength);
+							}
+							break;
+						case ELandscapeToolNoiseMode::Sub:
+							if (Delta > 0)
+							{
+								DataScanline[X] = FMath::Lerp(DataScanline[X], FlattenHeight, Strength);
+							}
+							break;
+						default:
+						case ELandscapeToolNoiseMode::Both:
+							DataScanline[X] = FMath::Lerp(DataScanline[X], FlattenHeight, Strength);
+							break;
 						}
-						break;
-					default:
-					case ELandscapeToolNoiseMode::Both:
-						HeightData[HeightDataIndex] = FMath::Lerp(HeightData[HeightDataIndex], FlattenHeight, Strength);
-						break;
 					}
-				}
-				else
-				{
-					typename ToolTarget::CacheClass::DataType DestValue = -(FlattenNormal.X * X + FlattenNormal.Y * Y + FlattenPlaneDist) / FlattenNormal.Z;
-					//float PlaneDist = FlattenNormal | FVector(X, Y, HeightData(HeightDataIndex)) + FlattenPlaneDist;
-					float PlaneDist = HeightData[HeightDataIndex] - DestValue;
-					DestValue = HeightData[HeightDataIndex] - PlaneDist * Strength;
-					switch (UISettings->FlattenMode)
+					else
 					{
-					case ELandscapeToolNoiseMode::Add:
-						if (PlaneDist < 0)
+						typename ToolTarget::CacheClass::DataType DestValue = -(FlattenNormal.X * X + FlattenNormal.Y * Y + FlattenPlaneDist) / FlattenNormal.Z;
+						//float PlaneDist = FlattenNormal | FVector(X, Y, HeightData(HeightDataIndex)) + FlattenPlaneDist;
+						float PlaneDist = DataScanline[X] - DestValue;
+						DestValue = DataScanline[X] - PlaneDist * Strength;
+						switch (UISettings->FlattenMode)
 						{
-							HeightData[HeightDataIndex] = FMath::Lerp(HeightData[HeightDataIndex], DestValue, Strength);
-						}
-						break;
-					case ELandscapeToolNoiseMode::Sub:
-						if (PlaneDist > 0)
-						{
-							HeightData[HeightDataIndex] = FMath::Lerp(HeightData[HeightDataIndex], DestValue, Strength);
+						case ELandscapeToolNoiseMode::Add:
+							if (PlaneDist < 0)
+							{
+								DataScanline[X] = FMath::Lerp(DataScanline[X], DestValue, Strength);
+							}
+							break;
+						case ELandscapeToolNoiseMode::Sub:
+							if (PlaneDist > 0)
+							{
+								DataScanline[X] = FMath::Lerp(DataScanline[X], DestValue, Strength);
+							}
+							break;
+						default:
+						case ELandscapeToolNoiseMode::Both:
+							DataScanline[X] = FMath::Lerp(DataScanline[X], DestValue, Strength);
+							break;
 						}
-						break;
-					default:
-					case ELandscapeToolNoiseMode::Both:
-						HeightData[HeightDataIndex] = FMath::Lerp(HeightData[HeightDataIndex], DestValue, Strength);
-						break;
 					}
 				}
 			}
 		}
 
-		this->Cache.SetCachedData(X1, Y1, X2, Y2, HeightData, UISettings->PaintingRestriction);
+		this->Cache.SetCachedData(X1, Y1, X2, Y2, Data, UISettings->PaintingRestriction);
 		this->Cache.Flush();
 	}
 };
 
 template<class ToolTarget>
-class FLandscapeToolFlatten : public FLandscapeToolPaintBase < ToolTarget, FLandscapeToolStrokeFlatten<ToolTarget> >
+class FLandscapeToolFlatten : public FLandscapeToolPaintBase<ToolTarget, FLandscapeToolStrokeFlatten<ToolTarget>>
 {
 protected:
 	UStaticMesh* PlaneMesh;
@@ -574,7 +622,7 @@ protected:
 
 public:
 	FLandscapeToolFlatten(FEdModeLandscape* InEdMode)
-		: FLandscapeToolPaintBase<ToolTarget, FLandscapeToolStrokeFlatten<ToolTarget> >(InEdMode)
+		: FLandscapeToolPaintBase<ToolTarget, FLandscapeToolStrokeFlatten<ToolTarget>>(InEdMode)
 		, PlaneMesh(LoadObject<UStaticMesh>(NULL, TEXT("/Engine/EditorLandscapeResources/FlattenPlaneMesh.FlattenPlaneMesh")))
 		, MeshComponent(NULL)
 	{
@@ -658,21 +706,26 @@ public:
 		if (!this->LandscapeInfo) return;
 
 		// Get list of verts to update
-		TMap<FIntPoint, float> BrushInfo;
-		int32 X1, Y1, X2, Y2;
-		if (!Brush->ApplyBrush(MousePositions, BrushInfo, X1, Y1, X2, Y2))
+		FLandscapeBrushData BrushInfo = Brush->ApplyBrush(MousePositions);
+		if (!BrushInfo)
 		{
 			return;
 		}
 
+		int32 X1, Y1, X2, Y2;
+		BrushInfo.GetInclusiveBounds(X1, Y1, X2, Y2);
+
 		// Tablet pressure
-		float Pressure = ViewportClient->Viewport->IsPenActive() ? ViewportClient->Viewport->GetTabletPressure() : 1.f;
+		float Pressure = ViewportClient->Viewport->IsPenActive() ? ViewportClient->Viewport->GetTabletPressure() : 1.0f;
 
 		// expand the area by one vertex in each direction to ensure normals are calculated correctly
-		X1 -= 1;
-		Y1 -= 1;
-		X2 += 1;
-		Y2 += 1;
+		if (ToolTarget::TargetType == ELandscapeToolTargetType::Heightmap)
+		{
+			X1 -= 1;
+			Y1 -= 1;
+			X2 += 1;
+			Y2 += 1;
+		}
 
 		this->Cache.CacheData(X1, Y1, X2, Y2);
 		TArray<typename ToolTarget::CacheClass::DataType> Data;
@@ -687,42 +740,47 @@ public:
 		bool bUseWeightTargetValue = UISettings->bUseWeightTargetValue && ToolTarget::TargetType == ELandscapeToolTargetType::Weightmap;
 
 		// Apply the brush
-		for (auto It = BrushInfo.CreateIterator(); It; ++It)
+		for (int32 Y = BrushInfo.GetBounds().Min.Y; Y < BrushInfo.GetBounds().Max.Y; Y++)
 		{
-			int32 X, Y;
-			ALandscape::UnpackKey(It.Key(), X, Y);
+			const float* BrushScanline = BrushInfo.GetDataPtr(FIntPoint(0, Y));
+			auto* DataScanline = Data.GetData() + (Y - Y1) * (X2 - X1 + 1) + (0 - X1);
 
-			if (It.Value() > 0.f)
+			for (int32 X = BrushInfo.GetBounds().Min.X; X < BrushInfo.GetBounds().Max.X; X++)
 			{
-				float OriginalValue = Data[(X - X1) + (Y - Y1)*(1 + X2 - X1)];
-				if (bUseWeightTargetValue)
+				const float BrushValue = BrushScanline[X];
+
+				if (BrushValue > 0.0f)
 				{
-					FNoiseParameter NoiseParam(0, UISettings->NoiseScale, 255.f / 2.f);
-					float DestValue = ELandscapeToolNoiseMode::Conversion(ELandscapeToolNoiseMode::Add, NoiseParam.NoiseAmount, NoiseParam.Sample(X, Y)) * UISettings->WeightTargetValue;
-					switch (UISettings->NoiseMode)
+					float OriginalValue = DataScanline[X];
+					if (bUseWeightTargetValue)
 					{
-					case ELandscapeToolNoiseMode::Add:
-						if (OriginalValue >= DestValue)
+						FNoiseParameter NoiseParam(0, UISettings->NoiseScale, 255.0f / 2.0f);
+						float DestValue = ELandscapeToolNoiseMode::Conversion(ELandscapeToolNoiseMode::Add, NoiseParam.NoiseAmount, NoiseParam.Sample(X, Y)) * UISettings->WeightTargetValue;
+						switch (UISettings->NoiseMode)
 						{
-							continue;
-						}
-						break;
-					case ELandscapeToolNoiseMode::Sub:
-						DestValue += (1.f - UISettings->WeightTargetValue) * NoiseParam.NoiseAmount;
-						if (OriginalValue <= DestValue)
-						{
-							continue;
+						case ELandscapeToolNoiseMode::Add:
+							if (OriginalValue >= DestValue)
+							{
+								continue;
+							}
+							break;
+						case ELandscapeToolNoiseMode::Sub:
+							DestValue += (1.0f - UISettings->WeightTargetValue) * NoiseParam.NoiseAmount;
+							if (OriginalValue <= DestValue)
+							{
+								continue;
+							}
+							break;
 						}
-						break;
+						DataScanline[X] = ToolTarget::CacheClass::ClampValue(FMath::RoundToInt(FMath::Lerp(OriginalValue, DestValue, BrushValue * UISettings->ToolStrength * Pressure)));
+					}
+					else
+					{
+						float TotalStrength = BrushValue * UISettings->ToolStrength * Pressure * ToolTarget::StrengthMultiplier(this->LandscapeInfo, UISettings->BrushRadius);
+						FNoiseParameter NoiseParam(0, UISettings->NoiseScale, TotalStrength * BrushSizeAdjust);
+						float PaintAmount = ELandscapeToolNoiseMode::Conversion(UISettings->NoiseMode, NoiseParam.NoiseAmount, NoiseParam.Sample(X, Y));
+						DataScanline[X] = ToolTarget::CacheClass::ClampValue(OriginalValue + PaintAmount);
 					}
-					Data[(X - X1) + (Y - Y1)*(1 + X2 - X1)] = ToolTarget::CacheClass::ClampValue(FMath::RoundToInt(FMath::Lerp(OriginalValue, DestValue, It.Value() * UISettings->ToolStrength * Pressure)));
-				}
-				else
-				{
-					float TotalStrength = It.Value() * UISettings->ToolStrength * Pressure * ToolTarget::StrengthMultiplier(this->LandscapeInfo, UISettings->BrushRadius);
-					FNoiseParameter NoiseParam(0, UISettings->NoiseScale, TotalStrength * BrushSizeAdjust);
-					float PaintAmount = ELandscapeToolNoiseMode::Conversion(UISettings->NoiseMode, NoiseParam.NoiseAmount, NoiseParam.Sample(X, Y));
-					Data[(X - X1) + (Y - Y1)*(1 + X2 - X1)] = ToolTarget::CacheClass::ClampValue(OriginalValue + PaintAmount);
 				}
 			}
 		}
diff --git a/Engine/Source/Editor/LandscapeEditor/Private/LandscapeEdModeRampTool.cpp b/Engine/Source/Editor/LandscapeEditor/Private/LandscapeEdModeRampTool.cpp
--- a/Engine/Source/Editor/LandscapeEditor/Private/LandscapeEdModeRampTool.cpp
+++ b/Engine/Source/Editor/LandscapeEditor/Private/LandscapeEdModeRampTool.cpp
@@ -270,12 +270,12 @@ public:
 			if (InKey == EKeys::LeftBracket)
 			{
 				NewValue = OldValue - OldValue * Diff;
-				NewValue = FMath::Min(NewValue, OldValue - 1.f);
+				NewValue = FMath::Min(NewValue, OldValue - 1.0f);
 			}
 			else
 			{
 				NewValue = OldValue + OldValue * Diff;
-				NewValue = FMath::Max(NewValue, OldValue + 1.f);
+				NewValue = FMath::Max(NewValue, OldValue + 1.0f);
 			}
 
 			NewValue = FMath::RoundToFloat(FMath::Clamp(NewValue, SliderMin, SliderMax));
diff --git a/Engine/Source/Editor/LandscapeEditor/Private/LandscapeEdModeTools.h b/Engine/Source/Editor/LandscapeEditor/Private/LandscapeEdModeTools.h
--- a/Engine/Source/Editor/LandscapeEditor/Private/LandscapeEdModeTools.h
+++ b/Engine/Source/Editor/LandscapeEditor/Private/LandscapeEdModeTools.h
@@ -150,7 +150,7 @@ private:
 #endif
 
 template<typename DataType>
-inline void LowPassFilter(int32 X1, int32 Y1, int32 X2, int32 Y2, TMap<FIntPoint, float>& BrushInfo, TArray<DataType>& Data, const float DetailScale, const float ApplyRatio = 1.f)
+inline void LowPassFilter(int32 X1, int32 Y1, int32 X2, int32 Y2, FLandscapeBrushData& BrushInfo, TArray<DataType>& Data, const float DetailScale, const float ApplyRatio = 1.0f)
 {
 #if WITH_KISSFFT
 	// Low-pass filter
@@ -165,12 +165,15 @@ inline void LowPassFilter(int32 X1, int32 Y1, int32 X2, int32 Y2, TMap<FIntPoint
 	kiss_fft_cpx *buf = (kiss_fft_cpx *)KISS_FFT_MALLOC(sizeof(kiss_fft_cpx) * Dims[0] * Dims[1]);
 	kiss_fft_cpx *out = (kiss_fft_cpx *)KISS_FFT_MALLOC(sizeof(kiss_fft_cpx) * Dims[0] * Dims[1]);
 
-	for (int X = X1 + 1; X <= X2 - 1 - FFTWidth % 2; X++)
+	for (int Y = Y1 + 1; Y <= Y2 - 1 - FFTHeight % 2; Y++)
 	{
-		for (int Y = Y1 + 1; Y <= Y2 - 1 - FFTHeight % 2; Y++)
+		auto* DataScanline = Data.GetData() + (Y - Y1) * (X2 - X1 + 1) + (0 - X1);
+		auto* bufScanline = buf + (Y - (Y1 + 1)) * Dims[1] + (0 - (X1 + 1));
+
+		for (int X = X1 + 1; X <= X2 - 1 - FFTWidth % 2; X++)
 		{
-			buf[(X - X1 - 1) + (Y - Y1 - 1)*(Dims[1])].r = Data[(X - X1) + (Y - Y1)*(1 + X2 - X1)];
-			buf[(X - X1 - 1) + (Y - Y1 - 1)*(Dims[1])].i = 0;
+			bufScanline[X].r = DataScanline[X];
+			bufScanline[X].i = 0;
 		}
 	}
 
@@ -180,7 +183,7 @@ inline void LowPassFilter(int32 X1, int32 Y1, int32 X2, int32 Y2, TMap<FIntPoint
 	int32 CenterPos[2] = { Dims[0] >> 1, Dims[1] >> 1 };
 	for (int Y = 0; Y < Dims[0]; Y++)
 	{
-		float DistFromCenter = 0.f;
+		float DistFromCenter = 0.0f;
 		for (int X = 0; X < Dims[1]; X++)
 		{
 			if (Y < CenterPos[0])
@@ -210,7 +213,7 @@ inline void LowPassFilter(int32 X1, int32 Y1, int32 X2, int32 Y2, TMap<FIntPoint
 				}
 			}
 			// High frequency removal
-			float Ratio = 1.f - DetailScale;
+			float Ratio = 1.0f - DetailScale;
 			float Dist = FMath::Min<float>((Dims[0] * Ratio)*(Dims[0] * Ratio), (Dims[1] * Ratio)*(Dims[1] * Ratio));
 			float Filter = 1.0 / (1.0 + DistFromCenter / Dist);
 			out[X + Y*Dims[1]].r *= Filter;
@@ -222,15 +225,24 @@ inline void LowPassFilter(int32 X1, int32 Y1, int32 X2, int32 Y2, TMap<FIntPoint
 	kiss_fftnd(sti, out, buf);
 
 	float Scale = Dims[0] * Dims[1];
-	for (auto It = BrushInfo.CreateConstIterator(); It; ++It)
+	const int32 BrushX1 = FMath::Max<int32>(BrushInfo.GetBounds().Min.X, X1 + 1);
+	const int32 BrushY1 = FMath::Max<int32>(BrushInfo.GetBounds().Min.Y, Y1 + 1);
+	const int32 BrushX2 = FMath::Min<int32>(BrushInfo.GetBounds().Max.X, X2 - FFTWidth % 2);
+	const int32 BrushY2 = FMath::Min<int32>(BrushInfo.GetBounds().Max.Y, Y2 - FFTHeight % 2);
+	for (int32 Y = BrushY1; Y < BrushY2; Y++)
 	{
-		int32 X, Y;
-		ALandscape::UnpackKey(It.Key(), X, Y);
+		const float* BrushScanline = BrushInfo.GetDataPtr(FIntPoint(0, Y));
+		auto* DataScanline = Data.GetData() + (Y - Y1) * (X2 - X1 + 1) + (0 - X1);
+		auto* bufScanline = buf + (Y - (Y1 + 1)) * Dims[1] + (0 - (X1 + 1));
 
-		if (It.Value() > 0.f)
+		for (int32 X = BrushX1; X < BrushX2; X++)
 		{
-			Data[(X - X1) + (Y - Y1)*(1 + X2 - X1)] = FMath::Lerp((float)Data[(X - X1) + (Y - Y1)*(1 + X2 - X1)], buf[(X - X1 - 1) + (Y - Y1 - 1)*(Dims[1])].r / Scale, It.Value() * ApplyRatio);
-			//buf[(X-X1-1) + (Y-Y1-1)*(Dims[1])].r / Scale;
+			const float BrushValue = BrushScanline[X];
+
+			if (BrushValue > 0.0f)
+			{
+				DataScanline[X] = FMath::Lerp((float)DataScanline[X], bufScanline[X].r / Scale, BrushValue * ApplyRatio);
+			}
 		}
 	}
 
@@ -256,6 +268,7 @@ struct TLandscapeEditCache
 	{
 	}
 
+	// X2/Y2 Coordinates are "inclusive" max values
 	void CacheData(int32 X1, int32 Y1, int32 X2, int32 Y2)
 	{
 		if (!Valid)
@@ -390,10 +403,10 @@ struct TLandscapeEditCache
 		AccessorType* P11 = CachedData.Find(ALandscape::MakeKey(X + 1, Y + 1));
 
 		// Search for nearest value if missing data
-		float V00 = P00 ? *P00 : (P10 ? *P10 : (P01 ? *P01 : (P11 ? *P11 : 0.f)));
-		float V10 = P10 ? *P10 : (P00 ? *P00 : (P11 ? *P11 : (P01 ? *P01 : 0.f)));
-		float V01 = P01 ? *P01 : (P00 ? *P00 : (P11 ? *P11 : (P10 ? *P10 : 0.f)));
-		float V11 = P11 ? *P11 : (P10 ? *P10 : (P01 ? *P01 : (P00 ? *P00 : 0.f)));
+		float V00 = P00 ? *P00 : (P10 ? *P10 : (P01 ? *P01 : (P11 ? *P11 : 0.0f)));
+		float V10 = P10 ? *P10 : (P00 ? *P00 : (P11 ? *P11 : (P01 ? *P01 : 0.0f)));
+		float V01 = P01 ? *P01 : (P00 ? *P00 : (P11 ? *P11 : (P10 ? *P10 : 0.0f)));
+		float V11 = P11 ? *P11 : (P10 ? *P10 : (P01 ? *P01 : (P00 ? *P00 : 0.0f)));
 
 		return FMath::Lerp(
 			FMath::Lerp(V00, V10, LandscapeX - X),
@@ -409,15 +422,15 @@ struct TLandscapeEditCache
 		AccessorType* P11 = CachedData.Find(ALandscape::MakeKey(X + 1, Y + 1));
 
 		// Search for nearest value if missing data
-		float V00 = P00 ? *P00 : (P10 ? *P10 : (P01 ? *P01 : (P11 ? *P11 : 0.f)));
-		float V10 = P10 ? *P10 : (P00 ? *P00 : (P11 ? *P11 : (P01 ? *P01 : 0.f)));
-		float V01 = P01 ? *P01 : (P00 ? *P00 : (P11 ? *P11 : (P10 ? *P10 : 0.f)));
-		float V11 = P11 ? *P11 : (P10 ? *P10 : (P01 ? *P01 : (P00 ? *P00 : 0.f)));
+		float V00 = P00 ? *P00 : (P10 ? *P10 : (P01 ? *P01 : (P11 ? *P11 : 0.0f)));
+		float V10 = P10 ? *P10 : (P00 ? *P00 : (P11 ? *P11 : (P01 ? *P01 : 0.0f)));
+		float V01 = P01 ? *P01 : (P00 ? *P00 : (P11 ? *P11 : (P10 ? *P10 : 0.0f)));
+		float V11 = P11 ? *P11 : (P10 ? *P10 : (P01 ? *P01 : (P00 ? *P00 : 0.0f)));
 
-		FVector Vert00 = FVector(0.f, 0.f, V00);
-		FVector Vert01 = FVector(0.f, 1.f, V01);
-		FVector Vert10 = FVector(1.f, 0.f, V10);
-		FVector Vert11 = FVector(1.f, 1.f, V11);
+		FVector Vert00 = FVector(0.0f, 0.0f, V00);
+		FVector Vert01 = FVector(0.0f, 1.0f, V01);
+		FVector Vert10 = FVector(1.0f, 0.0f, V10);
+		FVector Vert11 = FVector(1.0f, 1.0f, V11);
 
 		FVector FaceNormal1 = ((Vert00 - Vert10) ^ (Vert10 - Vert11)).GetSafeNormal();
 		FVector FaceNormal2 = ((Vert11 - Vert01) ^ (Vert01 - Vert00)).GetSafeNormal();
@@ -449,6 +462,7 @@ struct TLandscapeEditCache
 		return Value == 0;
 	}
 
+	// X2/Y2 Coordinates are "inclusive" max values
 	bool GetCachedData(int32 X1, int32 Y1, int32 X2, int32 Y2, TArray<AccessorType>& OutData)
 	{
 		int32 NumSamples = (1 + X2 - X1)*(1 + Y2 - Y1);
@@ -475,6 +489,7 @@ struct TLandscapeEditCache
 		return bHasNonZero;
 	}
 
+	// X2/Y2 Coordinates are "inclusive" max values
 	void SetCachedData(int32 X1, int32 Y1, int32 X2, int32 Y2, TArray<AccessorType>& Data, ELandscapeLayerPaintingRestriction::Type PaintingRestriction = ELandscapeLayerPaintingRestriction::None)
 	{
 		// Update cache
@@ -491,6 +506,7 @@ struct TLandscapeEditCache
 	}
 
 	// Get the original data before we made any changes with the SetCachedData interface.
+	// X2/Y2 Coordinates are "inclusive" max values
 	void GetOriginalData(int32 X1, int32 Y1, int32 X2, int32 Y2, TArray<AccessorType>& OutOriginalData)
 	{
 		int32 NumSamples = (1 + X2 - X1)*(1 + Y2 - Y1);
@@ -518,6 +534,7 @@ struct TLandscapeEditCache
 protected:
 	Accessor& DataAccess;
 private:
+	// X2/Y2 Coordinates are "inclusive" max values
 	void CacheOriginalData(int32 X1, int32 Y1, int32 X2, int32 Y2)
 	{
 		for (int32 Y = Y1; Y <= Y2; Y++)
@@ -626,7 +643,7 @@ struct FHeightmapAccessor
 					AInstancedFoliageActor* IFA = AInstancedFoliageActor::GetInstancedFoliageActorForLevel(CollisionComponent->GetComponentLevel());
 					if (IFA)
 					{
-						CollisionComponent->SnapFoliageInstances(*IFA, PreUpdateLocalBoxes[Index].TransformBy(LandscapeInfo->GetLandscapeProxy()->LandscapeActorToWorld().ToMatrixWithScale()).ExpandBy(1.f));
+						CollisionComponent->SnapFoliageInstances(*IFA, PreUpdateLocalBoxes[Index].TransformBy(LandscapeInfo->GetLandscapeProxy()->LandscapeActorToWorld().ToMatrixWithScale()).ExpandBy(1.0f));
 					}
 				}
 			}
@@ -722,7 +739,7 @@ struct FXYOffsetmapAccessor
 				FVector* Value = Data.Find(ALandscape::MakeKey(X, Y));
 				if (Value)
 				{
-					Value->Z = ((float)NewHeights.FindRef(ALandscape::MakeKey(X, Y)) - 32768.f) * LANDSCAPE_ZSCALE;
+					Value->Z = ((float)NewHeights.FindRef(ALandscape::MakeKey(X, Y)) - 32768.0f) * LANDSCAPE_ZSCALE;
 				}
 			}
 		}
@@ -741,7 +758,7 @@ struct FXYOffsetmapAccessor
 				FVector* Value = Data.Find(ALandscape::MakeKey(X, Y));
 				if (Value)
 				{
-					Value->Z = ((float)NewHeights.FindRef(ALandscape::MakeKey(X, Y)) - 32768.f) * LANDSCAPE_ZSCALE;
+					Value->Z = ((float)NewHeights.FindRef(ALandscape::MakeKey(X, Y)) - 32768.0f) * LANDSCAPE_ZSCALE;
 				}
 			}
 		}
@@ -762,7 +779,7 @@ struct FXYOffsetmapAccessor
 			{
 				for (int32 X = X1; X <= X2; ++X)
 				{
-					NewHeights[X - X1 + (Y - Y1) * (X2 - X1 + 1)] = FMath::Clamp<uint16>(Data[(X - X1 + (Y - Y1) * (X2 - X1 + 1))].Z * LANDSCAPE_INV_ZSCALE + 32768.f, 0, 65535);
+					NewHeights[X - X1 + (Y - Y1) * (X2 - X1 + 1)] = FMath::Clamp<uint16>(Data[(X - X1 + (Y - Y1) * (X2 - X1 + 1))].Z * LANDSCAPE_INV_ZSCALE + 32768.0f, 0, 65535);
 				}
 			}
 
@@ -801,7 +818,7 @@ struct FXYOffsetmapAccessor
 				{
 					ULandscapeHeightfieldCollisionComponent* CollisionComponent = CollisionComponents[Index];
 					AInstancedFoliageActor* IFA = AInstancedFoliageActor::GetInstancedFoliageActorForLevel(CollisionComponent->GetComponentLevel());
-					CollisionComponent->SnapFoliageInstances(*IFA, PreUpdateLocalBoxes[Index].TransformBy(LandscapeInfo->GetLandscapeProxy()->LandscapeActorToWorld().ToMatrixWithScale()).ExpandBy(1.f));
+					CollisionComponent->SnapFoliageInstances(*IFA, PreUpdateLocalBoxes[Index].TransformBy(LandscapeInfo->GetLandscapeProxy()->LandscapeActorToWorld().ToMatrixWithScale()).ExpandBy(1.0f));
 				}
 			}
 			else
@@ -1185,20 +1202,20 @@ struct FHeightmapToolTarget
 			// Adjust strength based on brush size and drawscale, so strength 1 = one hemisphere
 			return BrushRadius * LANDSCAPE_INV_ZSCALE / (LandscapeInfo->DrawScale.Z);
 		}
-		return 5.f * LANDSCAPE_INV_ZSCALE;
+		return 5.0f * LANDSCAPE_INV_ZSCALE;
 	}
 
 	static FMatrix ToWorldMatrix(ULandscapeInfo* LandscapeInfo)
 	{
-		FMatrix Result = FTranslationMatrix(FVector(0, 0, -32768.f));
-		Result *= FScaleMatrix(FVector(1.f, 1.f, LANDSCAPE_ZSCALE) * LandscapeInfo->DrawScale);
+		FMatrix Result = FTranslationMatrix(FVector(0, 0, -32768.0f));
+		Result *= FScaleMatrix(FVector(1.0f, 1.0f, LANDSCAPE_ZSCALE) * LandscapeInfo->DrawScale);
 		return Result;
 	}
 
 	static FMatrix FromWorldMatrix(ULandscapeInfo* LandscapeInfo)
 	{
-		FMatrix Result = FScaleMatrix(FVector(1.f, 1.f, LANDSCAPE_INV_ZSCALE) / (LandscapeInfo->DrawScale));
-		Result *= FTranslationMatrix(FVector(0, 0, 32768.f));
+		FMatrix Result = FScaleMatrix(FVector(1.0f, 1.0f, LANDSCAPE_INV_ZSCALE) / (LandscapeInfo->DrawScale));
+		Result *= FTranslationMatrix(FVector(0, 0, 32768.0f));
 		return Result;
 	}
 };
@@ -1211,7 +1228,7 @@ struct FWeightmapToolTarget
 
 	static float StrengthMultiplier(ULandscapeInfo* LandscapeInfo, float BrushRadius)
 	{
-		return 255.f;
+		return 255.0f;
 	}
 
 	static FMatrix ToWorldMatrix(ULandscapeInfo* LandscapeInfo) { return FMatrix::Identity; }
@@ -1262,7 +1279,7 @@ public:
 
 		// Save the mouse position
 		LastMousePosition = FVector2D(InHitLocation);
-		MousePositions.Emplace(InHitLocation.X, InHitLocation.Y, IsShiftDown(ViewportClient->Viewport));
+		MousePositions.Emplace(LastMousePosition, ViewportClient ? IsShiftDown(ViewportClient->Viewport) : false); // Copy tool sometimes activates without a specific viewport via ctrl+c hotkey
 		TimeSinceLastMouseMove = 0.0f;
 
 		ToolStroke->Apply(ViewportClient, EdMode->CurrentBrush, EdMode->UISettings, MousePositions);
@@ -1275,14 +1292,14 @@ public:
 	{
 		if (bToolActive)
 		{
-			if (MousePositions.Num())
+			if (MousePositions.Num() > 0)
 			{
 				ToolStroke->Apply(ViewportClient, EdMode->CurrentBrush, EdMode->UISettings, MousePositions);
 				MousePositions.Empty(1);
 			}
 			else if (TStrokeClass::UseContinuousApply && TimeSinceLastMouseMove >= 0.25f)
 			{
-				MousePositions.Emplace(LastMousePosition.X, LastMousePosition.Y, IsShiftDown(ViewportClient->Viewport));
+				MousePositions.Emplace(LastMousePosition, IsShiftDown(ViewportClient->Viewport));
 				ToolStroke->Apply(ViewportClient, EdMode->CurrentBrush, EdMode->UISettings, MousePositions);
 				MousePositions.Empty(1);
 			}
@@ -1317,8 +1334,11 @@ public:
 			if (bToolActive)
 			{
 				// Save the mouse position
-				LastMousePosition = FVector2D(HitLocation);
-				MousePositions.Emplace(HitLocation.X, HitLocation.Y, IsShiftDown(ViewportClient->Viewport));
+				if (MousePositions.Num() == 0 || LastMousePosition != FVector2D(HitLocation))
+				{
+					LastMousePosition = FVector2D(HitLocation);
+					MousePositions.Emplace(LastMousePosition, IsShiftDown(ViewportClient->Viewport));
+				}
 				TimeSinceLastMouseMove = 0.0f;
 			}
 		}
diff --git a/Engine/Source/Editor/LandscapeEditor/Private/LandscapeEdModeXYOffset.cpp b/Engine/Source/Editor/LandscapeEditor/Private/LandscapeEdModeXYOffset.cpp
--- a/Engine/Source/Editor/LandscapeEditor/Private/LandscapeEdModeXYOffset.cpp
+++ b/Engine/Source/Editor/LandscapeEditor/Private/LandscapeEdModeXYOffset.cpp
@@ -17,7 +17,7 @@ namespace
 {
 	FORCEINLINE FVector4 GetWorldPos(const FMatrix& LocalToWorld, FVector2D LocalXY, uint16 Height, FVector2D XYOffset)
 	{
-		return LocalToWorld.TransformPosition(FVector(LocalXY.X + XYOffset.X, LocalXY.Y + XYOffset.Y, ((float)Height - 32768.f) * LANDSCAPE_ZSCALE));
+		return LocalToWorld.TransformPosition(FVector(LocalXY.X + XYOffset.X, LocalXY.Y + XYOffset.Y, ((float)Height - 32768.0f) * LANDSCAPE_ZSCALE));
 	}
 
 	FORCEINLINE FVector4 GetWorldPos(const FMatrix& LocalToWorld, FVector2D LocalXY, FVector XYOffsetVector)
@@ -67,27 +67,31 @@ public:
 		}
 
 		// Get list of verts to update
-		TMap<FIntPoint, float> BrushInfo;
-		int32 X1, Y1, X2, Y2;
-		if (!Brush->ApplyBrush(MousePositions, BrushInfo, X1, Y1, X2, Y2))
+		// TODO - only retrieve bounds as we don't need the data
+		// or use the brush data?
+		FLandscapeBrushData BrushInfo = Brush->ApplyBrush(MousePositions);
+		if (!BrushInfo)
 		{
 			return;
 		}
 
+		int32 X1, Y1, X2, Y2;
+		BrushInfo.GetInclusiveBounds(X1, Y1, X2, Y2);
+
 		//LandscapeInfo->Modify();
 		//LandscapeInfo->LandscapeProxy->Modify();
 
 		FVector DrawScale3D = LandscapeInfo->DrawScale;
 		// Tablet pressure
-		float Pressure = ViewportClient->Viewport->IsPenActive() ? ViewportClient->Viewport->GetTabletPressure() : 1.f;
+		float Pressure = ViewportClient->Viewport->IsPenActive() ? ViewportClient->Viewport->GetTabletPressure() : 1.0f;
 
 		// expand the area by one vertex in each direction to ensure normals are calculated correctly
 		/*
-				X1 -= 1;
-				Y1 -= 1;
-				X2 += 1;
-				Y2 += 1;
-				*/
+			X1 -= 1;
+			Y1 -= 1;
+			X2 += 1;
+			Y2 += 1;
+		*/
 		{
 			int32 ValidX1, ValidX2, ValidY1, ValidY2;
 			ValidX1 = ValidY1 = INT_MAX;
@@ -124,7 +128,7 @@ public:
 			return;
 		}
 
-		const float AreaResolution = LANDSCAPE_XYOFFSET_SCALE; //1.f/256.f;
+		const float AreaResolution = LANDSCAPE_XYOFFSET_SCALE; //1.0f/256.0f;
 
 		Cache.CacheData(X1, Y1, X2, Y2);
 
@@ -136,7 +140,7 @@ public:
 		// Retopologize algorithm...
 		{
 			// Calculate surface world space area without missing area...
-			float TotalArea = 0.f;
+			float TotalArea = 0.0f;
 			int32 QuadNum = 0;
 			const int32 MaxIterNum = 300;
 
@@ -236,17 +240,17 @@ public:
 
 				float AverageArea = RemainArea / RemainQuads;
 				float AreaBaseError = AverageArea * 0.5f;
-				float TotalLineArea = 0.f;
+				float TotalLineArea = 0.0f;
 				float TargetLineArea = AverageArea * QuadY[Y - Y1];
 				float YOffset = Y + 1, PreYOffset = Y + 1; // Need to be bigger than previous Y
 				float StepSize = FPlatformMath::Sqrt(2) * 0.25f;
 				float LineAreaDiff = FLT_MAX; //Abs(TotalLineArea - TargetLineArea);
 				int32 IterNum = 0;
-				float TotalHeightError = 0.f;
+				float TotalHeightError = 0.0f;
 
 				while (NewXYOffset[(Y - Y1) * (X2 - X1 + 1)].Y + Y > XYOffsetVectorData[(FPlatformMath::FloorToInt(YOffset) - Y1) * (X2 - X1 + 1)].Y + FPlatformMath::FloorToInt(YOffset))
 				{
-					YOffset = YOffset + 1.f;
+					YOffset = YOffset + 1.0f;
 					if (YOffset >= Y2)
 					{
 						YOffset = Y2;
@@ -258,8 +262,8 @@ public:
 				while (FMath::Abs(TotalLineArea - TargetLineArea) > AreaErrorThreshold)
 				{
 					IterNum++;
-					TotalLineArea = 0.f;
-					TotalHeightError = 0.f;
+					TotalLineArea = 0.0f;
+					TotalHeightError = 0.0f;
 					//for (int32 X = X1; X < X2; ++X)
 					for (int32 X = MinX[Y - Y1]; X < MaxX[Y - Y1]; ++X)
 					{
@@ -383,17 +387,17 @@ public:
 
 				float AverageArea = RemainArea / RemainQuads;
 				float AreaBaseError = AverageArea * 0.5f;
-				float TotalLineArea = 0.f;
+				float TotalLineArea = 0.0f;
 				float TargetLineArea = AverageArea * QuadX[X - X1];
 				float XOffset = X + 1, PreXOffset = X + 1; // Need to be bigger than previous Y
 				float StepSize = FMath::Sqrt(2) * 0.25f;
 				float LineAreaDiff = FLT_MAX; // Abs(TotalLineArea - TargetLineArea);
 				int32 IterNum = 0;
-				float TotalHeightError = 0.f;
+				float TotalHeightError = 0.0f;
 
 				while (NewXYOffset[X - X1].X + X > NewYOffsets[FMath::FloorToInt(XOffset) - X1].X + FMath::FloorToFloat(XOffset))
 				{
-					XOffset = XOffset + 1.f;
+					XOffset = XOffset + 1.0f;
 					if (XOffset >= X2)
 					{
 						XOffset = X2;
@@ -404,7 +408,7 @@ public:
 
 				while (FMath::Abs(TotalLineArea - TargetLineArea) > AreaErrorThreshold)
 				{
-					TotalLineArea = 0.f;
+					TotalLineArea = 0.0f;
 					IterNum++;
 					for (int32 Y = MinY[X - X1]; Y < MaxY[X - X1]; ++Y)
 					{
@@ -527,9 +531,9 @@ public:
 				FVector2D TransformedLocal(FMath::Abs(X - W * 0.5f), FMath::Abs(Y - H * 0.5f) * (W / H));
 				float Cos = FMath::Abs(TransformedLocal.X) / TransformedLocal.Size();
 				float Sin = FMath::Abs(TransformedLocal.Y) / TransformedLocal.Size();
-				float RatioX = FalloffRadius > 0.f ? 1.f - FMath::Clamp((FMath::Abs(TransformedLocal.X) - Cos*SquareRadius) / FalloffRadius, 0.f, 1.f) : 1.f;
-				float RatioY = FalloffRadius > 0.f ? 1.f - FMath::Clamp((FMath::Abs(TransformedLocal.Y) - Sin*SquareRadius) / FalloffRadius, 0.f, 1.f) : 1.f;
-				float Ratio = TransformedLocal.Size() > SquareRadius ? RatioX * RatioY : 1.f; //TransformedLocal.X / LW * TransformedLocal.Y / LW;
+				float RatioX = FalloffRadius > 0.0f ? 1.0f - FMath::Clamp((FMath::Abs(TransformedLocal.X) - Cos*SquareRadius) / FalloffRadius, 0.0f, 1.0f) : 1.0f;
+				float RatioY = FalloffRadius > 0.0f ? 1.0f - FMath::Clamp((FMath::Abs(TransformedLocal.Y) - Sin*SquareRadius) / FalloffRadius, 0.0f, 1.0f) : 1.0f;
+				float Ratio = TransformedLocal.Size() > SquareRadius ? RatioX * RatioY : 1.0f; //TransformedLocal.X / LW * TransformedLocal.Y / LW;
 				float PaintAmount = Ratio*Ratio*(3 - 2 * Ratio);
 
 				XYOffsetVectorData[Index] = FMath::Lerp(XYOffsetVectorData[Index], NewXYOffset[Index], PaintAmount);
diff --git a/Engine/Source/Editor/LandscapeEditor/Private/LandscapeEditorObject.cpp b/Engine/Source/Editor/LandscapeEditor/Private/LandscapeEditorObject.cpp
--- a/Engine/Source/Editor/LandscapeEditor/Private/LandscapeEditorObject.cpp
+++ b/Engine/Source/Editor/LandscapeEditor/Private/LandscapeEditorObject.cpp
@@ -28,7 +28,7 @@ ULandscapeEditorObject::ULandscapeEditorObject(const FObjectInitializer& ObjectI
 	, RampWidth(2000)
 	, RampSideFalloff(0.4f)
 
-	, SmoothFilterKernelScale(1.0f)
+	, SmoothFilterKernelSize(4)
 	, bDetailSmooth(false)
 	, DetailScale(0.3f)
 
@@ -42,7 +42,7 @@ ULandscapeEditorObject::ULandscapeEditorObject(const FObjectInitializer& ObjectI
 	, SedimentCapacity(0.3f)
 	, HErodeIterationNum(75)
 	, RainDistMode(ELandscapeToolHydroErosionMode::Both)
-	, RainDistScale(60.f)
+	, RainDistScale(60.0f)
 	, bHErosionDetailSmooth(true)
 	, HErosionDetailScale(0.01f)
 
@@ -223,7 +223,7 @@ void ULandscapeEditorObject::Load()
 	NoiseMode = (ELandscapeToolNoiseMode::Type)InNoiseMode;
 	GConfig->GetFloat(TEXT("LandscapeEdit"), TEXT("NoiseScale"), NoiseScale, GEditorUserSettingsIni);
 
-	GConfig->GetFloat(TEXT("LandscapeEdit"), TEXT("SmoothFilterKernelScale"), SmoothFilterKernelScale, GEditorUserSettingsIni);
+	GConfig->GetInt(TEXT("LandscapeEdit"), TEXT("SmoothFilterKernelSize"), SmoothFilterKernelSize, GEditorUserSettingsIni);
 	GConfig->GetFloat(TEXT("LandscapeEdit"), TEXT("DetailScale"), DetailScale, GEditorUserSettingsIni);
 	bool InbDetailSmooth = bDetailSmooth;
 	GConfig->GetBool(TEXT("LandscapeEdit"), TEXT("bDetailSmooth"), InbDetailSmooth, GEditorUserSettingsIni);
@@ -318,7 +318,7 @@ void ULandscapeEditorObject::Save()
 
 	GConfig->SetInt(TEXT("LandscapeEdit"), TEXT("NoiseMode"), (int32)NoiseMode, GEditorUserSettingsIni);
 	GConfig->SetFloat(TEXT("LandscapeEdit"), TEXT("NoiseScale"), NoiseScale, GEditorUserSettingsIni);
-	GConfig->SetFloat(TEXT("LandscapeEdit"), TEXT("SmoothFilterKernelScale"), SmoothFilterKernelScale, GEditorUserSettingsIni);
+	GConfig->SetInt(TEXT("LandscapeEdit"), TEXT("SmoothFilterKernelSize"), SmoothFilterKernelSize, GEditorUserSettingsIni);
 	GConfig->SetFloat(TEXT("LandscapeEdit"), TEXT("DetailScale"), DetailScale, GEditorUserSettingsIni);
 	GConfig->SetBool(TEXT("LandscapeEdit"), TEXT("bDetailSmooth"), bDetailSmooth, GEditorUserSettingsIni);
 
diff --git a/Engine/Source/Editor/LandscapeEditor/Public/LandscapeEdMode.h b/Engine/Source/Editor/LandscapeEditor/Public/LandscapeEdMode.h
--- a/Engine/Source/Editor/LandscapeEditor/Public/LandscapeEdMode.h
+++ b/Engine/Source/Editor/LandscapeEditor/Public/LandscapeEdMode.h
@@ -22,31 +22,82 @@ template<typename TargetType> class FLandscapeToolCopyPaste;
 struct FLandscapeToolMousePosition
 {
 	// Stored in heightmap space.
-	float PositionX, PositionY;
+	FVector2D Position;
 	bool bShiftDown;
 
-	FLandscapeToolMousePosition(float InPositionX, float InPositionY, bool InbShiftDown)
-		: PositionX(InPositionX)
-		, PositionY(InPositionY)
+	FLandscapeToolMousePosition(FVector2D InPosition, bool InbShiftDown)
+		: Position(InPosition)
 		, bShiftDown(InbShiftDown)
 	{
 	}
 };
 
-class FLandscapeBrush : public FGCObject
+enum class ELandscapeBrushType
 {
+	Normal = 0,
+	Alpha,
+	Component,
+	Gizmo,
+	Splines
+};
+
+class FLandscapeBrushData
+{
+protected:
+	FIntRect Bounds;
+	TArray<float> BrushAlpha;
+
 public:
-	enum EBrushType
+	FLandscapeBrushData()
+		: Bounds()
 	{
-		BT_Normal = 0,
-		BT_Alpha,
-		BT_Component,
-		BT_Gizmo,
-		BT_Splines
-	};
+	}
+
+	FLandscapeBrushData(FIntRect InBounds)
+		: Bounds(InBounds)
+	{
+		BrushAlpha.SetNumZeroed(Bounds.Area());
+	}
+
+	FIntRect GetBounds() const
+	{
+		return Bounds;
+	}
+
+	// For compatibility with older landscape code that uses inclusive bounds in 4 int32s
+	void GetInclusiveBounds(int32& X1, int32& Y1, int32& X2, int32& Y2) const
+	{
+		X1 = Bounds.Min.X;
+		Y1 = Bounds.Min.Y;
+		X2 = Bounds.Max.X - 1;
+		Y2 = Bounds.Max.Y - 1;
+	}
+
+	float* GetDataPtr(FIntPoint Position)
+	{
+		return BrushAlpha.GetData() + (Position.Y - Bounds.Min.Y) * Bounds.Width() + (Position.X - Bounds.Min.X);
+	}
+	const float* GetDataPtr(FIntPoint Position) const
+	{
+		return BrushAlpha.GetData() + (Position.Y - Bounds.Min.Y) * Bounds.Width() + (Position.X - Bounds.Min.X);
+	}
+
+	FORCEINLINE_EXPLICIT_OPERATOR_BOOL() const
+	{
+		return BrushAlpha.Num() != 0;
+	}
+
+	FORCEINLINE bool operator!() const
+	{
+		return !(bool)*this;
+	}
+};
 
+class FLandscapeBrush : public FGCObject
+{
+public:
 	virtual void MouseMove(float LandscapeX, float LandscapeY) = 0;
-	virtual bool ApplyBrush(const TArray<FLandscapeToolMousePosition>& MousePositions, TMap<FIntPoint, float>& OutBrush, int32& OutX1, int32& OutY1, int32& OutX2, int32& OutY2) = 0;
+	virtual FLandscapeBrushData ApplyBrush(const TArray<FLandscapeToolMousePosition>& MousePositions) = 0;
 	virtual bool InputKey(FEditorViewportClient* InViewportClient, FViewport* InViewport, FKey InKey, EInputEvent InEvent) { return false; }
 	virtual void Tick(FEditorViewportClient* ViewportClient, float DeltaTime) {};
 	virtual void BeginStroke(float LandscapeX, float LandscapeY, class FLandscapeTool* CurrentTool);
@@ -57,7 +108,7 @@ public:
 	virtual UMaterialInterface* GetBrushMaterial() { return NULL; }
 	virtual const TCHAR* GetBrushName() = 0;
 	virtual FText GetDisplayName() = 0;
-	virtual EBrushType GetBrushType() { return BT_Normal; }
+	virtual ELandscapeBrushType GetBrushType() { return ELandscapeBrushType::Normal; }
 
 	// FGCObject interface
 	virtual void AddReferencedObjects(FReferenceCollector& Collector) override {}
@@ -100,8 +151,8 @@ namespace ELandscapeToolTargetTypeMask
 {
 	enum Type : uint8
 	{
-		Heightmap = 1 << ELandscapeToolTargetType::Heightmap,
-		Weightmap = 1 << ELandscapeToolTargetType::Weightmap,
+		Heightmap  = 1 << ELandscapeToolTargetType::Heightmap,
+		Weightmap  = 1 << ELandscapeToolTargetType::Weightmap,
 		Visibility = 1 << ELandscapeToolTargetType::Visibility,
 
 		NA = 0,
@@ -161,17 +212,18 @@ struct FLandscapeToolTarget
 	}
 };
 
+enum class ELandscapeToolType
+{
+	Normal = 0,
+	Mask,
+};
+
 /**
  * FLandscapeTool
  */
 class FLandscapeTool
 {
 public:
-	enum EToolType
-	{
-		TT_Normal = 0,
-		TT_Mask,
-	};
 	virtual void EnterTool() {}
 	virtual void ExitTool() {}
 	virtual bool BeginTool(FEditorViewportClient* ViewportClient, const FLandscapeToolTarget& Target, const FVector& InHitLocation) = 0;
@@ -210,7 +262,7 @@ public:
 
 	// Functions which doesn't need Viewport data...
 	virtual void Process(int32 Index, int32 Arg) {}
-	virtual EToolType GetToolType() { return TT_Normal; }
+	virtual ELandscapeToolType GetToolType() { return ELandscapeToolType::Normal; }
 	virtual ELandscapeToolTargetTypeMask::Type GetSupportedTargetTypes() { return ELandscapeToolTargetTypeMask::NA; };
 
 	int32					PreviousBrushIndex;
diff --git a/Engine/Source/Runtime/Landscape/Classes/Landscape.h b/Engine/Source/Runtime/Landscape/Classes/Landscape.h
--- a/Engine/Source/Runtime/Landscape/Classes/Landscape.h
+++ b/Engine/Source/Runtime/Landscape/Classes/Landscape.h
@@ -51,11 +51,13 @@ class ALandscape : public ALandscapeProxy
 	LANDSCAPE_API bool HasAllComponent(); // determine all component is in this actor
 	
 	// Include Components with overlapped vertices
+	// X2/Y2 Coordinates are "inclusive" max values
 	LANDSCAPE_API static void CalcComponentIndicesOverlap(const int32 X1, const int32 Y1, const int32 X2, const int32 Y2, const int32 ComponentSizeQuads, 
 		int32& ComponentIndexX1, int32& ComponentIndexY1, int32& ComponentIndexX2, int32& ComponentIndexY2);
 
 	// Exclude Components with overlapped vertices
-	static void CalcComponentIndices(const int32 X1, const int32 Y1, const int32 X2, const int32 Y2, const int32 ComponentSizeQuads, 
+	// X2/Y2 Coordinates are "inclusive" max values
+	LANDSCAPE_API static void CalcComponentIndicesNoOverlap(const int32 X1, const int32 Y1, const int32 X2, const int32 Y2, const int32 ComponentSizeQuads,
 		int32& ComponentIndexX1, int32& ComponentIndexY1, int32& ComponentIndexX2, int32& ComponentIndexY2);
 
 	static void SplitHeightmap(ULandscapeComponent* Comp, bool bMoveToCurrentLevel = false);
diff --git a/Engine/Source/Runtime/Landscape/Private/LandscapeEdit.cpp b/Engine/Source/Runtime/Landscape/Private/LandscapeEdit.cpp
--- a/Engine/Source/Runtime/Landscape/Private/LandscapeEdit.cpp
+++ b/Engine/Source/Runtime/Landscape/Private/LandscapeEdit.cpp
@@ -154,7 +154,7 @@ ULandscapeMaterialInstanceConstant* ALandscapeProxy::GetLayerThumbnailMIC(UMater
 	MaterialInstance->SetParentEditorOnly(CombinationMaterialInstance);
 	MaterialInstance->bIsLayerThumbnail = true;
 
-	FLinearColor Mask(1.f, 0.f, 0.f, 0.f);
+	FLinearColor Mask(1.0f, 0.0f, 0.0f, 0.0f);
 	MaterialInstance->SetVectorParameterValueEditorOnly(FName(*FString::Printf(TEXT("LayerMask_%s"), *LayerName.ToString())), Mask);
 	MaterialInstance->SetTextureParameterValueEditorOnly(FName(TEXT("Weightmap0")), ThumbnailWeightmap);
 	MaterialInstance->SetTextureParameterValueEditorOnly(FName(TEXT("Heightmap")), ThumbnailHeightmap);
@@ -256,10 +256,10 @@ void ULandscapeComponent::UpdateMaterialInstances()
 		MaterialInstance->SetParentEditorOnly(CombinationMaterialInstance);
 
 		FLinearColor Masks[4];
-		Masks[0] = FLinearColor(1.f, 0.f, 0.f, 0.f);
-		Masks[1] = FLinearColor(0.f, 1.f, 0.f, 0.f);
-		Masks[2] = FLinearColor(0.f, 0.f, 1.f, 0.f);
-		Masks[3] = FLinearColor(0.f, 0.f, 0.f, 1.f);
+		Masks[0] = FLinearColor(1.0f, 0.0f, 0.0f, 0.0f);
+		Masks[1] = FLinearColor(0.0f, 1.0f, 0.0f, 0.0f);
+		Masks[2] = FLinearColor(0.0f, 0.0f, 1.0f, 0.0f);
+		Masks[3] = FLinearColor(0.0f, 0.0f, 0.0f, 1.0f);
 
 		// Set the layer mask
 		for (int32 AllocIdx = 0; AllocIdx < WeightmapLayerAllocations.Num(); AllocIdx++)
@@ -1011,7 +1011,7 @@ void ULandscapeComponent::GenerateHeightmapMips(TArray<FColor*>& HeightmapTextur
 			int32 PrevMipSubY2 = ComponentY2 - SubsectionSizeQuads*SubsectionY;
 
 			int32 PrevMipSubsectionSizeQuads = SubsectionSizeQuads;
-			float InvPrevMipSubsectionSizeQuads = 1.f / (float)SubsectionSizeQuads;
+			float InvPrevMipSubsectionSizeQuads = 1.0f / (float)SubsectionSizeQuads;
 
 			int32 PrevMipSizeU = HeightmapSizeU;
 			int32 PrevMipSizeV = HeightmapSizeV;
@@ -1025,7 +1025,7 @@ void ULandscapeComponent::GenerateHeightmapMips(TArray<FColor*>& HeightmapTextur
 				int32 MipSizeV = HeightmapSizeV >> Mip;
 
 				int32 MipSubsectionSizeQuads = ((SubsectionSizeQuads + 1) >> Mip) - 1;
-				float InvMipSubsectionSizeQuads = 1.f / (float)MipSubsectionSizeQuads;
+				float InvMipSubsectionSizeQuads = 1.0f / (float)MipSubsectionSizeQuads;
 
 				int32 MipHeightmapOffsetX = HeightmapOffsetX >> Mip;
 				int32 MipHeightmapOffsetY = HeightmapOffsetY >> Mip;
@@ -1340,7 +1340,7 @@ void ULandscapeComponent::UpdateMipsTempl(int32 InNumSubsections, int32 InSubsec
 			int32 PrevMipSubY2 = ComponentY2 - InSubsectionSizeQuads*SubsectionY;
 
 			int32 PrevMipSubsectionSizeQuads = InSubsectionSizeQuads;
-			float InvPrevMipSubsectionSizeQuads = 1.f / (float)InSubsectionSizeQuads;
+			float InvPrevMipSubsectionSizeQuads = 1.0f / (float)InSubsectionSizeQuads;
 
 			int32 PrevMipSizeU = WeightmapSizeU;
 			int32 PrevMipSizeV = WeightmapSizeV;
@@ -1351,7 +1351,7 @@ void ULandscapeComponent::UpdateMipsTempl(int32 InNumSubsections, int32 InSubsec
 				int32 MipSizeV = WeightmapSizeV >> Mip;
 
 				int32 MipSubsectionSizeQuads = ((InSubsectionSizeQuads + 1) >> Mip) - 1;
-				float InvMipSubsectionSizeQuads = 1.f / (float)MipSubsectionSizeQuads;
+				float InvMipSubsectionSizeQuads = 1.0f / (float)MipSubsectionSizeQuads;
 
 				// Area to update in current mip level coords
 				int32 MipSubX1 = FMath::FloorToInt((float)MipSubsectionSizeQuads * (float)PrevMipSubX1 * InvPrevMipSubsectionSizeQuads);
@@ -1491,15 +1491,15 @@ float ULandscapeComponent::GetLayerWeightAtLocation(const FVector& InLocation, U
 	}
 
 	// Fill the cache if necessary
-		if (LayerCache->Num() == 0)
-		{
+	if (LayerCache->Num() == 0)
+	{
 		FLandscapeComponentDataInterface CDI(this);
-			if (!CDI.GetWeightmapTextureData(LayerInfo, *LayerCache))
-			{
-				// no data for this layer for this component.
-				return 0.f;
-			}
+		if (!CDI.GetWeightmapTextureData(LayerInfo, *LayerCache))
+		{
+			// no data for this layer for this component.
+			return 0.0f;
 		}
+	}
 
 	// Find location
 	// TODO: Root landscape isn't always loaded, would Proxy suffice?
@@ -1523,10 +1523,10 @@ float ULandscapeComponent::GetLayerWeightAtLocation(const FVector& InLocation, U
 	int32 IdxY2 = FMath::Min<int32>(((Y2 / SubsectionSizeQuads) * (SubsectionSizeQuads + 1)) + (Y2 % SubsectionSizeQuads), Stride - 1);
 
 	// sample
-	float Sample11 = (float)((*LayerCache)[IdxX1 + Stride*IdxY1]) / 255.f;
-	float Sample21 = (float)((*LayerCache)[IdxX2 + Stride*IdxY1]) / 255.f;
-	float Sample12 = (float)((*LayerCache)[IdxX1 + Stride*IdxY2]) / 255.f;
-	float Sample22 = (float)((*LayerCache)[IdxX2 + Stride*IdxY2]) / 255.f;
+	float Sample11 = (float)((*LayerCache)[IdxX1 + Stride*IdxY1]) / 255.0f;
+	float Sample21 = (float)((*LayerCache)[IdxX2 + Stride*IdxY1]) / 255.0f;
+	float Sample12 = (float)((*LayerCache)[IdxX1 + Stride*IdxY2]) / 255.0f;
+	float Sample22 = (float)((*LayerCache)[IdxX2 + Stride*IdxY2]) / 255.0f;
 
 	float LerpX = FMath::Fractional(TestX);
 	float LerpY = FMath::Fractional(TestY);
@@ -1554,12 +1554,8 @@ void ULandscapeComponent::GetComponentExtent(int32& MinX, int32& MinY, int32& Ma
 
 void ULandscapeInfo::GetComponentsInRegion(int32 X1, int32 Y1, int32 X2, int32 Y2, TSet<ULandscapeComponent*>& OutComponents)
 {
-	if (ComponentSizeQuads <= 0)
-	{
-		return;
-	}
-
 	// Find component range for this block of data
+	// X2/Y2 Coordinates are "inclusive" max values
 	int32 ComponentIndexX1, ComponentIndexY1, ComponentIndexX2, ComponentIndexY2;
 	ALandscape::CalcComponentIndicesOverlap(X1, Y1, X2, Y2, ComponentSizeQuads, ComponentIndexX1, ComponentIndexY1, ComponentIndexX2, ComponentIndexY2);
 
@@ -1812,10 +1808,10 @@ void ALandscapeProxy::Import(FGuid Guid, int32 VertsX, int32 VertsY,
 	{
 		for (int32 QuadX = 0; QuadX < NumPatchesX; QuadX++)
 		{
-			FVector Vert00 = FVector(0.f, 0.f, ((float)HEIGHTDATA(QuadX + 0, QuadY + 0) - 32768.f)*LANDSCAPE_ZSCALE) * DrawScale3D;
-			FVector Vert01 = FVector(0.f, 1.f, ((float)HEIGHTDATA(QuadX + 0, QuadY + 1) - 32768.f)*LANDSCAPE_ZSCALE) * DrawScale3D;
-			FVector Vert10 = FVector(1.f, 0.f, ((float)HEIGHTDATA(QuadX + 1, QuadY + 0) - 32768.f)*LANDSCAPE_ZSCALE) * DrawScale3D;
-			FVector Vert11 = FVector(1.f, 1.f, ((float)HEIGHTDATA(QuadX + 1, QuadY + 1) - 32768.f)*LANDSCAPE_ZSCALE) * DrawScale3D;
+			FVector Vert00 = FVector(0.0f, 0.0f, ((float)HEIGHTDATA(QuadX + 0, QuadY + 0) - 32768.0f)*LANDSCAPE_ZSCALE) * DrawScale3D;
+			FVector Vert01 = FVector(0.0f, 1.0f, ((float)HEIGHTDATA(QuadX + 0, QuadY + 1) - 32768.0f)*LANDSCAPE_ZSCALE) * DrawScale3D;
+			FVector Vert10 = FVector(1.0f, 0.0f, ((float)HEIGHTDATA(QuadX + 1, QuadY + 0) - 32768.0f)*LANDSCAPE_ZSCALE) * DrawScale3D;
+			FVector Vert11 = FVector(1.0f, 1.0f, ((float)HEIGHTDATA(QuadX + 1, QuadY + 1) - 32768.0f)*LANDSCAPE_ZSCALE) * DrawScale3D;
 
 			FVector FaceNormal1 = ((Vert00 - Vert10) ^ (Vert10 - Vert11)).GetSafeNormal();
 			FVector FaceNormal2 = ((Vert11 - Vert01) ^ (Vert01 - Vert00)).GetSafeNormal();
@@ -1956,7 +1952,7 @@ void ALandscapeProxy::Import(FGuid Guid, int32 VertsX, int32 VertsY,
 					else if (TotalWeight != 255)
 					{
 						// normalization...
-						float Factor = 255.f / TotalWeight;
+						float Factor = 255.0f / TotalWeight;
 						TotalWeight = 0;
 						for (int32 WeightLayerIndex = 0; WeightLayerIndex < WeightValues.Num(); WeightLayerIndex++)
 						{
@@ -1998,7 +1994,7 @@ void ALandscapeProxy::Import(FGuid Guid, int32 VertsX, int32 VertsY,
 			// Heightmap offsets
 			int32 HeightmapOffsetX = (ComponentX - ComponentsPerHeightmap*HmX) * NumSubsections * (SubsectionSizeQuads + 1);
 
-			LandscapeComponent->HeightmapScaleBias = FVector4(1.f / (float)HeightmapInfo.HeightmapSizeU, 1.f / (float)HeightmapInfo.HeightmapSizeV, (float)((HeightmapOffsetX)) / (float)HeightmapInfo.HeightmapSizeU, ((float)(HeightmapOffsetY)) / (float)HeightmapInfo.HeightmapSizeV);
+			LandscapeComponent->HeightmapScaleBias = FVector4(1.0f / (float)HeightmapInfo.HeightmapSizeU, 1.0f / (float)HeightmapInfo.HeightmapSizeV, (float)((HeightmapOffsetX)) / (float)HeightmapInfo.HeightmapSizeU, ((float)(HeightmapOffsetY)) / (float)HeightmapInfo.HeightmapSizeV);
 			LandscapeComponent->HeightmapTexture = HeightmapInfo.HeightmapTexture;
 
 			// Weightmap is sized the same as the component
@@ -2006,7 +2002,7 @@ void ALandscapeProxy::Import(FGuid Guid, int32 VertsX, int32 VertsY,
 			// Should be power of two
 			check(((WeightmapSize - 1) & WeightmapSize) == 0);
 
-			LandscapeComponent->WeightmapScaleBias = FVector4(1.f / (float)WeightmapSize, 1.f / (float)WeightmapSize, 0.5f / (float)WeightmapSize, 0.5f / (float)WeightmapSize);
+			LandscapeComponent->WeightmapScaleBias = FVector4(1.0f / (float)WeightmapSize, 1.0f / (float)WeightmapSize, 0.5f / (float)WeightmapSize, 0.5f / (float)WeightmapSize);
 			LandscapeComponent->WeightmapSubsectionOffset = (float)(SubsectionSizeQuads + 1) / (float)WeightmapSize;
 
 			// Pointers to the texture data where we'll store each layer. Stride is 4 (FColor)
@@ -2151,8 +2147,8 @@ void ALandscapeProxy::Import(FGuid Guid, int32 VertsX, int32 VertsY,
 
 							HeightmapInfo.HeightmapTextureMipData[0][HeightTexDataIdx].R = HeightValue >> 8;
 							HeightmapInfo.HeightmapTextureMipData[0][HeightTexDataIdx].G = HeightValue & 255;
-							HeightmapInfo.HeightmapTextureMipData[0][HeightTexDataIdx].B = FMath::RoundToInt(127.5f * (Normal.X + 1.f));
-							HeightmapInfo.HeightmapTextureMipData[0][HeightTexDataIdx].A = FMath::RoundToInt(127.5f * (Normal.Y + 1.f));
+							HeightmapInfo.HeightmapTextureMipData[0][HeightTexDataIdx].B = FMath::RoundToInt(127.5f * (Normal.X + 1.0f));
+							HeightmapInfo.HeightmapTextureMipData[0][HeightTexDataIdx].A = FMath::RoundToInt(127.5f * (Normal.Y + 1.0f));
 
 							for (int32 WeightmapIndex = 0; WeightmapIndex < WeightValues.Num(); WeightmapIndex++)
 							{
@@ -2269,7 +2265,7 @@ bool ALandscapeProxy::ExportToRawMesh(int32 InExportLOD, FRawMesh& OutRawMesh) c
 	GetComponents<ULandscapeComponent>(RegisteredLandscapeComponents);
 
 	const FIntRect LandscapeSectionRect = GetBoundingRect();
-	const FVector2D LandscapeUVScale = FVector2D(1.f, 1.f) / FVector2D(LandscapeSectionRect.Size());
+	const FVector2D LandscapeUVScale = FVector2D(1.0f, 1.0f) / FVector2D(LandscapeSectionRect.Size());
 
 	// User specified LOD to export
 	int32 LandscapeLODToExport = ExportLOD;
@@ -2483,12 +2479,12 @@ FVector ULandscapeInfo::GetLandscapeCenterPos(float& LengthZ, int32 MinX /*= MAX
 		}
 
 		const int32 Dist = (ComponentSizeQuads + 1) >> 1; // Should be same in ALandscapeGizmoActiveActor::SetTargetLandscape
-		FVector2D MidPoint(((float)(MinX + MaxX)) / 2.f, ((float)(MinY + MaxY)) / 2.f);
+		FVector2D MidPoint(((float)(MinX + MaxX)) / 2.0f, ((float)(MinY + MaxY)) / 2.0f);
 		MinX = FMath::FloorToInt(MidPoint.X) - Dist;
 		MaxX = FMath::CeilToInt(MidPoint.X) + Dist;
 		MinY = FMath::FloorToInt(MidPoint.Y) - Dist;
 		MaxY = FMath::CeilToInt(MidPoint.Y) + Dist;
-		check(MidPoint.X == ((float)(MinX + MaxX)) / 2.f && MidPoint.Y == ((float)(MinY + MaxY)) / 2.f);
+		check(MidPoint.X == ((float)(MinX + MaxX)) / 2.0f && MidPoint.Y == ((float)(MinY + MaxY)) / 2.0f);
 	}
 
 	check(MinX != MAX_int32);
@@ -2538,7 +2534,7 @@ FVector ULandscapeInfo::GetLandscapeCenterPos(float& LengthZ, int32 MinX /*= MAX
 	}
 	LengthZ = (MaxZ - MinZ + 2 * MarginZ) * ScaleZ;
 
-	const FVector LocalPosition(((float)(MinX + MaxX)) / 2.f, ((float)(MinY + MaxY)) / 2.f, MinZ - MarginZ);
+	const FVector LocalPosition(((float)(MinX + MaxX)) / 2.0f, ((float)(MinY + MaxY)) / 2.0f, MinZ - MarginZ);
 	//return GetLandscapeProxy()->TransformLandscapeLocationToWorld(LocalPosition);
 	return GetLandscapeProxy()->LandscapeActorToWorld().TransformPosition(LocalPosition);
 }
@@ -3441,7 +3437,7 @@ void ULandscapeLayerInfoObject::PostEditChangeProperty(FPropertyChangedEvent& Pr
 	{
 		if (PropertyName == NAME_Hardness)
 		{
-			Hardness = FMath::Clamp<float>(Hardness, 0.f, 1.f);
+			Hardness = FMath::Clamp<float>(Hardness, 0.0f, 1.0f);
 		}
 		else if (PropertyName == NAME_PhysMaterial)
 		{
@@ -3502,7 +3498,7 @@ void ULandscapeLayerInfoObject::PostLoad()
 		{
 			SetFlags(RF_Standalone);
 		}
-		Hardness = FMath::Clamp<float>(Hardness, 0.f, 1.f);
+		Hardness = FMath::Clamp<float>(Hardness, 0.0f, 1.0f);
 	}
 }
 
@@ -3609,11 +3605,11 @@ namespace
 	inline float AdjustStaticLightingResolution(float StaticLightingResolution, int32 NumSubsections, int32 SubsectionSizeQuads, int32 ComponentSizeQuads)
 	{
 		// Change Lighting resolution to proper one...
-		if (StaticLightingResolution > 1.f)
+		if (StaticLightingResolution > 1.0f)
 		{
 			StaticLightingResolution = (int32)StaticLightingResolution;
 		}
-		else if (StaticLightingResolution < 1.f)
+		else if (StaticLightingResolution < 1.0f)
 		{
 			// Restrict to 1/16
 			if (StaticLightingResolution < 0.0625)
@@ -3624,12 +3620,12 @@ namespace
 			// Adjust to 1/2^n
 			int32 i = 2;
 			int32 LightmapSize = (NumSubsections * (SubsectionSizeQuads + 1)) >> 1;
-			while (StaticLightingResolution < (1.f / i) && LightmapSize > 4)
+			while (StaticLightingResolution < (1.0f / i) && LightmapSize > 4)
 			{
 				i <<= 1;
 				LightmapSize >>= 1;
 			}
-			StaticLightingResolution = 1.f / i;
+			StaticLightingResolution = 1.0f / i;
 
 			int32 PixelPaddingX = GPixelFormats[PF_DXT1].BlockSizeX;
 
@@ -4067,7 +4063,7 @@ void ULandscapeComponent::PostEditChangeProperty(FPropertyChangedEvent& Property
 	}
 	else if (GIsEditor && PropertyName == FName(TEXT("StaticLightingResolution")))
 	{
-		if (StaticLightingResolution > 0.f)
+		if (StaticLightingResolution > 0.0f)
 		{
 			StaticLightingResolution = ::AdjustStaticLightingResolution(StaticLightingResolution, NumSubsections, SubsectionSizeQuads, ComponentSizeQuads);
 		}
@@ -4528,7 +4524,7 @@ void ULandscapeComponent::InitHeightmapData(TArray<FColor>& Heights, bool bUpdat
 	int32 MipSizeU = HeightmapSizeU;
 	int32 MipSizeV = HeightmapSizeV;
 
-	HeightmapScaleBias = FVector4(1.f / (float)HeightmapSizeU, 1.f / (float)HeightmapSizeV, 0.f, 0.f);
+	HeightmapScaleBias = FVector4(1.0f / (float)HeightmapSizeU, 1.0f / (float)HeightmapSizeV, 0.0f, 0.0f);
 
 	int32 Mip = 0;
 	while (MipSizeU > 1 && MipSizeV > 1 && MipSubsectionSizeQuads >= 1)
@@ -4607,7 +4603,7 @@ void ULandscapeComponent::InitWeightmapData(TArray<ULandscapeLayerInfoObject*>&
 	check(WeightmapLayerAllocations.Num() > 0 && WeightmapTextures.Num() > 0);
 
 	int32 WeightmapSize = ComponentSizeVerts;
-	WeightmapScaleBias = FVector4(1.f / (float)WeightmapSize, 1.f / (float)WeightmapSize, 0.5f / (float)WeightmapSize, 0.5f / (float)WeightmapSize);
+	WeightmapScaleBias = FVector4(1.0f / (float)WeightmapSize, 1.0f / (float)WeightmapSize, 0.5f / (float)WeightmapSize, 0.5f / (float)WeightmapSize);
 	WeightmapSubsectionOffset = (float)(SubsectionSizeQuads + 1) / (float)WeightmapSize;
 
 	// Channel remapping
diff --git a/Engine/Source/Runtime/Landscape/Private/LandscapeEditInterface.cpp b/Engine/Source/Runtime/Landscape/Private/LandscapeEditInterface.cpp
--- a/Engine/Source/Runtime/Landscape/Private/LandscapeEditInterface.cpp
+++ b/Engine/Source/Runtime/Landscape/Private/LandscapeEditInterface.cpp
@@ -74,7 +74,7 @@ void ALandscape::CalcComponentIndicesOverlap(const int32 X1, const int32 Y1, con
 	ComponentIndexY2 = (Y2 >= 0) ? Y2 / ComponentSizeQuads : (Y2+1) / ComponentSizeQuads - 1;
 }
 
-void ALandscape::CalcComponentIndices(const int32 X1, const int32 Y1, const int32 X2, const int32 Y2, const int32 ComponentSizeQuads, 
+void ALandscape::CalcComponentIndicesNoOverlap(const int32 X1, const int32 Y1, const int32 X2, const int32 Y2, const int32 ComponentSizeQuads, 
 									  int32& ComponentIndexX1, int32& ComponentIndexY1, int32& ComponentIndexX2, int32& ComponentIndexY2)
 {
 	// Find component range for this block of data
@@ -231,10 +231,10 @@ void FLandscapeEditDataInterface::SetHeightData(int32 X1, int32 Y1, int32 X2, in
 		{
 			for( int32 X=0;X<NumVertsX-1;X++ )
 			{
-				FVector Vert00 = FVector(XYOffsets[(X+0) + NumVertsX*(Y+0)].X,		XYOffsets[(X+0) + NumVertsX*(Y+0)].Y, ((float)Data[(X+0) + Stride*(Y+0)] - 32768.f)*LANDSCAPE_ZSCALE) * DrawScale;
-				FVector Vert01 = FVector(XYOffsets[(X+0) + NumVertsX*(Y+0)].X,		XYOffsets[(X+0) + NumVertsX*(Y+0)].Y+1.f, ((float)Data[(X+0) + Stride*(Y+1)] - 32768.f)*LANDSCAPE_ZSCALE) * DrawScale;
-				FVector Vert10 = FVector(XYOffsets[(X+0) + NumVertsX*(Y+0)].X+1.f,	XYOffsets[(X+0) + NumVertsX*(Y+0)].Y, ((float)Data[(X+1) + Stride*(Y+0)] - 32768.f)*LANDSCAPE_ZSCALE) * DrawScale;
-				FVector Vert11 = FVector(XYOffsets[(X+0) + NumVertsX*(Y+0)].X+1.f,	XYOffsets[(X+0) + NumVertsX*(Y+0)].Y+1.f, ((float)Data[(X+1) + Stride*(Y+1)] - 32768.f)*LANDSCAPE_ZSCALE) * DrawScale;
+				FVector Vert00 = FVector(XYOffsets[(X+0) + NumVertsX*(Y+0)].X,		XYOffsets[(X+0) + NumVertsX*(Y+0)].Y, ((float)Data[(X+0) + Stride*(Y+0)] - 32768.0f)*LANDSCAPE_ZSCALE) * DrawScale;
+				FVector Vert01 = FVector(XYOffsets[(X+0) + NumVertsX*(Y+0)].X,		XYOffsets[(X+0) + NumVertsX*(Y+0)].Y+1.0f, ((float)Data[(X+0) + Stride*(Y+1)] - 32768.0f)*LANDSCAPE_ZSCALE) * DrawScale;
+				FVector Vert10 = FVector(XYOffsets[(X+0) + NumVertsX*(Y+0)].X+1.0f,	XYOffsets[(X+0) + NumVertsX*(Y+0)].Y, ((float)Data[(X+1) + Stride*(Y+0)] - 32768.0f)*LANDSCAPE_ZSCALE) * DrawScale;
+				FVector Vert11 = FVector(XYOffsets[(X+0) + NumVertsX*(Y+0)].X+1.0f,	XYOffsets[(X+0) + NumVertsX*(Y+0)].Y+1.0f, ((float)Data[(X+1) + Stride*(Y+1)] - 32768.0f)*LANDSCAPE_ZSCALE) * DrawScale;
 
 				FVector FaceNormal1 = ((Vert00-Vert10) ^ (Vert10-Vert11)).GetSafeNormal();
 				FVector FaceNormal2 = ((Vert11-Vert01) ^ (Vert01-Vert00)).GetSafeNormal(); 
@@ -351,8 +351,8 @@ void FLandscapeEditDataInterface::SetHeightData(int32 X1, int32 Y1, int32 X2, in
 							if( VertexNormals && LandscapeX > X1 && LandscapeX < X2 && LandscapeY > Y1 && LandscapeY < Y2 )
 							{
 								FVector Normal = VertexNormals[DataIndex].GetSafeNormal();
-								TexData.B = FMath::RoundToInt( 127.5f * (Normal.X + 1.f) );
-								TexData.A = FMath::RoundToInt( 127.5f * (Normal.Y + 1.f) );
+								TexData.B = FMath::RoundToInt( 127.5f * (Normal.X + 1.0f) );
+								TexData.A = FMath::RoundToInt( 127.5f * (Normal.Y + 1.0f) );
 							}
 							else if (NormalData)
 							{
@@ -457,10 +457,10 @@ void FLandscapeEditDataInterface::RecalculateNormals()
 		{
 			for( int32 X=0;X<Stride-1;X++ )
 			{
-				FVector Vert00 = FVector(XYOffsets[(X+0) + Stride*(Y+0)].X,		XYOffsets[(X+0) + Stride*(Y+0)].Y, ((float)HeightData[(X+0) + Stride*(Y+0)] - 32768.f)*LANDSCAPE_ZSCALE) * DrawScale;
-				FVector Vert01 = FVector(XYOffsets[(X+0) + Stride*(Y+0)].X,		XYOffsets[(X+0) + Stride*(Y+0)].Y+1.f, ((float)HeightData[(X+0) + Stride*(Y+1)] - 32768.f)*LANDSCAPE_ZSCALE) * DrawScale;
-				FVector Vert10 = FVector(XYOffsets[(X+0) + Stride*(Y+0)].X+1.f,	XYOffsets[(X+0) + Stride*(Y+0)].Y, ((float)HeightData[(X+1) + Stride*(Y+0)] - 32768.f)*LANDSCAPE_ZSCALE) * DrawScale;
-				FVector Vert11 = FVector(XYOffsets[(X+0) + Stride*(Y+0)].X+1.f,	XYOffsets[(X+0) + Stride*(Y+0)].Y+1.f,((float)HeightData[(X+1) + Stride*(Y+1)] - 32768.f)*LANDSCAPE_ZSCALE) * DrawScale;
+				FVector Vert00 = FVector(XYOffsets[(X+0) + Stride*(Y+0)].X,		XYOffsets[(X+0) + Stride*(Y+0)].Y, ((float)HeightData[(X+0) + Stride*(Y+0)] - 32768.0f)*LANDSCAPE_ZSCALE) * DrawScale;
+				FVector Vert01 = FVector(XYOffsets[(X+0) + Stride*(Y+0)].X,		XYOffsets[(X+0) + Stride*(Y+0)].Y+1.0f, ((float)HeightData[(X+0) + Stride*(Y+1)] - 32768.0f)*LANDSCAPE_ZSCALE) * DrawScale;
+				FVector Vert10 = FVector(XYOffsets[(X+0) + Stride*(Y+0)].X+1.0f,	XYOffsets[(X+0) + Stride*(Y+0)].Y, ((float)HeightData[(X+1) + Stride*(Y+0)] - 32768.0f)*LANDSCAPE_ZSCALE) * DrawScale;
+				FVector Vert11 = FVector(XYOffsets[(X+0) + Stride*(Y+0)].X+1.0f,	XYOffsets[(X+0) + Stride*(Y+0)].Y+1.0f,((float)HeightData[(X+1) + Stride*(Y+1)] - 32768.0f)*LANDSCAPE_ZSCALE) * DrawScale;
 
 				FVector FaceNormal1 = ((Vert00-Vert10) ^ (Vert10-Vert11)).GetSafeNormal();
 				FVector FaceNormal2 = ((Vert11-Vert01) ^ (Vert01-Vert00)).GetSafeNormal(); 
@@ -501,8 +501,8 @@ void FLandscapeEditDataInterface::RecalculateNormals()
 
 						// Update the texture
 						FVector Normal = VertexNormals[DataIndex].GetSafeNormal();
-						TexData.B = FMath::RoundToInt( 127.5f * (Normal.X + 1.f) );
-						TexData.A = FMath::RoundToInt( 127.5f * (Normal.Y + 1.f) );
+						TexData.B = FMath::RoundToInt( 127.5f * (Normal.X + 1.0f) );
+						TexData.A = FMath::RoundToInt( 127.5f * (Normal.Y + 1.0f) );
 					}
 				}
 			}
@@ -539,7 +539,7 @@ void FLandscapeEditDataInterface::GetHeightDataTemplFast(const int32 X1, const i
 {
 	if (!LandscapeInfo) return;
 	int32 ComponentIndexX1, ComponentIndexY1, ComponentIndexX2, ComponentIndexY2;
-	ALandscape::CalcComponentIndices(X1, Y1, X2, Y2, ComponentSizeQuads, ComponentIndexX1, ComponentIndexY1, ComponentIndexX2, ComponentIndexY2);
+	ALandscape::CalcComponentIndicesNoOverlap(X1, Y1, X2, Y2, ComponentSizeQuads, ComponentIndexX1, ComponentIndexY1, ComponentIndexX2, ComponentIndexY2);
 
 	for( int32 ComponentIndexY=ComponentIndexY1;ComponentIndexY<=ComponentIndexY2;ComponentIndexY++ )
 	{
@@ -1284,7 +1284,7 @@ void FLandscapeEditDataInterface::GetHeightDataTempl(int32& ValidX1, int32& Vali
 								// Find the texture data corresponding to this vertex
 								uint16 Value[4] = {0, 0, 0, 0};
 								int32 Dist[4] = {INT_MAX, INT_MAX, INT_MAX, INT_MAX};
-								float ValueX = 0.f, ValueY = 0.f;
+								float ValueX = 0.0f, ValueY = 0.0f;
 								bool Exist[4] = {false, false, false, false};
 
 								// Use data already stored for 0, 2
@@ -1661,7 +1661,7 @@ void ULandscapeComponent::DeleteLayer( ULandscapeLayerInfoObject* LayerInfo, str
 								if( LayerIdx != DeleteLayerIdx && LayerNoWeightBlends[LayerIdx]==false )
 								{
 									uint8& Weight = LayerDataPtrs[LayerIdx][TexDataIndex];
-									Weight = FMath::Clamp<int32>( FMath::RoundToInt(255.f * (float)Weight/(float)OtherLayerWeightSum), 0, 255 );
+									Weight = FMath::Clamp<int32>( FMath::RoundToInt(255.0f * (float)Weight/(float)OtherLayerWeightSum), 0, 255 );
 								}
 							}
 						}
@@ -2012,7 +2012,7 @@ namespace
 
 	template<> void TArrayStoreData<FVector>::Store(int32 LandscapeX, int32 LandscapeY, FVector2D Offset)
 	{
-		Data[ (LandscapeY-Y1) * Stride + (LandscapeX-X1) ] = FVector(Offset.X, Offset.Y, 0.f);
+		Data[ (LandscapeY-Y1) * Stride + (LandscapeX-X1) ] = FVector(Offset.X, Offset.Y, 0.0f);
 	}
 
 	// Data items should be initialized with ArraySize
@@ -2098,7 +2098,7 @@ namespace
 		}
 		else
 		{
-			SparseData.Add(ALandscape::MakeKey(LandscapeX,LandscapeY), FVector(Offset.X, Offset.Y, 0.f));
+			SparseData.Add(ALandscape::MakeKey(LandscapeX,LandscapeY), FVector(Offset.X, Offset.Y, 0.0f));
 		}
 	}
 };
@@ -2321,7 +2321,7 @@ void FLandscapeEditDataInterface::SetAlphaData(ULandscapeLayerInfoObject* const
 
 									if (OtherLayerWeightSum != 255)
 									{
-										const float Factor = 255.f / OtherLayerWeightSum;
+										const float Factor = 255.0f / OtherLayerWeightSum;
 										OtherLayerWeightSum = 0;
 
 										// Normalize
@@ -2739,7 +2739,7 @@ template<typename TStoreData>
 void FLandscapeEditDataInterface::GetWeightDataTemplFast(ULandscapeLayerInfoObject* LayerInfo, const int32 X1, const int32 Y1, const int32 X2, const int32 Y2, TStoreData& StoreData)
 {
 	int32 ComponentIndexX1, ComponentIndexY1, ComponentIndexX2, ComponentIndexY2;
-	ALandscape::CalcComponentIndices(X1, Y1, X2, Y2, ComponentSizeQuads, ComponentIndexX1, ComponentIndexY1, ComponentIndexX2, ComponentIndexY2);
+	ALandscape::CalcComponentIndicesNoOverlap(X1, Y1, X2, Y2, ComponentSizeQuads, ComponentIndexX1, ComponentIndexY1, ComponentIndexX2, ComponentIndexY2);
 
 	// Channel remapping
 	int32 ChannelOffsets[4] = {(int32)STRUCT_OFFSET(FColor,R),(int32)STRUCT_OFFSET(FColor,G),(int32)STRUCT_OFFSET(FColor,B),(int32)STRUCT_OFFSET(FColor,A)};
@@ -2920,7 +2920,7 @@ void FLandscapeEditDataInterface::GetWeightDataTempl(ULandscapeLayerInfoObject*
 	ValidX1 = INT_MAX; ValidX2 = INT_MIN; ValidY1 = INT_MAX; ValidY2 = INT_MIN;
 
 	int32 ComponentIndexX1, ComponentIndexY1, ComponentIndexX2, ComponentIndexY2;
-	ALandscape::CalcComponentIndices(X1, Y1, X2, Y2, ComponentSizeQuads, ComponentIndexX1, ComponentIndexY1, ComponentIndexX2, ComponentIndexY2);
+	ALandscape::CalcComponentIndicesNoOverlap(X1, Y1, X2, Y2, ComponentSizeQuads, ComponentIndexX1, ComponentIndexY1, ComponentIndexX2, ComponentIndexY2);
 
 	int32 ComponentSizeX = ComponentIndexX2-ComponentIndexX1+1;
 	int32 ComponentSizeY = ComponentIndexY2-ComponentIndexY1+1;
@@ -3313,7 +3313,7 @@ void FLandscapeEditDataInterface::GetWeightDataTempl(ULandscapeLayerInfoObject*
 									// Find the texture data corresponding to this vertex
 									uint8 Value[4] = {0, 0, 0, 0};
 									int32 Dist[4] = {INT_MAX, INT_MAX, INT_MAX, INT_MAX};
-									float ValueX = 0.f, ValueY = 0.f;
+									float ValueX = 0.0f, ValueY = 0.0f;
 									bool Exist[4] = {false, false, false, false};
 
 									// Use data already stored for 0, 2
@@ -3674,7 +3674,7 @@ template<typename TStoreData>
 void FLandscapeEditDataInterface::GetSelectDataTempl(const int32 X1, const int32 Y1, const int32 X2, const int32 Y2, TStoreData& StoreData)
 {
 	int32 ComponentIndexX1, ComponentIndexY1, ComponentIndexX2, ComponentIndexY2;
-	ALandscape::CalcComponentIndices(X1, Y1, X2, Y2, ComponentSizeQuads, ComponentIndexX1, ComponentIndexY1, ComponentIndexX2, ComponentIndexY2);
+	ALandscape::CalcComponentIndicesNoOverlap(X1, Y1, X2, Y2, ComponentSizeQuads, ComponentIndexX1, ComponentIndexY1, ComponentIndexX2, ComponentIndexY2);
 
 	for( int32 ComponentIndexY=ComponentIndexY1;ComponentIndexY<=ComponentIndexY2;ComponentIndexY++ )
 	{
@@ -3774,7 +3774,7 @@ void FLandscapeEditDataInterface::SetSelectData(int32 X1, int32 Y1, int32 X2, in
 	check(ComponentSizeQuads > 0);
 	// Find component range for this block of data
 	int32 ComponentIndexX1, ComponentIndexY1, ComponentIndexX2, ComponentIndexY2;
-	ALandscape::CalcComponentIndices(X1, Y1, X2, Y2, ComponentSizeQuads, ComponentIndexX1, ComponentIndexY1, ComponentIndexX2, ComponentIndexY2);
+	ALandscape::CalcComponentIndicesNoOverlap(X1, Y1, X2, Y2, ComponentSizeQuads, ComponentIndexX1, ComponentIndexY1, ComponentIndexX2, ComponentIndexY2);
 
 	for( int32 ComponentIndexY=ComponentIndexY1;ComponentIndexY<=ComponentIndexY2;ComponentIndexY++ )
 	{
@@ -3910,7 +3910,7 @@ void FLandscapeEditDataInterface::SetXYOffsetDataTempl(int32 X1, int32 Y1, int32
 	check(ComponentSizeQuads > 0);
 	// Find component range for this block of data
 	int32 ComponentIndexX1, ComponentIndexY1, ComponentIndexX2, ComponentIndexY2;
-	ALandscape::CalcComponentIndices(X1, Y1, X2, Y2, ComponentSizeQuads, ComponentIndexX1, ComponentIndexY1, ComponentIndexX2, ComponentIndexY2);
+	ALandscape::CalcComponentIndicesNoOverlap(X1, Y1, X2, Y2, ComponentSizeQuads, ComponentIndexX1, ComponentIndexY1, ComponentIndexX2, ComponentIndexY2);
 
 	const FColor DefaultValue(128, 0, 128, 0);
 
@@ -4009,8 +4009,8 @@ void FLandscapeEditDataInterface::SetXYOffsetDataTempl(int32 X1, int32 Y1, int32
 							int32 TexY = WeightmapOffsetY + (SubsectionSizeQuads+1) * SubIndexY + SubY;
 							FColor& TexData = XYOffsetTextureData[ TexX + TexY * SizeU ];
 
-							uint16 XOffset = FMath::Clamp<uint16>(Value.X * LANDSCAPE_INV_XYOFFSET_SCALE + 32768.f, 0, 65535);
-							uint16 YOffset = FMath::Clamp<uint16>(Value.Y * LANDSCAPE_INV_XYOFFSET_SCALE + 32768.f, 0, 65535);
+							uint16 XOffset = FMath::Clamp<uint16>(Value.X * LANDSCAPE_INV_XYOFFSET_SCALE + 32768.0f, 0, 65535);
+							uint16 YOffset = FMath::Clamp<uint16>(Value.Y * LANDSCAPE_INV_XYOFFSET_SCALE + 32768.0f, 0, 65535);
 
 							TexData.R = XOffset >> 8;
 							TexData.G = XOffset & 255;
@@ -4613,7 +4613,7 @@ template<typename TStoreData>
 void FLandscapeEditDataInterface::GetXYOffsetDataTemplFast(const int32 X1, const int32 Y1, const int32 X2, const int32 Y2, TStoreData& StoreData)
 {
 	int32 ComponentIndexX1, ComponentIndexY1, ComponentIndexX2, ComponentIndexY2;
-	ALandscape::CalcComponentIndices(X1, Y1, X2, Y2, ComponentSizeQuads, ComponentIndexX1, ComponentIndexY1, ComponentIndexX2, ComponentIndexY2);
+	ALandscape::CalcComponentIndicesNoOverlap(X1, Y1, X2, Y2, ComponentSizeQuads, ComponentIndexX1, ComponentIndexY1, ComponentIndexX2, ComponentIndexY2);
 
 	for( int32 ComponentIndexY=ComponentIndexY1;ComponentIndexY<=ComponentIndexY2;ComponentIndexY++ )
 	{
@@ -4667,7 +4667,7 @@ void FLandscapeEditDataInterface::GetXYOffsetDataTemplFast(const int32 X1, const
 							}
 							else
 							{
-								StoreData.Store(LandscapeX, LandscapeY, FVector2D(0.f, 0.f) );
+								StoreData.Store(LandscapeX, LandscapeY, FVector2D(0.0f, 0.0f) );
 							}
 						}
 					}
