Fix "LandscapeEdModeComponentTools.cpp:1494:10: error: expected a class or namespace"

diff --git a/Engine/Source/Editor/LandscapeEditor/Private/LandscapeEdModeComponentTools.cpp b/Engine/Source/Editor/LandscapeEditor/Private/LandscapeEdModeComponentTools.cpp
--- a/Engine/Source/Editor/LandscapeEditor/Private/LandscapeEdModeComponentTools.cpp
+++ b/Engine/Source/Editor/LandscapeEditor/Private/LandscapeEdModeComponentTools.cpp
@@ -1491,12 +1491,12 @@ public:
 		BackupCurrentBrush = this->EdMode->CurrentBrush;
 		this->EdMode->CurrentBrush = this->EdMode->GizmoBrush;
 
-		return FLandscapeToolBase::BeginTool(ViewportClient, InTarget, InHitLocation);
+		return FLandscapeToolBase<FLandscapeToolStrokePaste<ToolTarget>>::BeginTool(ViewportClient, InTarget, InHitLocation);
 	}
 
 	virtual void EndTool(FEditorViewportClient* ViewportClient) override
 	{
-		FLandscapeToolBase::EndTool(ViewportClient);
+		FLandscapeToolBase<FLandscapeToolStrokePaste<ToolTarget>>::EndTool(ViewportClient);
 
 		this->EdMode->CurrentBrush = BackupCurrentBrush;
 	}
@@ -1793,12 +1793,12 @@ public:
 			this->EdMode->CurrentBrush = this->EdMode->GizmoBrush;
 		}
 
-		return FLandscapeToolBase::BeginTool(ViewportClient, InTarget, InHitLocation);
+		return FLandscapeToolBase<FLandscapeToolStrokePaste<ToolTarget>>::BeginTool(ViewportClient, InTarget, InHitLocation);
 	}
 
 	virtual void EndTool(FEditorViewportClient* ViewportClient) override
 	{
-		FLandscapeToolBase::EndTool(ViewportClient);
+		FLandscapeToolBase<FLandscapeToolStrokePaste<ToolTarget>>::EndTool(ViewportClient);
 
 		if (bUseGizmoRegion)
 		{
