Fix utter stupidity

diff --git a/Engine/Source/Editor/LandscapeEditor/Private/LandscapeEdModeComponentTools.cpp b/Engine/Source/Editor/LandscapeEditor/Private/LandscapeEdModeComponentTools.cpp
--- a/Engine/Source/Editor/LandscapeEditor/Private/LandscapeEdModeComponentTools.cpp
+++ b/Engine/Source/Editor/LandscapeEditor/Private/LandscapeEdModeComponentTools.cpp
@@ -1491,12 +1491,12 @@ public:
 		BackupCurrentBrush = this->EdMode->CurrentBrush;
 		this->EdMode->CurrentBrush = this->EdMode->GizmoBrush;
 
-		return FLandscapeToolBase<FLandscapeToolStrokePaste<ToolTarget>>::BeginTool(ViewportClient, InTarget, InHitLocation);
+		return FLandscapeToolBase<FLandscapeToolStrokeCopy<ToolTarget>>::BeginTool(ViewportClient, InTarget, InHitLocation);
 	}
 
 	virtual void EndTool(FEditorViewportClient* ViewportClient) override
 	{
-		FLandscapeToolBase<FLandscapeToolStrokePaste<ToolTarget>>::EndTool(ViewportClient);
+		FLandscapeToolBase<FLandscapeToolStrokeCopy<ToolTarget>>::EndTool(ViewportClient);
 
 		this->EdMode->CurrentBrush = BackupCurrentBrush;
 	}
