Fix for crash entering landscape edit mode

diff --git a/Engine/Source/Editor/LandscapeEditor/Private/LandscapeEdMode.cpp b/Engine/Source/Editor/LandscapeEditor/Private/LandscapeEdMode.cpp
--- a/Engine/Source/Editor/LandscapeEditor/Private/LandscapeEdMode.cpp
+++ b/Engine/Source/Editor/LandscapeEditor/Private/LandscapeEdMode.cpp
@@ -1651,7 +1651,7 @@ int32 FEdModeLandscape::UpdateLandscapeList()
 		for (auto It = GetLandscapeInfoMap(GetWorld()).Map.CreateIterator(); It; ++It)
 		{
 			ULandscapeInfo* LandscapeInfo = It.Value();
-			if (LandscapeInfo)
+			if (LandscapeInfo && LandscapeInfo->bIsValid)
 			{
 				if (CurrentToolTarget.LandscapeInfo == LandscapeInfo)
 				{
