Back out revision 218 from //depot/UE4/Engine/Source/Editor/LandscapeEditor/Private/LandscapeEdMode.cpp

diff --git a/Engine/Source/Editor/LandscapeEditor/Private/LandscapeEdMode.cpp b/Engine/Source/Editor/LandscapeEditor/Private/LandscapeEdMode.cpp
--- a/Engine/Source/Editor/LandscapeEditor/Private/LandscapeEdMode.cpp
+++ b/Engine/Source/Editor/LandscapeEditor/Private/LandscapeEdMode.cpp
@@ -1651,7 +1651,7 @@ int32 FEdModeLandscape::UpdateLandscapeList()
 		for (auto It = GetLandscapeInfoMap(GetWorld()).Map.CreateIterator(); It; ++It)
 		{
 			ULandscapeInfo* LandscapeInfo = It.Value();
-			if (LandscapeInfo && LandscapeInfo->bIsValid)
+			if (LandscapeInfo)
 			{
 				if (CurrentToolTarget.LandscapeInfo == LandscapeInfo)
 				{
