Fix crash with landscape ramp tool (and maybe others) when undoing the creation of the landscape

diff --git a/Engine/Source/Editor/LandscapeEditor/Private/LandscapeEdMode.cpp b/Engine/Source/Editor/LandscapeEditor/Private/LandscapeEdMode.cpp
--- a/Engine/Source/Editor/LandscapeEditor/Private/LandscapeEdMode.cpp
+++ b/Engine/Source/Editor/LandscapeEditor/Private/LandscapeEdMode.cpp
@@ -1651,7 +1651,7 @@ int32 FEdModeLandscape::UpdateLandscapeList()
 		for (auto It = GetLandscapeInfoMap(GetWorld()).Map.CreateIterator(); It; ++It)
 		{
 			ULandscapeInfo* LandscapeInfo = It.Value();
-			if (LandscapeInfo)
+			if (LandscapeInfo && !LandscapeInfo->IsPendingKill())
 			{
 				if (CurrentToolTarget.LandscapeInfo == LandscapeInfo)
 				{
@@ -2266,7 +2266,7 @@ void FEdModeLandscape::ActorSelectionChangeNotify()
 		break;
 		}
 		}
-		*/
+	*/
 }
 
 void FEdModeLandscape::ActorMoveNotify()
@@ -2274,6 +2274,11 @@ void FEdModeLandscape::ActorMoveNotify()
 	//GUnrealEd->UpdateFloatingPropertyWindows();
 }
 
+void FEdModeLandscape::PostUndo()
+{
+	OnWorldChange();
+}
+
 /** Forces all level editor viewports to realtime mode */
 void FEdModeLandscape::ForceRealTimeViewports(const bool bEnable, const bool bStoreCurrentState)
 {
diff --git a/Engine/Source/Editor/LandscapeEditor/Public/LandscapeEdMode.h b/Engine/Source/Editor/LandscapeEditor/Public/LandscapeEdMode.h
--- a/Engine/Source/Editor/LandscapeEditor/Public/LandscapeEdMode.h
+++ b/Engine/Source/Editor/LandscapeEditor/Public/LandscapeEdMode.h
@@ -597,6 +597,8 @@ public:
 
 	virtual void ActorMoveNotify() override;
 
+	virtual void PostUndo() override;
+
 	virtual EEditAction::Type GetActionEditDuplicate() override;
 	virtual EEditAction::Type GetActionEditDelete() override;
 	virtual EEditAction::Type GetActionEditCut() override;
