Refresh dynamic foliage when editing landscape

diff --git a/Engine/Source/Editor/LandscapeEditor/Private/LandscapeEdModeTools.h b/Engine/Source/Editor/LandscapeEditor/Private/LandscapeEdModeTools.h
--- a/Engine/Source/Editor/LandscapeEditor/Private/LandscapeEdModeTools.h
+++ b/Engine/Source/Editor/LandscapeEditor/Private/LandscapeEdModeTools.h
@@ -603,6 +603,9 @@ struct FHeightmapAccessor
 				Component->InvalidateLightingCache();
 			}
 
+			// Flush dynamic foliage (grass)
+			LandscapeInfo->GetLandscapeProxy()->FlushFoliageComponents(&Components);
+
 			// Notify foliage to move any attached instances
 			bool bUpdateFoliage = false;
 			for (ULandscapeComponent* Component : Components)
@@ -653,10 +656,6 @@ struct FHeightmapAccessor
 				LandscapeEdit->SetHeightData(X1, Y1, X2, Y2, Data, 0, true);
 			}
 		}
-		else
-		{
-			ChangedComponents.Empty();
-		}
 	}
 
 	void Flush()
@@ -783,6 +782,9 @@ struct FXYOffsetmapAccessor
 				}
 			}
 
+			// Flush dynamic foliage (grass)
+			LandscapeInfo->GetLandscapeProxy()->FlushFoliageComponents(&Components);
+
 			// Notify foliage to move any attached instances
 			bool bUpdateFoliage = false;
 			for (ULandscapeComponent* Component : Components)
@@ -828,10 +830,6 @@ struct FXYOffsetmapAccessor
 				LandscapeEdit->SetHeightData(X1, Y1, X2, Y2, NewHeights.GetData(), 0, true);
 			}
 		}
-		else
-		{
-			ChangedComponents.Empty();
-		}
 	}
 
 	void Flush()
@@ -938,6 +936,9 @@ struct FAlphamapAccessor
 		TSet<ULandscapeComponent*> Components;
 		if (LandscapeEdit.GetComponentsInRegion(X1, Y1, X2, Y2, &Components))
 		{
+			// Flush dynamic foliage (grass)
+			LandscapeInfo->GetLandscapeProxy()->FlushFoliageComponents(&Components);
+
 			LandscapeEdit.SetAlphaData(LayerInfo, X1, Y1, X2, Y2, Data, 0, PaintingRestriction, bBlendWeight, bUseTotalNormalize);
 			ModifiedComponents.Append(Components);
 		}
@@ -1006,7 +1007,8 @@ struct FFullWeightmapAccessor
 {
 	enum { bUseInterp = bInUseInterp };
 	FFullWeightmapAccessor(ULandscapeInfo* InLandscapeInfo)
-		: LandscapeEdit(InLandscapeInfo)
+		: LandscapeInfo(InLandscapeInfo)
+		, LandscapeEdit(InLandscapeInfo)
 	{
 	}
 	void GetData(int32& X1, int32& Y1, int32& X2, int32& Y2, TMap<FIntPoint, TArray<uint8>>& Data)
@@ -1023,8 +1025,12 @@ struct FFullWeightmapAccessor
 
 	void SetData(int32 X1, int32 Y1, int32 X2, int32 Y2, const uint8* Data, ELandscapeLayerPaintingRestriction::Type PaintingRestriction)
 	{
-		if (LandscapeEdit.GetComponentsInRegion(X1, Y1, X2, Y2))
+		TSet<ULandscapeComponent*> Components;
+		if (LandscapeEdit.GetComponentsInRegion(X1, Y1, X2, Y2, &Components))
 		{
+			// Flush dynamic foliage (grass)
+			LandscapeInfo->GetLandscapeProxy()->FlushFoliageComponents(&Components);
+
 			LandscapeEdit.SetAlphaData(DirtyLayerInfos, X1, Y1, X2, Y2, Data, 0, PaintingRestriction);
 		}
 		DirtyLayerInfos.Empty();
@@ -1036,7 +1042,9 @@ struct FFullWeightmapAccessor
 	}
 
 	TSet<ULandscapeLayerInfoObject*> DirtyLayerInfos;
+
 private:
+	ULandscapeInfo* LandscapeInfo;
 	FLandscapeEditDataInterface LandscapeEdit;
 };
 
diff --git a/Engine/Source/Runtime/Landscape/Classes/LandscapeProxy.h b/Engine/Source/Runtime/Landscape/Classes/LandscapeProxy.h
--- a/Engine/Source/Runtime/Landscape/Classes/LandscapeProxy.h
+++ b/Engine/Source/Runtime/Landscape/Classes/LandscapeProxy.h
@@ -453,7 +453,7 @@ public:
 	virtual ALandscape* GetLandscapeActor();
 
 	/** Flush the foliage cache */
-	LANDSCAPE_API void FlushFoliageComponents(TSet<ULandscapeComponent*>* OnlyForComponents = nullptr);
+	LANDSCAPE_API void FlushFoliageComponents(const TSet<ULandscapeComponent*>* OnlyForComponents = nullptr);
 
 	/** 
 		Update Foliage 
diff --git a/Engine/Source/Runtime/Landscape/Private/Landscape.cpp b/Engine/Source/Runtime/Landscape/Private/Landscape.cpp
--- a/Engine/Source/Runtime/Landscape/Private/Landscape.cpp
+++ b/Engine/Source/Runtime/Landscape/Private/Landscape.cpp
@@ -2101,22 +2101,20 @@ bool LandscapeMaterialsParameterSetUpdater(FStaticParameterSet &StaticParameterS
 	return UpdateParameterSet<FStaticTerrainLayerWeightParameter, UMaterialExpressionLandscapeLayerWeight>(StaticParameterSet.TerrainLayerWeightParameters, ParentMaterial);
 }
 
-void ALandscapeProxy::FlushFoliageComponents(TSet<ULandscapeComponent*>* OnlyForComponents)
+void ALandscapeProxy::FlushFoliageComponents(const TSet<ULandscapeComponent*>* OnlyForComponents)
 {
 	if (OnlyForComponents)
 	{
-		OnlyForComponents->Add(nullptr); // if the weak pointer in the cache is invalid, we should kill them anyway
-		for (int32 CacheIndex = 0; CacheIndex < FoliageCache.PerComponent.Num(); CacheIndex++)
+		for (FCachedLandscapeFoliage::FPerComponent& CacheItem : FoliageCache.PerComponent)
 		{
-			FCachedLandscapeFoliage::FPerComponent& CacheItem = FoliageCache.PerComponent[CacheIndex];
-			if (OnlyForComponents->Contains(CacheItem.BasedOn.Get()))
+			ULandscapeComponent* Component = CacheItem.BasedOn.Get();
+			// if the weak pointer in the cache is invalid, we should kill them anyway
+			if (Component == nullptr || OnlyForComponents->Contains(Component))
 			{
-				for (int32 LayerIndex = 0; LayerIndex < CacheItem.Layers.Num(); LayerIndex++)
+				for (FCachedLandscapeFoliage::FPerLayer& LayerItem : CacheItem.Layers)
 				{
-					FCachedLandscapeFoliage::FPerLayer& LayerItem = CacheItem.Layers[LayerIndex];
-					for (int32 FoliageIndex = 0; FoliageIndex < LayerItem.Foliage.Num(); FoliageIndex++)
+					for (FCachedLandscapeFoliage::FPerGrassComp& GrassItem : LayerItem.Foliage)
 					{
-						FCachedLandscapeFoliage::FPerGrassComp& GrassItem = LayerItem.Foliage[FoliageIndex];
 						UHierarchicalInstancedStaticMeshComponent *Used = GrassItem.Foliage.Get();
 						if (Used)
 						{
diff --git a/Engine/Source/Runtime/Landscape/Private/LandscapeEdit.cpp b/Engine/Source/Runtime/Landscape/Private/LandscapeEdit.cpp
--- a/Engine/Source/Runtime/Landscape/Private/LandscapeEdit.cpp
+++ b/Engine/Source/Runtime/Landscape/Private/LandscapeEdit.cpp
@@ -345,6 +345,10 @@ void ULandscapeComponent::PostEditUndo()
 		EditToolRenderData->UpdateDebugColorMaterial();
 		EditToolRenderData->UpdateSelectionMaterial(EditToolRenderData->SelectedType);
 	}
+
+	TSet<ULandscapeComponent*> Components;
+	Components.Add(this);
+	GetLandscapeProxy()->FlushFoliageComponents(&Components);
 }
 
 void ULandscapeComponent::FixupWeightmaps()
