Fix copy/paste tool not updating landscape physical material map

diff --git a/Engine/Source/Editor/LandscapeEditor/Private/LandscapeEdModeTools.h b/Engine/Source/Editor/LandscapeEditor/Private/LandscapeEdModeTools.h
--- a/Engine/Source/Editor/LandscapeEditor/Private/LandscapeEdModeTools.h
+++ b/Engine/Source/Editor/LandscapeEditor/Private/LandscapeEdModeTools.h
@@ -901,9 +901,9 @@ struct FAlphamapAccessor
 	~FAlphamapAccessor()
 	{
 		// Recreate collision for modified components to update the physical materials
-		for (auto It = ModifiedComponents.CreateConstIterator(); It; ++It)
+		for (ULandscapeComponent* Component : ModifiedComponents)
 		{
-			ULandscapeHeightfieldCollisionComponent* CollisionComponent = (*It)->CollisionComponent.Get();
+			ULandscapeHeightfieldCollisionComponent* CollisionComponent = Component->CollisionComponent.Get();
 			if (CollisionComponent)
 			{
 				CollisionComponent->RecreateCollision(false);
@@ -911,7 +911,7 @@ struct FAlphamapAccessor
 				// We need to trigger navigation mesh build, in case user have painted holes on a landscape
 				if (LayerInfo == ALandscapeProxy::VisibilityLayer)
 				{
-					UNavigationSystem* NavSys = UNavigationSystem::GetCurrent(*It);
+					UNavigationSystem* NavSys = UNavigationSystem::GetCurrent(Component);
 					if (NavSys)
 					{
 						NavSys->UpdateNavOctree(CollisionComponent);
@@ -1011,6 +1011,31 @@ struct FFullWeightmapAccessor
 		, LandscapeEdit(InLandscapeInfo)
 	{
 	}
+
+	~FFullWeightmapAccessor()
+	{
+		// Recreate collision for modified components to update the physical materials
+		for (ULandscapeComponent* Component : ModifiedComponents)
+		{
+			ULandscapeHeightfieldCollisionComponent* CollisionComponent = Component->CollisionComponent.Get();
+			if (CollisionComponent)
+			{
+				CollisionComponent->RecreateCollision(false);
+
+				// We need to trigger navigation mesh build, in case user have painted holes on a landscape
+				if (LandscapeInfo->GetLayerInfoIndex(ALandscapeProxy::VisibilityLayer) != INDEX_NONE)
+				{
+					UNavigationSystem* NavSys = UNavigationSystem::GetCurrent(Component);
+					if (NavSys)
+					{
+						NavSys->UpdateNavOctree(CollisionComponent);
+					}
+				}
+			}
+		}
+	}
+
+
 	void GetData(int32& X1, int32& Y1, int32& X2, int32& Y2, TMap<FIntPoint, TArray<uint8>>& Data)
 	{
 		// Do not Support for interpolation....
@@ -1032,6 +1057,7 @@ struct FFullWeightmapAccessor
 			LandscapeInfo->GetLandscapeProxy()->FlushFoliageComponents(&Components);
 
 			LandscapeEdit.SetAlphaData(DirtyLayerInfos, X1, Y1, X2, Y2, Data, 0, PaintingRestriction);
+			ModifiedComponents.Append(Components);
 		}
 		DirtyLayerInfos.Empty();
 	}
@@ -1046,6 +1072,7 @@ struct FFullWeightmapAccessor
 private:
 	ULandscapeInfo* LandscapeInfo;
 	FLandscapeEditDataInterface LandscapeEdit;
+	TSet<ULandscapeComponent*> ModifiedComponents;
 };
 
 struct FLandscapeFullWeightCache : public TLandscapeEditCache < FFullWeightmapAccessor<false>, TArray<uint8> >
