Fix Landscape region select tool

diff --git a/Engine/Source/Editor/LandscapeEditor/Private/LandscapeEdModeComponentTools.cpp b/Engine/Source/Editor/LandscapeEditor/Private/LandscapeEdModeComponentTools.cpp
--- a/Engine/Source/Editor/LandscapeEditor/Private/LandscapeEdModeComponentTools.cpp
+++ b/Engine/Source/Editor/LandscapeEditor/Private/LandscapeEdModeComponentTools.cpp
@@ -138,8 +138,8 @@ public:
 						if (BrushValue > 0.0f && LandscapeInfo->IsValidPosition(X, Y))
 						{
 							float PaintValue = BrushValue * UISettings->ToolStrength * Pressure;
-							float Value = DataScanline[X];
-							checkSlow(Value == LandscapeInfo->SelectedRegion.FindRef(Key));
+							float Value = DataScanline[X] / 255.0f;
+							checkSlow(FMath::IsNearlyEqual(Value, LandscapeInfo->SelectedRegion.FindRef(Key), 1 / 255.0f));
 							if (bInvert)
 							{
 								Value = FMath::Max(Value - PaintValue, 0.0f);
diff --git a/Engine/Source/Editor/LandscapeEditor/Private/LandscapeEdModeTools.h b/Engine/Source/Editor/LandscapeEditor/Private/LandscapeEdModeTools.h
--- a/Engine/Source/Editor/LandscapeEditor/Private/LandscapeEdModeTools.h
+++ b/Engine/Source/Editor/LandscapeEditor/Private/LandscapeEdModeTools.h
@@ -465,24 +465,32 @@ struct TLandscapeEditCache
 	// X2/Y2 Coordinates are "inclusive" max values
 	bool GetCachedData(int32 X1, int32 Y1, int32 X2, int32 Y2, TArray<AccessorType>& OutData)
 	{
-		int32 NumSamples = (1 + X2 - X1)*(1 + Y2 - Y1);
+		const int32 XSize = (1 + X2 - X1);
+		const int32 YSize = (1 + Y2 - Y1);
+		const int32 NumSamples = XSize * YSize;
 		OutData.Empty(NumSamples);
 		OutData.AddUninitialized(NumSamples);
 		bool bHasNonZero = false;
 
 		for (int32 Y = Y1; Y <= Y2; Y++)
 		{
+			const int32 YOffset = (Y - Y1) * XSize;
 			for (int32 X = X1; X <= X2; X++)
 			{
+				const int32 XYOffset = YOffset + (X - X1);
 				AccessorType* Ptr = GetValueRef(X, Y);
 				if (Ptr)
 				{
-					OutData[(X-X1) + (Y-Y1)*(1+X2-X1)] = *Ptr;
+					OutData[XYOffset] = *Ptr;
 					if (!IsZeroValue(*Ptr))
 					{
 						bHasNonZero = true;
 					}
 				}
+				else
+				{
+					OutData[XYOffset] = (AccessorType)0;
+				}
 			}
 		}
 
