UE4 - fix crash on raw importing

diff --git a/Engine/Source/Editor/LandscapeEditor/Private/LandscapeEdMode.cpp b/Engine/Source/Editor/LandscapeEditor/Private/LandscapeEdMode.cpp
--- a/Engine/Source/Editor/LandscapeEditor/Private/LandscapeEdMode.cpp
+++ b/Engine/Source/Editor/LandscapeEditor/Private/LandscapeEdMode.cpp
@@ -2453,6 +2453,7 @@ void FEdModeLandscape::ImportData(const FLandscapeTargetListInfo& TargetInfo, co
 			else
 			{
 				// I would love to deprecate the raw/r8/r16 support for landscape, r16 doesn't even handle endianness issues...
+				RawData = &Data;
 			}
 
 			if (TargetInfo.TargetType == ELandscapeToolTargetType::Heightmap)
