Fix stupid mistake

diff --git a/Engine/Source/Editor/LandscapeEditor/Private/LandscapeEdModeBrushes.cpp b/Engine/Source/Editor/LandscapeEditor/Private/LandscapeEdModeBrushes.cpp
--- a/Engine/Source/Editor/LandscapeEditor/Private/LandscapeEdModeBrushes.cpp
+++ b/Engine/Source/Editor/LandscapeEditor/Private/LandscapeEdModeBrushes.cpp
@@ -637,7 +637,7 @@ public:
 	virtual TOptional<bool> InputKey(FEditorViewportClient* InViewportClient, FViewport* InViewport, FKey InKey, EInputEvent InEvent) override
 	{
 
-		if (Key == EKeys::LeftMouseButton && Event == IE_Pressed)
+		if (InKey == EKeys::LeftMouseButton && InEvent == IE_Pressed)
 		{
 			int32 HitX = InViewport->GetMouseX();
 			int32 HitY = InViewport->GetMouseY();
