Missing change from CL 2412231

diff --git a/Engine/Source/Editor/LandscapeEditor/Public/LandscapeEdMode.h b/Engine/Source/Editor/LandscapeEditor/Public/LandscapeEdMode.h
--- a/Engine/Source/Editor/LandscapeEditor/Public/LandscapeEdMode.h
+++ b/Engine/Source/Editor/LandscapeEditor/Public/LandscapeEdMode.h
@@ -98,7 +98,7 @@ class FLandscapeBrush : public FGCObject
 public:
 	virtual void MouseMove(float LandscapeX, float LandscapeY) = 0;
 	virtual FLandscapeBrushData ApplyBrush(const TArray<FLandscapeToolMousePosition>& MousePositions) = 0;
-	virtual TOptional<bool> InputKey(FEditorViewportClient* InViewportClient, FViewport* InViewport, FKey InKey, EInputEvent InEvent) { return false; }
+	virtual TOptional<bool> InputKey(FEditorViewportClient* InViewportClient, FViewport* InViewport, FKey InKey, EInputEvent InEvent) { return TOptional<bool>(); }
 	virtual void Tick(FEditorViewportClient* ViewportClient, float DeltaTime) {};
 	virtual void BeginStroke(float LandscapeX, float LandscapeY, class FLandscapeTool* CurrentTool);
 	virtual void EndStroke();
