Fixed compile error introduced in cL 2412231

diff --git a/Engine/Source/Editor/LandscapeEditor/Private/LandscapeEdModeComponentTools.cpp b/Engine/Source/Editor/LandscapeEditor/Private/LandscapeEdModeComponentTools.cpp
--- a/Engine/Source/Editor/LandscapeEditor/Private/LandscapeEdModeComponentTools.cpp
+++ b/Engine/Source/Editor/LandscapeEditor/Private/LandscapeEdModeComponentTools.cpp
@@ -1812,7 +1812,7 @@ public:
 	virtual void EnterTool()
 	{
 		// Make sure gizmo actor is selected
-		ALandscapeGizmoActiveActor* Gizmo = EdMode->CurrentGizmoActor.Get();
+		ALandscapeGizmoActiveActor* Gizmo = this->EdMode->CurrentGizmoActor.Get();
 		if (Gizmo)
 		{
 			GEditor->SelectNone(false, true);
