Feature request: Landscape "pattern" brush now has world-space settings
Also added FVector2D::GetRotated(degrees) and GConfig->Get/SetVector2D(...) as they were missing.

diff --git a/Engine/Source/Editor/LandscapeEditor/Classes/LandscapeEditorObject.h b/Engine/Source/Editor/LandscapeEditor/Classes/LandscapeEditorObject.h
--- a/Engine/Source/Editor/LandscapeEditor/Classes/LandscapeEditorObject.h
+++ b/Engine/Source/Editor/LandscapeEditor/Classes/LandscapeEditorObject.h
@@ -187,6 +187,28 @@ struct FLandscapeImportLayer : public FLandscapeImportLayerInfo
 	}
 };
 
+USTRUCT()
+struct FLandscapePatternBrushWorldSpaceSettings
+{
+	GENERATED_USTRUCT_BODY()
+
+	UPROPERTY(Category="World-Space", EditAnywhere)
+	FVector2D Origin;
+
+	UPROPERTY(Category = "World-Space", EditAnywhere, Meta = (ClampMin = "-360", ClampMax = "360", UIMin = "-180", UIMax = "180"))
+	float Rotation;
+
+	// if true, the texture used for the pattern is centered on the PatternOrigin.
+	// if false, the corner of the texture is placed at the PatternOrigin
+	UPROPERTY(Category = "World-Space", EditAnywhere)
+	bool bCenterTextureOnOrigin;
+
+	UPROPERTY(Category = "World-Space", EditAnywhere)
+	float RepeatSize;
+
+	FLandscapePatternBrushWorldSpaceSettings() = default;
+};
+
 UCLASS()
 class ULandscapeEditorObject : public UObject
 {
@@ -206,7 +228,7 @@ class ULandscapeEditorObject : public UObject
 	bool bUseWeightTargetValue;
 
 	// Enable to make tools blend towards a target value
-	UPROPERTY(Category="Tool Settings", EditAnywhere, NonTransactional, meta=(DisplayName="Use Target Value", editcondition="bUseWeightTargetValue", ShowForTools="Paint,Sculpt,Noise", ClampMin="0", ClampMax="10", UIMin="0", UIMax="1"))
+	UPROPERTY(Category="Tool Settings", EditAnywhere, NonTransactional, meta=(DisplayName="Use Target Value", EditCondition="bUseWeightTargetValue", ShowForTools="Paint,Sculpt,Noise", ClampMin="0", ClampMax="10", UIMin="0", UIMax="1"))
 	float WeightTargetValue;
 
 	// I have no idea what this is for but it's used by the noise and erosion tools, and isn't exposed to the UI
@@ -232,11 +254,11 @@ class ULandscapeEditorObject : public UObject
 	bool bUseFlattenTarget;
 
 	// Target height to flatten towards (in Unreal Units)
-	UPROPERTY(Category="Tool Settings", EditAnywhere, NonTransactional, meta=(ShowForTools="Flatten", ShowForTargetTypes="Heightmap", editcondition="bUseFlattenTarget", UIMin="-32768", UIMax="32768"))
+	UPROPERTY(Category="Tool Settings", EditAnywhere, NonTransactional, meta=(ShowForTools="Flatten", ShowForTargetTypes="Heightmap", EditCondition="bUseFlattenTarget", UIMin="-32768", UIMax="32768"))
 	float FlattenTarget;
 
 	// Whether to show the preview grid for the flatten target height
-	UPROPERTY(Category = "Tool Settings", EditAnywhere, NonTransactional, AdvancedDisplay, meta = (DisplayName = "Show Preview Grid", ShowForTools = "Flatten", ShowForTargetTypes = "Heightmap", editcondition = "bUseFlattenTarget", HideEditConditionToggle, UIMin = "-32768", UIMax = "32768"))
+	UPROPERTY(Category = "Tool Settings", EditAnywhere, NonTransactional, AdvancedDisplay, meta = (DisplayName = "Show Preview Grid", ShowForTools = "Flatten", ShowForTargetTypes = "Heightmap", EditCondition = "bUseFlattenTarget", HideEditConditionToggle, UIMin = "-32768", UIMax = "32768"))
 	bool bShowFlattenTargetPreview;
 
 	// Ramp Tool:
@@ -261,7 +283,7 @@ class ULandscapeEditorObject : public UObject
 	bool bDetailSmooth;
 
 	// Larger detail smoothing values remove more details, while smaller values preserve more details
-	UPROPERTY(Category="Tool Settings", EditAnywhere, NonTransactional, meta=(DisplayName="Detail Smooth", editcondition="bDetailSmooth", ShowForTools="Smooth", ClampMin="0", ClampMax="0.99"))
+	UPROPERTY(Category="Tool Settings", EditAnywhere, NonTransactional, meta=(DisplayName="Detail Smooth", EditCondition="bDetailSmooth", ShowForTools="Smooth", ClampMin="0", ClampMax="0.99"))
 	float DetailScale;
 
 	// Erosion Tool:
@@ -313,7 +335,7 @@ class ULandscapeEditorObject : public UObject
 	bool bHErosionDetailSmooth;
 
 	// Larger detail smoothing values remove more details, while smaller values preserve more details
-	UPROPERTY(Category="Tool Settings", EditAnywhere, NonTransactional, meta=(DisplayName="Detail Smooth", editcondition="bHErosionDetailSmooth", ShowForTools="HydraErosion", ClampMin="0", ClampMax="0.99"))
+	UPROPERTY(Category="Tool Settings", EditAnywhere, NonTransactional, meta=(DisplayName="Detail Smooth", EditCondition="bHErosionDetailSmooth", ShowForTools="HydraErosion", ClampMin="0", ClampMax="0.99"))
 	float HErosionDetailScale;
 
 	// Noise Tool:
@@ -472,6 +494,12 @@ class ULandscapeEditorObject : public UObject
 	UPROPERTY(Category="Brush Settings", EditAnywhere, NonTransactional, meta=(DisplayName="Texture Pan V", ShowForBrushes="BrushSet_Pattern", ClampMin="0", ClampMax="1"))
 	float AlphaBrushPanV;
 
+	UPROPERTY(Category = "Brush Settings", EditAnywhere, NonTransactional, meta = (DisplayName = "Use World-Space", ShowForBrushes = "BrushSet_Pattern"))
+	bool bUseWorldSpacePatternBrush;
+
+	UPROPERTY(Category = "Brush Settings", EditAnywhere, NonTransactional, meta = (DisplayName = "World-Space Settings", ShowForBrushes = "BrushSet_Pattern", EditCondition = "bUseWorldSpacePatternBrush"))
+	FLandscapePatternBrushWorldSpaceSettings WorldSpacePatternBrushSettings;
+
 	// Mask texture to use
 	UPROPERTY(Category="Brush Settings", EditAnywhere, NonTransactional, meta=(DisplayName="Texture", ShowForBrushes="BrushSet_Alpha,BrushSet_Pattern"))
 	UTexture2D* AlphaTexture;
diff --git a/Engine/Source/Editor/LandscapeEditor/Private/LandscapeEdModeBrushes.cpp b/Engine/Source/Editor/LandscapeEditor/Private/LandscapeEdModeBrushes.cpp
--- a/Engine/Source/Editor/LandscapeEditor/Private/LandscapeEdModeBrushes.cpp
+++ b/Engine/Source/Editor/LandscapeEditor/Private/LandscapeEdModeBrushes.cpp
@@ -1085,15 +1085,38 @@ public:
 			int32 SizeX = EdMode->UISettings->AlphaTextureSizeX;
 			int32 SizeY = EdMode->UISettings->AlphaTextureSizeY;
 
-			FLinearColor AlphaScaleBias(
-				1.0f / (EdMode->UISettings->AlphaBrushScale * SizeX),
-				1.0f / (EdMode->UISettings->AlphaBrushScale * SizeY),
-				EdMode->UISettings->AlphaBrushPanU,
-				EdMode->UISettings->AlphaBrushPanV
-				);
+			FLinearColor AlphaScaleBias;
+			float Angle;
+			if (EdMode->UISettings->bUseWorldSpacePatternBrush)
+			{
+				FVector2D LocalOrigin = -FVector2D(Proxy->LandscapeActorToWorld().InverseTransformPosition(FVector(EdMode->UISettings->WorldSpacePatternBrushSettings.Origin, 0.0f)));
+				const FVector2D Scale = FVector2D(
+					ScaleXY / (EdMode->UISettings->WorldSpacePatternBrushSettings.RepeatSize * ((float)SizeX / SizeY)),
+					ScaleXY / (EdMode->UISettings->WorldSpacePatternBrushSettings.RepeatSize));
+				LocalOrigin *= Scale;
+				Angle = FMath::DegreesToRadians(-EdMode->UISettings->WorldSpacePatternBrushSettings.Rotation);
+				if (EdMode->UISettings->WorldSpacePatternBrushSettings.bCenterTextureOnOrigin)
+				{
+					LocalOrigin += FVector2D(0.5f, 0.5f).GetRotated(-Angle);
+				}
+				AlphaScaleBias = FLinearColor(
+					Scale.X,
+					Scale.Y,
+					LocalOrigin.X,
+					LocalOrigin.Y);
+			}
+			else
+			{
+				AlphaScaleBias = FLinearColor(
+					1.0f / (EdMode->UISettings->AlphaBrushScale * SizeX),
+					1.0f / (EdMode->UISettings->AlphaBrushScale * SizeY),
+					EdMode->UISettings->AlphaBrushPanU,
+					EdMode->UISettings->AlphaBrushPanV);
+				Angle = FMath::DegreesToRadians(EdMode->UISettings->AlphaBrushRotation);
+			}
 
-			float Angle = FMath::DegreesToRadians(EdMode->UISettings->AlphaBrushRotation);
-			FLinearColor LandscapeLocation(Proxy->GetActorLocation().X, Proxy->GetActorLocation().Y, Proxy->GetActorLocation().Z, Angle);
+			FVector LandscapeLocation = Proxy->LandscapeActorToWorld().GetTranslation();
+			FLinearColor LandscapeLocationParam(LandscapeLocation.X, LandscapeLocation.Y, LandscapeLocation.Z, Angle);
 
 			int32 Channel = EdMode->UISettings->AlphaTextureChannel;
 			FLinearColor AlphaTextureMask(Channel == 0 ? 1 : 0, Channel == 1 ? 1 : 0, Channel == 2 ? 1 : 0, Channel == 3 ? 1 : 0);
@@ -1102,7 +1125,7 @@ public:
 			{
 				UMaterialInstanceDynamic* const MaterialInstance = BrushMaterialInstancePair.Value;
 				MaterialInstance->SetVectorParameterValue(FName(TEXT("AlphaScaleBias")),    AlphaScaleBias);
-				MaterialInstance->SetVectorParameterValue(FName(TEXT("LandscapeLocation")), LandscapeLocation);
+				MaterialInstance->SetVectorParameterValue(FName(TEXT("LandscapeLocation")), LandscapeLocationParam);
 				MaterialInstance->SetVectorParameterValue(FName(TEXT("AlphaTextureMask")),  AlphaTextureMask);
 				MaterialInstance->SetTextureParameterValue(FName(TEXT("AlphaTexture")),     EdMode->UISettings->AlphaTexture);
 			}
diff --git a/Engine/Source/Editor/LandscapeEditor/Private/LandscapeEditorDetailCustomization_AlphaBrush.cpp b/Engine/Source/Editor/LandscapeEditor/Private/LandscapeEditorDetailCustomization_AlphaBrush.cpp
--- a/Engine/Source/Editor/LandscapeEditor/Private/LandscapeEditorDetailCustomization_AlphaBrush.cpp
+++ b/Engine/Source/Editor/LandscapeEditor/Private/LandscapeEditorDetailCustomization_AlphaBrush.cpp
@@ -223,6 +223,38 @@ void FLandscapeEditorDetailCustomization_AlphaBrush::CustomizeDetails(IDetailLay
 
 	BrushSettingsCategory.AddProperty(PropertyHandle_AlphaTextureChannel)
 	.OverrideResetToDefault(false, FSimpleDelegate());
+
+	if (IsBrushSetActive("BrushSet_Pattern"))
+	{
+		TSharedRef<IPropertyHandle> PropertyHandle_bUseWorldSpacePatternBrush = DetailBuilder.GetProperty(GET_MEMBER_NAME_CHECKED(ULandscapeEditorObject, bUseWorldSpacePatternBrush));
+		TSharedRef<IPropertyHandle> PropertyHandle_AlphaBrushScale    = DetailBuilder.GetProperty(GET_MEMBER_NAME_CHECKED(ULandscapeEditorObject, AlphaBrushScale));
+		TSharedRef<IPropertyHandle> PropertyHandle_AlphaBrushRotation = DetailBuilder.GetProperty(GET_MEMBER_NAME_CHECKED(ULandscapeEditorObject, AlphaBrushRotation));
+		TSharedRef<IPropertyHandle> PropertyHandle_AlphaBrushPanU     = DetailBuilder.GetProperty(GET_MEMBER_NAME_CHECKED(ULandscapeEditorObject, AlphaBrushPanU));
+		TSharedRef<IPropertyHandle> PropertyHandle_AlphaBrushPanV     = DetailBuilder.GetProperty(GET_MEMBER_NAME_CHECKED(ULandscapeEditorObject, AlphaBrushPanV));
+
+		auto NonWorld_Visibility = TAttribute<EVisibility>::Create(TAttribute<EVisibility>::FGetter::CreateLambda([=]() { return GetPropertyValue<bool>(PropertyHandle_bUseWorldSpacePatternBrush) ? EVisibility::Collapsed : EVisibility::Visible;   } ));
+		auto World_Visibility    = TAttribute<EVisibility>::Create(TAttribute<EVisibility>::FGetter::CreateLambda([=]() { return GetPropertyValue<bool>(PropertyHandle_bUseWorldSpacePatternBrush) ? EVisibility::Visible   : EVisibility::Collapsed; } ));
+
+		BrushSettingsCategory.AddProperty(PropertyHandle_bUseWorldSpacePatternBrush);
+		BrushSettingsCategory.AddProperty(PropertyHandle_AlphaBrushScale)
+			.Visibility(NonWorld_Visibility);
+		BrushSettingsCategory.AddProperty(PropertyHandle_AlphaBrushRotation)
+			.Visibility(NonWorld_Visibility);
+		BrushSettingsCategory.AddProperty(PropertyHandle_AlphaBrushPanU)
+			.Visibility(NonWorld_Visibility);
+		BrushSettingsCategory.AddProperty(PropertyHandle_AlphaBrushPanV)
+			.Visibility(NonWorld_Visibility);
+
+		TSharedRef<IPropertyHandle> PropertyHandle_WorldSpacePatternBrushSettings = DetailBuilder.GetProperty(GET_MEMBER_NAME_CHECKED(ULandscapeEditorObject, WorldSpacePatternBrushSettings));
+		PropertyHandle_WorldSpacePatternBrushSettings->MarkHiddenByCustomization();
+		uint32 NumWorldSpaceSettings = 0;
+		PropertyHandle_WorldSpacePatternBrushSettings->GetNumChildren(NumWorldSpaceSettings);
+		for (uint32 i = 0; i < NumWorldSpaceSettings; i++)
+		{
+			BrushSettingsCategory.AddProperty(PropertyHandle_WorldSpacePatternBrushSettings->GetChildHandle(i))
+				.Visibility(World_Visibility);
+		}
+	}
 }
 END_SLATE_FUNCTION_BUILD_OPTIMIZATION
 
diff --git a/Engine/Source/Editor/LandscapeEditor/Private/LandscapeEditorObject.cpp b/Engine/Source/Editor/LandscapeEditor/Private/LandscapeEditorObject.cpp
--- a/Engine/Source/Editor/LandscapeEditor/Private/LandscapeEditorObject.cpp
+++ b/Engine/Source/Editor/LandscapeEditor/Private/LandscapeEditorObject.cpp
@@ -82,6 +82,8 @@ ULandscapeEditorObject::ULandscapeEditorObject(const FObjectInitializer& ObjectI
 	, AlphaBrushRotation(0.0f)
 	, AlphaBrushPanU(0.5f)
 	, AlphaBrushPanV(0.5f)
+	, bUseWorldSpacePatternBrush(false)
+	, WorldSpacePatternBrushSettings(FLandscapePatternBrushWorldSpaceSettings{FVector2D::ZeroVector, 0.0f, false, 3200})
 	, AlphaTexture(NULL)
 	, AlphaTextureChannel(EColorChannel::Red)
 	, AlphaTextureSizeX(1)
@@ -171,6 +173,10 @@ void ULandscapeEditorObject::Load()
 	GConfig->GetFloat(TEXT("LandscapeEdit"), TEXT("AlphaBrushRotation"), AlphaBrushRotation, GEditorUserSettingsIni);
 	GConfig->GetFloat(TEXT("LandscapeEdit"), TEXT("AlphaBrushPanU"), AlphaBrushPanU, GEditorUserSettingsIni);
 	GConfig->GetFloat(TEXT("LandscapeEdit"), TEXT("AlphaBrushPanV"), AlphaBrushPanV, GEditorUserSettingsIni);
+	GConfig->GetBool(TEXT("LandscapeEdit"), TEXT("bUseWorldSpacePatternBrush"), bUseWorldSpacePatternBrush, GEditorUserSettingsIni);
+	GConfig->GetVector2D(TEXT("LandscapeEdit"), TEXT("WorldSpacePatternBrushSettings.Origin"), WorldSpacePatternBrushSettings.Origin, GEditorUserSettingsIni);
+	GConfig->GetBool(TEXT("LandscapeEdit"), TEXT("WorldSpacePatternBrushSettings.bCenterTextureOnOrigin"), WorldSpacePatternBrushSettings.bCenterTextureOnOrigin, GEditorUserSettingsIni);
+	GConfig->GetFloat(TEXT("LandscapeEdit"), TEXT("WorldSpacePatternBrushSettings.RepeatSize"), WorldSpacePatternBrushSettings.RepeatSize, GEditorUserSettingsIni);
 	FString AlphaTextureName = (AlphaTexture != NULL) ? AlphaTexture->GetPathName() : FString();
 	int32 InAlphaTextureChannel = AlphaTextureChannel;
 	GConfig->GetString(TEXT("LandscapeEdit"), TEXT("AlphaTextureName"), AlphaTextureName, GEditorUserSettingsIni);
@@ -289,6 +295,9 @@ void ULandscapeEditorObject::Save()
 	GConfig->SetFloat(TEXT("LandscapeEdit"), TEXT("AlphaBrushRotation"), AlphaBrushRotation, GEditorUserSettingsIni);
 	GConfig->SetFloat(TEXT("LandscapeEdit"), TEXT("AlphaBrushPanU"), AlphaBrushPanU, GEditorUserSettingsIni);
 	GConfig->SetFloat(TEXT("LandscapeEdit"), TEXT("AlphaBrushPanV"), AlphaBrushPanV, GEditorUserSettingsIni);
+	GConfig->SetVector2D(TEXT("LandscapeEdit"), TEXT("WorldSpacePatternBrushSettings.Origin"), WorldSpacePatternBrushSettings.Origin, GEditorUserSettingsIni);
+	GConfig->SetBool(TEXT("LandscapeEdit"), TEXT("WorldSpacePatternBrushSettings.bCenterTextureOnOrigin"), WorldSpacePatternBrushSettings.bCenterTextureOnOrigin, GEditorUserSettingsIni);
+	GConfig->SetFloat(TEXT("LandscapeEdit"), TEXT("WorldSpacePatternBrushSettings.RepeatSize"), WorldSpacePatternBrushSettings.RepeatSize, GEditorUserSettingsIni);
 	const FString AlphaTextureName = (AlphaTexture != NULL) ? AlphaTexture->GetPathName() : FString();
 	GConfig->SetString(TEXT("LandscapeEdit"), TEXT("AlphaTextureName"), *AlphaTextureName, GEditorUserSettingsIni);
 	GConfig->SetInt(TEXT("LandscapeEdit"), TEXT("AlphaTextureChannel"), (int32)AlphaTextureChannel, GEditorUserSettingsIni);
diff --git a/Engine/Source/Runtime/Core/Private/Misc/ConfigCacheIni.cpp b/Engine/Source/Runtime/Core/Private/Misc/ConfigCacheIni.cpp
--- a/Engine/Source/Runtime/Core/Private/Misc/ConfigCacheIni.cpp
+++ b/Engine/Source/Runtime/Core/Private/Misc/ConfigCacheIni.cpp
@@ -2107,6 +2107,21 @@ bool FConfigCacheIni::GetColor
 	return false;
 }
 
+bool FConfigCacheIni::GetVector2D(
+	const TCHAR*   Section,
+	const TCHAR*   Key,
+	FVector2D&     Value,
+	const FString& Filename)
+{
+	FString Text;
+	if (GetString(Section, Key, Text, Filename))
+	{
+		return Value.InitFromString(Text);
+	}
+	return false;
+}
+
+
 bool FConfigCacheIni::GetVector
 (
  const TCHAR*		Section,
@@ -2262,6 +2277,15 @@ void FConfigCacheIni::SetColor
 	SetString( Section, Key, *Value.ToString(), Filename );
 }
 
+void FConfigCacheIni::SetVector2D(
+	const TCHAR*   Section,
+	const TCHAR*   Key,
+	FVector2D      Value,
+	const FString& Filename)
+{
+	SetString(Section, Key, *Value.ToString(), Filename);
+}
+
 void FConfigCacheIni::SetVector
 (
  const TCHAR*		Section,
diff --git a/Engine/Source/Runtime/Core/Public/Math/Vector2D.h b/Engine/Source/Runtime/Core/Public/Math/Vector2D.h
--- a/Engine/Source/Runtime/Core/Public/Math/Vector2D.h
+++ b/Engine/Source/Runtime/Core/Public/Math/Vector2D.h
@@ -368,6 +368,14 @@ public:
 	 */
 	float SizeSquared() const;
 
+	/**
+	 * Rotates around axis (0,0,1)
+	 *
+	 * @param AngleDeg Angle to rotate (in degrees)
+	 * @return Rotated Vector
+	 */
+	FVector2D GetRotated(float AngleDeg) const;
+
 	/**
 	 * Gets a normalized copy of the vector, checking it is safe to do so based on the length.
 	 * Returns zero vector if vector length is too small to safely normalize.
@@ -740,6 +748,19 @@ FORCEINLINE float FVector2D::SizeSquared() const
 	return X*X + Y*Y;
 }
 
+FORCEINLINE FVector2D FVector2D::GetRotated(const float AngleDeg) const
+{
+	// Based on FVector::RotateAngleAxis with Axis(0,0,1)
+
+	float S, C;
+	FMath::SinCos(&S, &C, FMath::DegreesToRadians(AngleDeg));
+
+	const float OMC = 1.0f - C;
+
+	return FVector2D(
+		C * X - S * Y,
+		S * X + C * Y);
+}
 
 FORCEINLINE FVector2D FVector2D::GetSafeNormal(float Tolerance) const
 {	
diff --git a/Engine/Source/Runtime/Core/Public/Misc/ConfigCacheIni.h b/Engine/Source/Runtime/Core/Public/Misc/ConfigCacheIni.h
--- a/Engine/Source/Runtime/Core/Public/Misc/ConfigCacheIni.h
+++ b/Engine/Source/Runtime/Core/Public/Misc/ConfigCacheIni.h
@@ -368,6 +368,11 @@ public:
 		FColor&				Value,
 		const FString&	Filename
 	);
+	bool GetVector2D(
+		const TCHAR*   Section,
+		const TCHAR*   Key,
+		FVector2D&     Value,
+		const FString& Filename);
 	bool GetVector
 	(
 		const TCHAR*		Section,
@@ -445,6 +450,11 @@ public:
 		FColor				Value,
 		const FString&	Filename
 	);
+	void SetVector2D(
+		const TCHAR*   Section,
+		const TCHAR*   Key,
+		FVector2D      Value,
+		const FString& Filename);
 	void SetVector
 	(
 		const TCHAR*		Section,
