Fix horrific use of "#pragma clang diagnostic ignored "-Wreorder"" in landscape editor data caches with a small refactor

diff --git a/Engine/Source/Editor/LandscapeEditor/Private/LandscapeEdModeTools.h b/Engine/Source/Editor/LandscapeEditor/Private/LandscapeEdModeTools.h
--- a/Engine/Source/Editor/LandscapeEditor/Private/LandscapeEdModeTools.h
+++ b/Engine/Source/Editor/LandscapeEditor/Private/LandscapeEdModeTools.h
@@ -262,8 +262,12 @@ inline void LowPassFilter(int32 X1, int32 Y1, int32 X2, int32 Y2, FLandscapeBrus
 template<class Accessor, typename AccessorType>
 struct TLandscapeEditCache
 {
-	TLandscapeEditCache(Accessor& InDataAccess)
-		: DataAccess(InDataAccess)
+public:
+	typedef AccessorType DataType;
+	Accessor DataAccess;
+
+	TLandscapeEditCache(const FLandscapeToolTarget& InTarget)
+		: DataAccess(InTarget)
 		, Valid(false)
 	{
 	}
@@ -539,8 +543,6 @@ struct TLandscapeEditCache
 		DataAccess.Flush();
 	}
 
-protected:
-	Accessor& DataAccess;
 private:
 	// X2/Y2 Coordinates are "inclusive" max values
 	void CacheOriginalData(int32 X1, int32 Y1, int32 X2, int32 Y2)
@@ -587,6 +589,11 @@ struct FHeightmapAccessor
 		LandscapeEdit = new FLandscapeEditDataInterface(InLandscapeInfo);
 	}
 
+	FHeightmapAccessor(const FLandscapeToolTarget& InTarget)
+		: FHeightmapAccessor(InTarget.LandscapeInfo.Get())
+	{
+	}
+
 	// accessors
 	void GetData(int32& X1, int32& Y1, int32& X2, int32& Y2, TMap<FIntPoint, uint16>& Data)
 	{
@@ -698,25 +705,14 @@ private:
 	TSet<ULandscapeComponent*> ChangedComponents;
 };
 
-struct FLandscapeHeightCache : public TLandscapeEditCache < FHeightmapAccessor<true>, uint16 >
+struct FLandscapeHeightCache : public TLandscapeEditCache<FHeightmapAccessor<true>, uint16>
 {
-	typedef uint16 DataType;
 	static uint16 ClampValue(int32 Value) { return FMath::Clamp(Value, 0, LandscapeDataAccess::MaxValue); }
 
-	FHeightmapAccessor<true> HeightmapAccessor;
-
-#ifdef __clang__ // @todo
-	#pragma clang diagnostic push
-	#pragma clang diagnostic ignored "-Wreorder"
-#endif
 	FLandscapeHeightCache(const FLandscapeToolTarget& InTarget)
-		: HeightmapAccessor(InTarget.LandscapeInfo.Get())
-		, TLandscapeEditCache(HeightmapAccessor)
+		: TLandscapeEditCache(InTarget)
 	{
 	}
-#ifdef __clang__
-	#pragma clang diagnostic pop
-#endif
 };
 
 //
@@ -732,6 +728,11 @@ struct FXYOffsetmapAccessor
 		LandscapeEdit = new FLandscapeEditDataInterface(InLandscapeInfo);
 	}
 
+	FXYOffsetmapAccessor(const FLandscapeToolTarget& InTarget)
+		: FXYOffsetmapAccessor(InTarget.LandscapeInfo.Get())
+	{
+	}
+
 	// accessors
 	void GetData(int32& X1, int32& Y1, int32& X2, int32& Y2, TMap<FIntPoint, FVector>& Data)
 	{
@@ -865,15 +866,10 @@ private:
 };
 
 template<bool bInUseInterp>
-struct FLandscapeXYOffsetCache : public TLandscapeEditCache < FXYOffsetmapAccessor<bInUseInterp>, FVector >
+struct FLandscapeXYOffsetCache : public TLandscapeEditCache<FXYOffsetmapAccessor<bInUseInterp>, FVector>
 {
-	typedef FVector DataType;
-
-	FXYOffsetmapAccessor<bInUseInterp> XYOffsetmapAccessor;
-
 	FLandscapeXYOffsetCache(const FLandscapeToolTarget& InTarget)
-		: TLandscapeEditCache< FXYOffsetmapAccessor<bInUseInterp>, FVector >(XYOffsetmapAccessor)
-		, XYOffsetmapAccessor(InTarget.LandscapeInfo.Get())
+		: TLandscapeEditCache(InTarget)
 	{
 	}
 };
@@ -906,6 +902,11 @@ struct FAlphamapAccessor
 		}
 	}
 
+	FAlphamapAccessor(const FLandscapeToolTarget& InTarget)
+		: FAlphamapAccessor(InTarget.LandscapeInfo.Get(), InTarget.LayerInfo.Get())
+	{
+	}
+
 	~FAlphamapAccessor()
 	{
 		// Recreate collision for modified components to update the physical materials
@@ -965,46 +966,32 @@ private:
 	bool bBlendWeight;
 };
 
-struct FLandscapeAlphaCache : public TLandscapeEditCache < FAlphamapAccessor<true, false>, uint8 >
+struct FLandscapeAlphaCache : public TLandscapeEditCache<FAlphamapAccessor<true, false>, uint8>
 {
-	typedef uint8 DataType;
 	static uint8 ClampValue(int32 Value) { return FMath::Clamp(Value, 0, 255); }
 
-	FAlphamapAccessor<true, false> AlphamapAccessor;
-
-#ifdef __clang__ // @todo
-	#pragma clang diagnostic push
-	#pragma clang diagnostic ignored "-Wreorder"
-#endif
 	FLandscapeAlphaCache(const FLandscapeToolTarget& InTarget)
-		: AlphamapAccessor(InTarget.LandscapeInfo.Get(), InTarget.LayerInfo.Get())
-		, TLandscapeEditCache(AlphamapAccessor)
+		: TLandscapeEditCache(InTarget)
 	{
 	}
-#ifdef __clang__
-	#pragma clang diagnostic pop
-#endif
 };
 
-struct FLandscapeVisCache : public TLandscapeEditCache < FAlphamapAccessor<false, false>, uint8 >
+struct FVisibilityAccessor : public FAlphamapAccessor<false, false>
 {
-	typedef uint8 DataType;
-	static uint8 ClampValue(int32 Value) { return FMath::Clamp(Value, 0, 255); }
+	FVisibilityAccessor(const FLandscapeToolTarget& InTarget)
+		: FAlphamapAccessor<false, false>(InTarget.LandscapeInfo.Get(), ALandscapeProxy::VisibilityLayer)
+	{
+	}
+};
 
-	FAlphamapAccessor<false, false> AlphamapAccessor;
+struct FLandscapeVisCache : public TLandscapeEditCache<FAlphamapAccessor<false, false>, uint8>
+{
+	static uint8 ClampValue(int32 Value) { return FMath::Clamp(Value, 0, 255); }
 
-#ifdef __clang__ // @todo
-	#pragma clang diagnostic push
-	#pragma clang diagnostic ignored "-Wreorder"
-#endif
 	FLandscapeVisCache(const FLandscapeToolTarget& InTarget)
-		: AlphamapAccessor(InTarget.LandscapeInfo.Get(), ALandscapeProxy::VisibilityLayer)
-		, TLandscapeEditCache(AlphamapAccessor)
+		: TLandscapeEditCache(InTarget)
 	{
 	}
-#ifdef __clang__
-	#pragma clang diagnostic pop
-#endif
 };
 
 //
@@ -1020,6 +1007,11 @@ struct FFullWeightmapAccessor
 	{
 	}
 
+	FFullWeightmapAccessor(const FLandscapeToolTarget& InTarget)
+		: FFullWeightmapAccessor(InTarget.LandscapeInfo.Get())
+	{
+	}
+
 	~FFullWeightmapAccessor()
 	{
 		// Recreate collision for modified components to update the physical materials
@@ -1043,7 +1035,6 @@ struct FFullWeightmapAccessor
 		}
 	}
 
-
 	void GetData(int32& X1, int32& Y1, int32& X2, int32& Y2, TMap<FIntPoint, TArray<uint8>>& Data)
 	{
 		// Do not Support for interpolation....
@@ -1083,24 +1074,12 @@ private:
 	TSet<ULandscapeComponent*> ModifiedComponents;
 };
 
-struct FLandscapeFullWeightCache : public TLandscapeEditCache < FFullWeightmapAccessor<false>, TArray<uint8> >
+struct FLandscapeFullWeightCache : public TLandscapeEditCache<FFullWeightmapAccessor<false>, TArray<uint8>>
 {
-	typedef TArray<uint8> DataType;
-
-	FFullWeightmapAccessor<false> WeightmapAccessor;
-
-#ifdef __clang__ // @todo
-	#pragma clang diagnostic push
-	#pragma clang diagnostic ignored "-Wreorder"
-#endif
 	FLandscapeFullWeightCache(const FLandscapeToolTarget& InTarget)
-		: WeightmapAccessor(InTarget.LandscapeInfo.Get())
-		, TLandscapeEditCache(WeightmapAccessor)
+		: TLandscapeEditCache(InTarget)
 	{
 	}
-#ifdef __clang__
-	#pragma clang diagnostic pop
-#endif
 
 	// Only for all weight case... the accessor type should be TArray<uint8>
 	void GetCachedData(int32 X1, int32 Y1, int32 X2, int32 Y2, TArray<uint8>& OutData, int32 ArraySize)
@@ -1165,7 +1144,7 @@ struct FLandscapeFullWeightCache : public TLandscapeEditCache < FFullWeightmapAc
 
 	void AddDirtyLayer(ULandscapeLayerInfoObject* LayerInfo)
 	{
-		WeightmapAccessor.DirtyLayerInfos.Add(LayerInfo);
+		DataAccess.DirtyLayerInfos.Add(LayerInfo);
 	}
 };
 
@@ -1181,6 +1160,11 @@ struct FDatamapAccessor
 	{
 	}
 
+	FDatamapAccessor(const FLandscapeToolTarget& InTarget)
+		: FDatamapAccessor(InTarget.LandscapeInfo.Get())
+	{
+	}
+
 	void GetData(int32& X1, int32& Y1, int32& X2, int32& Y2, TMap<FIntPoint, uint8>& Data)
 	{
 		LandscapeEdit.GetSelectData(X1, Y1, X2, Y2, Data);
@@ -1208,25 +1192,14 @@ private:
 	FLandscapeEditDataInterface LandscapeEdit;
 };
 
-struct FLandscapeDataCache : public TLandscapeEditCache < FDatamapAccessor<false>, uint8 >
+struct FLandscapeDataCache : public TLandscapeEditCache<FDatamapAccessor<false>, uint8>
 {
-	typedef uint8 DataType;
 	static uint8 ClampValue(int32 Value) { return FMath::Clamp(Value, 0, 255); }
 
-	FDatamapAccessor<false> DataAccessor;
-
-#ifdef __clang__ // @todo
-	#pragma clang diagnostic push
-	#pragma clang diagnostic ignored "-Wreorder"
-#endif
 	FLandscapeDataCache(const FLandscapeToolTarget& InTarget)
-		: DataAccessor(InTarget.LandscapeInfo.Get())
-		, TLandscapeEditCache(DataAccessor)
+		: TLandscapeEditCache(InTarget)
 	{
 	}
-#ifdef __clang__
-	#pragma clang diagnostic pop
-#endif
 };
 
 
