Fixed clang compile. Seems that disabling of the warning about incorrect initialization order was also allowing some bad syntax through, so after removing the disable it no longer compiles. I can't really understand why it compiled before! The error has nothing to do with the warning that was disabled!

diff --git a/Engine/Source/Editor/LandscapeEditor/Private/LandscapeEdModeTools.h b/Engine/Source/Editor/LandscapeEditor/Private/LandscapeEdModeTools.h
--- a/Engine/Source/Editor/LandscapeEditor/Private/LandscapeEdModeTools.h
+++ b/Engine/Source/Editor/LandscapeEditor/Private/LandscapeEdModeTools.h
@@ -710,7 +710,7 @@ struct FLandscapeHeightCache : public TLandscapeEditCache<FHeightmapAccessor<tru
 	static uint16 ClampValue(int32 Value) { return FMath::Clamp(Value, 0, LandscapeDataAccess::MaxValue); }
 
 	FLandscapeHeightCache(const FLandscapeToolTarget& InTarget)
-		: TLandscapeEditCache(InTarget)
+		: TLandscapeEditCache<FHeightmapAccessor<true>, uint16>(InTarget)
 	{
 	}
 };
@@ -869,7 +869,7 @@ template<bool bInUseInterp>
 struct FLandscapeXYOffsetCache : public TLandscapeEditCache<FXYOffsetmapAccessor<bInUseInterp>, FVector>
 {
 	FLandscapeXYOffsetCache(const FLandscapeToolTarget& InTarget)
-		: TLandscapeEditCache(InTarget)
+		: TLandscapeEditCache<FXYOffsetmapAccessor<bInUseInterp>, FVector>(InTarget)
 	{
 	}
 };
@@ -971,7 +971,7 @@ struct FLandscapeAlphaCache : public TLandscapeEditCache<FAlphamapAccessor<true,
 	static uint8 ClampValue(int32 Value) { return FMath::Clamp(Value, 0, 255); }
 
 	FLandscapeAlphaCache(const FLandscapeToolTarget& InTarget)
-		: TLandscapeEditCache(InTarget)
+		: TLandscapeEditCache<FAlphamapAccessor<true, false>, uint8>(InTarget)
 	{
 	}
 };
@@ -989,7 +989,7 @@ struct FLandscapeVisCache : public TLandscapeEditCache<FAlphamapAccessor<false,
 	static uint8 ClampValue(int32 Value) { return FMath::Clamp(Value, 0, 255); }
 
 	FLandscapeVisCache(const FLandscapeToolTarget& InTarget)
-		: TLandscapeEditCache(InTarget)
+		: TLandscapeEditCache<FAlphamapAccessor<false, false>, uint8>(InTarget)
 	{
 	}
 };
@@ -1077,7 +1077,7 @@ private:
 struct FLandscapeFullWeightCache : public TLandscapeEditCache<FFullWeightmapAccessor<false>, TArray<uint8>>
 {
 	FLandscapeFullWeightCache(const FLandscapeToolTarget& InTarget)
-		: TLandscapeEditCache(InTarget)
+		: TLandscapeEditCache<FFullWeightmapAccessor<false>, TArray<uint8>>(InTarget)
 	{
 	}
 
@@ -1197,7 +1197,7 @@ struct FLandscapeDataCache : public TLandscapeEditCache<FDatamapAccessor<false>,
 	static uint8 ClampValue(int32 Value) { return FMath::Clamp(Value, 0, 255); }
 
 	FLandscapeDataCache(const FLandscapeToolTarget& InTarget)
-		: TLandscapeEditCache(InTarget)
+		: TLandscapeEditCache<FDatamapAccessor<false>, uint8>(InTarget)
 	{
 	}
 };
