Fix mac compile error

diff --git a/Engine/Source/Editor/LandscapeEditor/Private/LandscapeEdModeComponentTools.cpp b/Engine/Source/Editor/LandscapeEditor/Private/LandscapeEdModeComponentTools.cpp
--- a/Engine/Source/Editor/LandscapeEditor/Private/LandscapeEdModeComponentTools.cpp
+++ b/Engine/Source/Editor/LandscapeEditor/Private/LandscapeEdModeComponentTools.cpp
@@ -12,6 +12,7 @@
 #include "LandscapeSplineProxies.h"
 #include "LandscapeEditorModule.h"
 #include "Editor/PropertyEditor/Public/PropertyEditorModule.h"
+#include "LandscapeEdMode.h"
 #include "LandscapeEdModeTools.h"
 #include "InstancedFoliageActor.h"
 #include "ComponentReregisterContext.h"
