Route FlushGrassComponents to correct LandscapeProxy(s) when landscape editing

[Intergated CL 2433662 from UE4-GDC2014]

diff --git a/Engine/Source/Editor/LandscapeEditor/Private/LandscapeEdModeTools.h b/Engine/Source/Editor/LandscapeEditor/Private/LandscapeEdModeTools.h
--- a/Engine/Source/Editor/LandscapeEditor/Private/LandscapeEdModeTools.h
+++ b/Engine/Source/Editor/LandscapeEditor/Private/LandscapeEdModeTools.h
@@ -619,7 +619,7 @@ struct FHeightmapAccessor
 			}
 
 			// Flush dynamic foliage (grass)
-			LandscapeInfo->GetLandscapeProxy()->FlushGrassComponents(&Components);
+			ALandscapeProxy::DistributeFlushGrassComponents(Components);
 
 			// Notify foliage to move any attached instances
 			bool bUpdateFoliage = false;
@@ -792,7 +792,7 @@ struct FXYOffsetmapAccessor
 			}
 
 			// Flush dynamic foliage (grass)
-			LandscapeInfo->GetLandscapeProxy()->FlushGrassComponents(&Components);
+			ALandscapeProxy::DistributeFlushGrassComponents(Components);
 
 			// Notify foliage to move any attached instances
 			bool bUpdateFoliage = false;
@@ -949,7 +949,7 @@ struct FAlphamapAccessor
 		if (LandscapeEdit.GetComponentsInRegion(X1, Y1, X2, Y2, &Components))
 		{
 			// Flush dynamic foliage (grass)
-			LandscapeInfo->GetLandscapeProxy()->FlushGrassComponents(&Components);
+			ALandscapeProxy::DistributeFlushGrassComponents(Components);
 
 			LandscapeEdit.SetAlphaData(LayerInfo, X1, Y1, X2, Y2, Data, 0, PaintingRestriction, bBlendWeight, bUseTotalNormalize);
 			ModifiedComponents.Append(Components);
@@ -1056,7 +1056,7 @@ struct FFullWeightmapAccessor
 		if (LandscapeEdit.GetComponentsInRegion(X1, Y1, X2, Y2, &Components))
 		{
 			// Flush dynamic foliage (grass)
-			LandscapeInfo->GetLandscapeProxy()->FlushGrassComponents(&Components);
+			ALandscapeProxy::DistributeFlushGrassComponents(Components);
 
 			LandscapeEdit.SetAlphaData(DirtyLayerInfos, X1, Y1, X2, Y2, Data, 0, PaintingRestriction);
 			ModifiedComponents.Append(Components);
diff --git a/Engine/Source/Runtime/Landscape/Classes/LandscapeProxy.h b/Engine/Source/Runtime/Landscape/Classes/LandscapeProxy.h
--- a/Engine/Source/Runtime/Landscape/Classes/LandscapeProxy.h
+++ b/Engine/Source/Runtime/Landscape/Classes/LandscapeProxy.h
@@ -524,6 +524,9 @@ public:
 	/** Flush the grass cache */
 	LANDSCAPE_API void FlushGrassComponents(const TSet<ULandscapeComponent*>* OnlyForComponents = nullptr, bool bFlushGrassMaps = true);
 
+	/* Static version to distribute to the appropriate proxy for each component */
+	LANDSCAPE_API static void DistributeFlushGrassComponents(const TSet<ULandscapeComponent*>& Components, bool bFlushGrassMaps = true);
+
 	/** 
 		Update Grass 
 		* @param Cameras to use for culling, if empty, then NO culling
diff --git a/Engine/Source/Runtime/Landscape/Private/LandscapeEdit.cpp b/Engine/Source/Runtime/Landscape/Private/LandscapeEdit.cpp
--- a/Engine/Source/Runtime/Landscape/Private/LandscapeEdit.cpp
+++ b/Engine/Source/Runtime/Landscape/Private/LandscapeEdit.cpp
@@ -2924,12 +2924,6 @@ void ULandscapeLayerInfoObject::PostEditChangeProperty(FPropertyChangedEvent& Pr
 {
 	static const FName NAME_Hardness = FName(TEXT("Hardness"));
 	static const FName NAME_PhysMaterial = FName(TEXT("PhysMaterial"));
-	static const FName NAME_GrassMesh = FName(TEXT("GrassMesh"));
-	static const FName NAME_GrassDensity = FName(TEXT("GrassDensity"));
-	static const FName NAME_PlacementJitter = FName(TEXT("PlacementJitter"));
-	static const FName NAME_EndCullDistance = FName(TEXT("EndCullDistance"));
-	static const FName NAME_RandomRotation = FName(TEXT("RandomRotation"));
-	static const FName NAME_AlignToSurface = FName(TEXT("AlignToSurface"));
 
 	Super::PostEditChangeProperty(PropertyChangedEvent);
 
@@ -2961,33 +2955,6 @@ void ULandscapeLayerInfoObject::PostEditChangeProperty(FPropertyChangedEvent& Pr
 				}
 			}
 		}
-		else if (
-			PropertyName == NAME_GrassMesh ||
-			PropertyName == NAME_GrassDensity ||
-			PropertyName == NAME_PlacementJitter ||
-			PropertyName == NAME_EndCullDistance ||
-			PropertyName == NAME_AlignToSurface ||
-			PropertyName == NAME_RandomRotation
-			)
-		{
-			// Only care current world object
-			for (TActorIterator<ALandscapeProxy> It(GWorld); It; ++It)
-			{
-				ALandscapeProxy* Proxy = *It;
-				ULandscapeInfo* Info = Proxy->GetLandscapeInfo(false);
-				if (Info)
-				{
-					for (int32 i = 0; i < Info->Layers.Num(); ++i)
-					{
-						if (Info->Layers[i].LayerInfoObj == this)
-						{
-							Proxy->FlushGrassComponents();
-							break;
-						}
-					}
-				}
-			}
-		}
 	}
 }
 
diff --git a/Engine/Source/Runtime/Landscape/Private/LandscapeGrass.cpp b/Engine/Source/Runtime/Landscape/Private/LandscapeGrass.cpp
--- a/Engine/Source/Runtime/Landscape/Private/LandscapeGrass.cpp
+++ b/Engine/Source/Runtime/Landscape/Private/LandscapeGrass.cpp
@@ -1152,6 +1152,20 @@ struct FAsyncGrassBuilder : public FGrassBuilderBase
 	}
 };
 
+/* Static version to distribute to the appropriate proxy for each component */
+void ALandscapeProxy::DistributeFlushGrassComponents(const TSet<ULandscapeComponent*>& Components, bool bFlushGrassMaps)
+{
+	TMap<ALandscapeProxy*, TSet<ULandscapeComponent*>> ByProxy;
+	for (auto Component : Components)
+	{
+		ByProxy.FindOrAdd(Component->GetLandscapeProxy()).Add(Component);
+	}
+	for (auto It = ByProxy.CreateConstIterator(); It; ++It)
+	{
+		It.Key()->FlushGrassComponents(&It.Value(), bFlushGrassMaps);
+	}
+}
+
 void ALandscapeProxy::FlushGrassComponents(const TSet<ULandscapeComponent*>* OnlyForComponents, bool bFlushGrassMaps)
 {
 	if (OnlyForComponents)
