Fix world-space pattern brush options

diff --git a/Engine/Source/Editor/LandscapeEditor/Private/LandscapeEdModeBrushes.cpp b/Engine/Source/Editor/LandscapeEditor/Private/LandscapeEdModeBrushes.cpp
--- a/Engine/Source/Editor/LandscapeEditor/Private/LandscapeEdModeBrushes.cpp
+++ b/Engine/Source/Editor/LandscapeEditor/Private/LandscapeEdModeBrushes.cpp
@@ -1028,14 +1028,39 @@ public:
 
 			for (int32 X = Bounds.Min.X; X < Bounds.Max.X; X++)
 			{
-				// Find alphamap sample location
-				float SampleX = (float)X / EdMode->UISettings->AlphaBrushScale + (float)SizeX * EdMode->UISettings->AlphaBrushPanU;
-				float SampleY = (float)Y / EdMode->UISettings->AlphaBrushScale + (float)SizeY * EdMode->UISettings->AlphaBrushPanV;
+				float Angle;
+				FVector2D Scale;
+				FVector2D Bias;
+				if (EdMode->UISettings->bUseWorldSpacePatternBrush)
+				{
+					FVector2D LocalOrigin = -FVector2D(LandscapeInfo->GetLandscapeProxy()->LandscapeActorToWorld().InverseTransformPosition(FVector(EdMode->UISettings->WorldSpacePatternBrushSettings.Origin, 0.0f)));
+					const FVector2D LocalScale = FVector2D(
+						ScaleXY / (EdMode->UISettings->WorldSpacePatternBrushSettings.RepeatSize * ((float)SizeX / SizeY)),
+						ScaleXY / (EdMode->UISettings->WorldSpacePatternBrushSettings.RepeatSize));
+					LocalOrigin *= LocalScale;
+					Angle = -EdMode->UISettings->WorldSpacePatternBrushSettings.Rotation;
+					if (EdMode->UISettings->WorldSpacePatternBrushSettings.bCenterTextureOnOrigin)
+					{
+						LocalOrigin += FVector2D(0.5f, 0.5f).GetRotated(-Angle);
+					}
+					Scale = FVector2D(SizeX, SizeY) * LocalScale;
+					Bias = FVector2D(SizeX, SizeY) * LocalOrigin;
+				}
+				else
+				{
+					Scale.X = 1.0f / EdMode->UISettings->AlphaBrushScale;
+					Scale.Y = 1.0f / EdMode->UISettings->AlphaBrushScale;
+					Bias.X = SizeX * EdMode->UISettings->AlphaBrushPanU;
+					Bias.Y = SizeY * EdMode->UISettings->AlphaBrushPanV;
+					Angle = EdMode->UISettings->AlphaBrushRotation;
+				}
 
-				float Angle = PI * EdMode->UISettings->AlphaBrushRotation / 180.0f;
+				// Find alphamap sample location
+				FVector2D SamplePos = FVector2D(X, Y) * Scale + Bias;
+				SamplePos = SamplePos.GetRotated(Angle);
 
-				float ModSampleX = FMath::Fmod(SampleX * FMath::Cos(Angle) - SampleY * FMath::Sin(Angle), (float)SizeX);
-				float ModSampleY = FMath::Fmod(SampleY * FMath::Cos(Angle) + SampleX * FMath::Sin(Angle), (float)SizeY);
+				float ModSampleX = FMath::Fmod(SamplePos.X, (float)SizeX);
+				float ModSampleY = FMath::Fmod(SamplePos.Y, (float)SizeY);
 
 				if (ModSampleX < 0.0f)
 				{
@@ -1094,7 +1119,7 @@ public:
 					ScaleXY / (EdMode->UISettings->WorldSpacePatternBrushSettings.RepeatSize * ((float)SizeX / SizeY)),
 					ScaleXY / (EdMode->UISettings->WorldSpacePatternBrushSettings.RepeatSize));
 				LocalOrigin *= Scale;
-				Angle = FMath::DegreesToRadians(-EdMode->UISettings->WorldSpacePatternBrushSettings.Rotation);
+				Angle = -EdMode->UISettings->WorldSpacePatternBrushSettings.Rotation;
 				if (EdMode->UISettings->WorldSpacePatternBrushSettings.bCenterTextureOnOrigin)
 				{
 					LocalOrigin += FVector2D(0.5f, 0.5f).GetRotated(-Angle);
@@ -1112,8 +1137,9 @@ public:
 					1.0f / (EdMode->UISettings->AlphaBrushScale * SizeY),
 					EdMode->UISettings->AlphaBrushPanU,
 					EdMode->UISettings->AlphaBrushPanV);
-				Angle = FMath::DegreesToRadians(EdMode->UISettings->AlphaBrushRotation);
+				Angle = EdMode->UISettings->AlphaBrushRotation;
 			}
+			Angle = FMath::DegreesToRadians(Angle);
 
 			FVector LandscapeLocation = Proxy->LandscapeActorToWorld().GetTranslation();
 			FLinearColor LandscapeLocationParam(LandscapeLocation.X, LandscapeLocation.Y, LandscapeLocation.Z, Angle);
