One more fix for cross-level foliage handling on landscape edit

diff --git a/Engine/Source/Editor/LandscapeEditor/Private/LandscapeEdModeTools.h b/Engine/Source/Editor/LandscapeEditor/Private/LandscapeEdModeTools.h
--- a/Engine/Source/Editor/LandscapeEditor/Private/LandscapeEdModeTools.h
+++ b/Engine/Source/Editor/LandscapeEditor/Private/LandscapeEdModeTools.h
@@ -625,8 +625,7 @@ struct FHeightmapAccessor
 			bool bUpdateFoliage = false;
 			for (ULandscapeComponent* Component : Components)
 			{
-				AInstancedFoliageActor* IFA = AInstancedFoliageActor::GetInstancedFoliageActorForLevel(Component->GetComponentLevel());
-				if (IFA)
+				if (AInstancedFoliageActor::HasFoliageAttached(Component))
 				{
 					bUpdateFoliage = true;
 					break;
