PR #936: Fixes for PVS-Studio warnings from Paul Eremeeff
These are mostly fixes to redundant comparisons and initialization code, but there's also a fix for serious looking scope issue in FontCache.cpp

diff --git a/Engine/Source/Editor/LandscapeEditor/Private/LandscapeEdMode.cpp b/Engine/Source/Editor/LandscapeEditor/Private/LandscapeEdMode.cpp
--- a/Engine/Source/Editor/LandscapeEditor/Private/LandscapeEdMode.cpp
+++ b/Engine/Source/Editor/LandscapeEditor/Private/LandscapeEdMode.cpp
@@ -821,7 +821,7 @@ EEditAction::Type FEdModeLandscape::GetActionEditDuplicate()
 
 	if (NewLandscapePreviewMode == ENewLandscapePreviewMode::None)
 	{
-		if (CurrentTool != NULL && CurrentTool != NULL)
+		if (CurrentTool != NULL)
 		{
 			Result = CurrentTool->GetActionEditDuplicate();
 		}
@@ -836,7 +836,7 @@ EEditAction::Type FEdModeLandscape::GetActionEditDelete()
 
 	if (NewLandscapePreviewMode == ENewLandscapePreviewMode::None)
 	{
-		if (CurrentTool != NULL && CurrentTool != NULL)
+		if (CurrentTool != NULL)
 		{
 			Result = CurrentTool->GetActionEditDelete();
 		}
@@ -871,7 +871,7 @@ EEditAction::Type FEdModeLandscape::GetActionEditCut()
 
 	if (NewLandscapePreviewMode == ENewLandscapePreviewMode::None)
 	{
-		if (CurrentTool != NULL && CurrentTool != NULL)
+		if (CurrentTool != NULL)
 		{
 			Result = CurrentTool->GetActionEditCut();
 		}
@@ -893,7 +893,7 @@ EEditAction::Type FEdModeLandscape::GetActionEditCopy()
 
 	if (NewLandscapePreviewMode == ENewLandscapePreviewMode::None)
 	{
-		if (CurrentTool != NULL && CurrentTool != NULL)
+		if (CurrentTool != NULL)
 		{
 			Result = CurrentTool->GetActionEditCopy();
 		}
@@ -919,7 +919,7 @@ EEditAction::Type FEdModeLandscape::GetActionEditPaste()
 
 	if (NewLandscapePreviewMode == ENewLandscapePreviewMode::None)
 	{
-		if (CurrentTool != NULL && CurrentTool != NULL)
+		if (CurrentTool != NULL)
 		{
 			Result = CurrentTool->GetActionEditPaste();
 		}
@@ -950,7 +950,7 @@ bool FEdModeLandscape::ProcessEditDuplicate()
 
 	if (NewLandscapePreviewMode == ENewLandscapePreviewMode::None)
 	{
-		if (CurrentTool != NULL && CurrentTool != NULL)
+		if (CurrentTool != NULL)
 		{
 			Result = CurrentTool->ProcessEditDuplicate();
 		}
@@ -970,7 +970,7 @@ bool FEdModeLandscape::ProcessEditDelete()
 
 	if (NewLandscapePreviewMode == ENewLandscapePreviewMode::None)
 	{
-		if (CurrentTool != NULL && CurrentTool != NULL)
+		if (CurrentTool != NULL)
 		{
 			Result = CurrentTool->ProcessEditDelete();
 		}
@@ -990,7 +990,7 @@ bool FEdModeLandscape::ProcessEditCut()
 
 	if (NewLandscapePreviewMode == ENewLandscapePreviewMode::None)
 	{
-		if (CurrentTool != NULL && CurrentTool != NULL)
+		if (CurrentTool != NULL)
 		{
 			Result = CurrentTool->ProcessEditCut();
 		}
@@ -1010,7 +1010,7 @@ bool FEdModeLandscape::ProcessEditCopy()
 
 	if (NewLandscapePreviewMode == ENewLandscapePreviewMode::None)
 	{
-		if (CurrentTool != NULL && CurrentTool != NULL)
+		if (CurrentTool != NULL)
 		{
 			Result = CurrentTool->ProcessEditCopy();
 		}
@@ -1052,7 +1052,7 @@ bool FEdModeLandscape::ProcessEditPaste()
 
 	if (NewLandscapePreviewMode == ENewLandscapePreviewMode::None)
 	{
-		if (CurrentTool != NULL && CurrentTool != NULL)
+		if (CurrentTool != NULL)
 		{
 			Result = CurrentTool->ProcessEditPaste();
 		}
diff --git a/Engine/Source/Editor/MeshPaint/Private/MeshPaintEdMode.cpp b/Engine/Source/Editor/MeshPaint/Private/MeshPaintEdMode.cpp
--- a/Engine/Source/Editor/MeshPaint/Private/MeshPaintEdMode.cpp
+++ b/Engine/Source/Editor/MeshPaint/Private/MeshPaintEdMode.cpp
@@ -4467,7 +4467,6 @@ void FEdModeMeshPaint::UpdateTexturePaintTargetList()
 							|| Texture2D->LODGroup == TEXTUREGROUP_CharacterNormalMap
 							|| Texture2D->LODGroup == TEXTUREGROUP_WeaponNormalMap
 							|| Texture2D->LODGroup == TEXTUREGROUP_VehicleNormalMap
-							|| Texture2D->LODGroup == TEXTUREGROUP_WorldNormalMap
 							|| Texture2DName.Contains( TEXT("_N0" ))
 							|| Texture2DName.Contains( TEXT("_N_" ))
 							|| Texture2DName.Contains( TEXT("_NORMAL" ))
diff --git a/Engine/Source/Editor/UnrealEd/Private/Editor.cpp b/Engine/Source/Editor/UnrealEd/Private/Editor.cpp
--- a/Engine/Source/Editor/UnrealEd/Private/Editor.cpp
+++ b/Engine/Source/Editor/UnrealEd/Private/Editor.cpp
@@ -2027,7 +2027,7 @@ void UEditorEngine::CloseEditedWorldAssets(UWorld* InWorld)
 	for (int32 Index = 0; Index < InWorld->StreamingLevels.Num(); ++Index)
 	{
 		ULevelStreaming* LevelStreaming = InWorld->StreamingLevels[Index];
-		if (LevelStreaming && LevelStreaming->LoadedLevel && LevelStreaming->LoadedLevel)
+		if (LevelStreaming && LevelStreaming->LoadedLevel)
 		{
 			ClosingWorlds.Add(CastChecked<UWorld>(LevelStreaming->LoadedLevel->GetOuter()));
 		}
diff --git a/Engine/Source/Editor/UnrealEd/Private/EditorServer.cpp b/Engine/Source/Editor/UnrealEd/Private/EditorServer.cpp
--- a/Engine/Source/Editor/UnrealEd/Private/EditorServer.cpp
+++ b/Engine/Source/Editor/UnrealEd/Private/EditorServer.cpp
@@ -3651,7 +3651,7 @@ bool UEditorEngine::Map_Check( UWorld* InWorld, const TCHAR* Str, FOutputDevice&
 	{
 		if(bCheckDeprecatedOnly)
 		{
-			if(WarningCount > 0)
+			if(ErrorCount > 0)
 			{
 				MapCheckLog.Notify(LOCTEXT("MapCheckGenErrors", "Map check generated errors!"));
 			}
diff --git a/Engine/Source/Editor/UnrealEd/Private/StaticLightingSystem/StaticLightingSystem.cpp b/Engine/Source/Editor/UnrealEd/Private/StaticLightingSystem/StaticLightingSystem.cpp
--- a/Engine/Source/Editor/UnrealEd/Private/StaticLightingSystem/StaticLightingSystem.cpp
+++ b/Engine/Source/Editor/UnrealEd/Private/StaticLightingSystem/StaticLightingSystem.cpp
@@ -194,7 +194,7 @@ void FStaticLightingManager::SendBuildDoneNotification( bool AutoApplyFailed )
 		LOCTEXT( "LightBuildDiscard", "Discard" ),
 		LOCTEXT( "LightBuildDiscardToolTip", "Ignores the built lighting data generated." ),
 		FSimpleDelegate::CreateStatic( &FStaticLightingManager::ProcessLightingData, true ) );
-	ApplyNow.VisibilityOnSuccess = EVisibility::Collapsed;
+	Discard.VisibilityOnSuccess = EVisibility::Collapsed;
 
 	Info.ButtonDetails.Add( ApplyNow );
 	Info.ButtonDetails.Add( Discard );
diff --git a/Engine/Source/Runtime/Core/Private/Math/Color.cpp b/Engine/Source/Runtime/Core/Private/Math/Color.cpp
--- a/Engine/Source/Runtime/Core/Private/Math/Color.cpp
+++ b/Engine/Source/Runtime/Core/Private/Math/Color.cpp
@@ -271,14 +271,12 @@ FColor FColor::MakeRandomColor()
 
 FColor FColor::MakeRedToGreenColorFromScalar(float Scalar)
 {
-	int32 R,G,B;
-	R=G=B=0;
-
 	float RedSclr = FMath::Clamp<float>((1.0f - Scalar)/0.5f,0.f,1.f);
 	float GreenSclr = FMath::Clamp<float>((Scalar/0.5f),0.f,1.f);
-	R = FMath::TruncToInt(255 * RedSclr);
-	G = FMath::TruncToInt(255 * GreenSclr);
-	return FColor(R,G,B);
+	int32 R = FMath::TruncToInt(255 * RedSclr);
+	int32 G = FMath::TruncToInt(255 * GreenSclr);
+	int32 B = 0;
+	return FColor(R, G, B);
 }
 
 void ComputeAndFixedColorAndIntensity(const FLinearColor& InLinearColor,FColor& OutColor,float& OutIntensity)
diff --git a/Engine/Source/Runtime/Core/Private/Windows/WindowsTextInputMethodSystem.cpp b/Engine/Source/Runtime/Core/Private/Windows/WindowsTextInputMethodSystem.cpp
--- a/Engine/Source/Runtime/Core/Private/Windows/WindowsTextInputMethodSystem.cpp
+++ b/Engine/Source/Runtime/Core/Private/Windows/WindowsTextInputMethodSystem.cpp
@@ -429,7 +429,7 @@ bool FWindowsTextInputMethodSystem::InitializeTSF()
 		static const DWORD WindowsVistaMinorVersion = 0;
 
 		// Install profile notification sink for versions of Windows Vista and after.
-		if(WindowsMajorVersion > WindowsVistaMajorVersion || (WindowsMajorVersion == WindowsVistaMajorVersion && WindowsMinorVersion >= WindowsMinorVersion))
+		if(WindowsMajorVersion > WindowsVistaMajorVersion || (WindowsMajorVersion == WindowsVistaMajorVersion && WindowsMinorVersion >= WindowsVistaMinorVersion))
 		{
 			Result = TSFSource->AdviseSink(IID_ITfInputProcessorProfileActivationSink, static_cast<ITfInputProcessorProfileActivationSink*>(TSFActivationProxy), &(TSFActivationProxy->TSFProfileCookie));
 			if(FAILED(Result))
diff --git a/Engine/Source/Runtime/Core/Public/Windows/WindowsPlatformCodeAnalysis.h b/Engine/Source/Runtime/Core/Public/Windows/WindowsPlatformCodeAnalysis.h
--- a/Engine/Source/Runtime/Core/Public/Windows/WindowsPlatformCodeAnalysis.h
+++ b/Engine/Source/Runtime/Core/Public/Windows/WindowsPlatformCodeAnalysis.h
@@ -4,7 +4,7 @@
 
 
 // Code analysis features
-#if defined( _PREFAST_ )
+#if defined( _PREFAST_ ) || defined( PVS_STUDIO )
 	#define USING_CODE_ANALYSIS 1
 #else
 	#define USING_CODE_ANALYSIS 0
diff --git a/Engine/Source/Runtime/Engine/Private/NiagaraEffectRenderer.cpp b/Engine/Source/Runtime/Engine/Private/NiagaraEffectRenderer.cpp
--- a/Engine/Source/Runtime/Engine/Private/NiagaraEffectRenderer.cpp
+++ b/Engine/Source/Runtime/Engine/Private/NiagaraEffectRenderer.cpp
@@ -154,7 +154,6 @@ void NiagaraEffectRendererSprites::GetDynamicMeshElements(const TArray<const FSc
 				CollectorResources.VertexFactory.SetFeatureLevel(ViewFamily.GetFeatureLevel());
 				CollectorResources.VertexFactory.SetParticleFactoryType(PVFT_Sprite);
 
-				PerViewUniformParameters.MacroUVParameters = FVector4(0.0f, 0.0f, 1.0f, 1.0f);
 				CollectorResources.UniformBuffer = FParticleSpriteUniformBufferRef::CreateUniformBufferImmediate(PerViewUniformParameters, UniformBuffer_SingleFrame);
 
 				CollectorResources.VertexFactory.InitResource();
diff --git a/Engine/Source/Runtime/Engine/Private/UnrealEngine.cpp b/Engine/Source/Runtime/Engine/Private/UnrealEngine.cpp
--- a/Engine/Source/Runtime/Engine/Private/UnrealEngine.cpp
+++ b/Engine/Source/Runtime/Engine/Private/UnrealEngine.cpp
@@ -11584,11 +11584,10 @@ int32 UEngine::RenderStatSoundWaves(UWorld* World, FViewport* Viewport, FCanvas*
 		}
 
 		int32 ActiveInstances = WaveInstances.Num() - FirstActiveIndex;
-		int32 R, G, B;
-		R = G = B = 0;
 		int32 Max = AudioDevice->MaxChannels / 2;
 		float f = FMath::Clamp<float>((float)(ActiveInstances - Max) / (float)Max, 0.f, 1.f);
-		R = FMath::TruncToInt(f * 255);
+		int32 R = FMath::TruncToInt(f * 255);
+
 		if (ActiveInstances > Max)
 		{
 			f = FMath::Clamp<float>((float)(Max - ActiveInstances) / (float)Max, 0.5f, 1.f);
@@ -11597,7 +11596,8 @@ int32 UEngine::RenderStatSoundWaves(UWorld* World, FViewport* Viewport, FCanvas*
 		{
 			f = 1.0f;
 		}
-		G = FMath::TruncToInt(f * 255);
+		int32 G = FMath::TruncToInt(f * 255);
+		int32 B = 0;
 
 		Canvas->DrawShadowedString(X, Y, *FString::Printf(TEXT(" Total: %i"), ActiveInstances), GetSmallFont(), FColor(R, G, B));
 		Y += 12;
diff --git a/Engine/Source/Runtime/Navmesh/Private/DetourTileCache/DetourTileCache.cpp b/Engine/Source/Runtime/Navmesh/Private/DetourTileCache/DetourTileCache.cpp
--- a/Engine/Source/Runtime/Navmesh/Private/DetourTileCache/DetourTileCache.cpp
+++ b/Engine/Source/Runtime/Navmesh/Private/DetourTileCache/DetourTileCache.cpp
@@ -684,7 +684,7 @@ dtStatus dtTileCache::buildNavMeshTile(const dtCompressedTileRef ref, dtNavMesh*
 	
 	bc.lcset = dtAllocTileCacheContourSet(m_talloc);
 	bc.lclusters = dtAllocTileCacheClusterSet(m_talloc);
-	if (!bc.lclusters || !bc.lclusters)
+	if (!bc.lcset || !bc.lclusters)
 		return status;
 	status = dtBuildTileCacheContours(m_talloc, *bc.layer, walkableClimbVx,
 									  m_params.maxSimplificationError, m_params.cs, m_params.ch,
