Fix: UE-12111 [CrashReport] Crash after undoing/deleting landscape

diff --git a/Engine/Source/Editor/LandscapeEditor/Private/LandscapeEdMode.cpp b/Engine/Source/Editor/LandscapeEditor/Private/LandscapeEdMode.cpp
--- a/Engine/Source/Editor/LandscapeEditor/Private/LandscapeEdMode.cpp
+++ b/Engine/Source/Editor/LandscapeEditor/Private/LandscapeEdMode.cpp
@@ -1645,36 +1645,44 @@ int32 FEdModeLandscape::UpdateLandscapeList()
 	}
 
 	int32 CurrentIndex = INDEX_NONE;
-	if (GetWorld())
+	UWorld* World = GetWorld();
+	
+	if (World)
 	{
 		int32 Index = 0;
-		for (auto It = GetLandscapeInfoMap(GetWorld()).Map.CreateIterator(); It; ++It)
+		auto& LandscapeInfoMap = GetLandscapeInfoMap(World);
+
+		for (auto It = LandscapeInfoMap.Map.CreateIterator(); It; ++It)
 		{
 			ULandscapeInfo* LandscapeInfo = It.Value();
 			if (LandscapeInfo && !LandscapeInfo->IsPendingKill())
 			{
-				if (CurrentToolTarget.LandscapeInfo == LandscapeInfo)
+				ALandscapeProxy* LandscapeProxy = LandscapeInfo->GetLandscapeProxy();
+				if (LandscapeProxy)
 				{
-					CurrentIndex = Index;
+					if (CurrentToolTarget.LandscapeInfo == LandscapeInfo)
+					{
+						CurrentIndex = Index;
+
+						// Update GizmoActor Landscape Target (is this necessary?)
+						if (CurrentGizmoActor.IsValid())
+						{
+							CurrentGizmoActor->SetTargetLandscape(LandscapeInfo);
+						}
+					}
 
-					// Update GizmoActor Landscape Target (is this necessary?)
-					if (CurrentGizmoActor.IsValid())
+					int32 MinX, MinY, MaxX, MaxY;
+					int32 Width = 0, Height = 0;
+					if (LandscapeInfo->GetLandscapeExtent(MinX, MinY, MaxX, MaxY))
 					{
-						CurrentGizmoActor->SetTargetLandscape(LandscapeInfo);
+						Width = MaxX - MinX + 1;
+						Height = MaxY - MinY + 1;
 					}
-				}
 
-				int32 MinX, MinY, MaxX, MaxY;
-				int32 Width = 0, Height = 0;
-				if (LandscapeInfo->GetLandscapeExtent(MinX, MinY, MaxX, MaxY))
-				{
-					Width = MaxX - MinX + 1;
-					Height = MaxY - MinY + 1;
+					LandscapeList.Add(FLandscapeListInfo(*LandscapeProxy->GetName(), LandscapeInfo,
+						LandscapeInfo->ComponentSizeQuads, LandscapeInfo->ComponentNumSubsections, Width, Height));
+					Index++;
 				}
-
-				LandscapeList.Add(FLandscapeListInfo(*LandscapeInfo->GetLandscapeProxy()->GetName(), LandscapeInfo,
-					LandscapeInfo->ComponentSizeQuads, LandscapeInfo->ComponentNumSubsections, Width, Height));
-				Index++;
 			}
 		}
 	}
