Fixed landscape tool crash when target landscape sub-level is hidden
UE-12433

diff --git a/Engine/Source/Editor/LandscapeEditor/Private/LandscapeEdMode.cpp b/Engine/Source/Editor/LandscapeEditor/Private/LandscapeEdMode.cpp
--- a/Engine/Source/Editor/LandscapeEditor/Private/LandscapeEdMode.cpp
+++ b/Engine/Source/Editor/LandscapeEditor/Private/LandscapeEdMode.cpp
@@ -529,7 +529,7 @@ void FEdModeLandscape::Tick(FEditorViewportClient* ViewportClient, float DeltaTi
 	if (NewLandscapePreviewMode == ENewLandscapePreviewMode::None)
 	{
 		bool bStaleTargetLandscapeInfo = CurrentToolTarget.LandscapeInfo.IsStale();
-		bool bStaleTargetLandscape = CurrentToolTarget.LandscapeInfo.IsValid() && CurrentToolTarget.LandscapeInfo->LandscapeActor.IsStale();
+		bool bStaleTargetLandscape = CurrentToolTarget.LandscapeInfo.IsValid() && (CurrentToolTarget.LandscapeInfo->GetLandscapeProxy() != nullptr);
 
 		if (bStaleTargetLandscapeInfo || bStaleTargetLandscape)
 		{
@@ -539,24 +539,29 @@ void FEdModeLandscape::Tick(FEditorViewportClient* ViewportClient, float DeltaTi
 		if (CurrentToolTarget.LandscapeInfo.IsValid())
 		{
 			ALandscapeProxy* LandscapeProxy = CurrentToolTarget.LandscapeInfo->GetLandscapeProxy();
+			
 			if (LandscapeProxy == NULL ||
 				LandscapeProxy->GetLandscapeMaterial() != CachedLandscapeMaterial)
 			{
 				UpdateTargetList();
 			}
-		}
-
-		if (CurrentTool)
-		{
-			CurrentTool->Tick(ViewportClient, DeltaTime);
-		}
-		if (CurrentBrush)
-		{
-			CurrentBrush->Tick(ViewportClient, DeltaTime);
-		}
-		if (CurrentBrush != GizmoBrush && CurrentGizmoActor.IsValid() && GizmoBrush && (GLandscapeEditRenderMode & ELandscapeEditRenderMode::Gizmo))
-		{
-			GizmoBrush->Tick(ViewportClient, DeltaTime);
+			else
+			{
+				if (CurrentTool)
+				{
+					CurrentTool->Tick(ViewportClient, DeltaTime);
+				}
+			
+				if (CurrentBrush)
+				{
+					CurrentBrush->Tick(ViewportClient, DeltaTime);
+				}
+			
+				if (CurrentBrush != GizmoBrush && CurrentGizmoActor.IsValid() && GizmoBrush && (GLandscapeEditRenderMode & ELandscapeEditRenderMode::Gizmo))
+				{
+					GizmoBrush->Tick(ViewportClient, DeltaTime);
+				}
+			}
 		}
 	}
 }
