Preventing creation of landscape actor in a hidden level, otherwise it's crashing

diff --git a/Engine/Source/Editor/LandscapeEditor/Private/LandscapeEditorDetailCustomization_NewLandscape.cpp b/Engine/Source/Editor/LandscapeEditor/Private/LandscapeEditorDetailCustomization_NewLandscape.cpp
--- a/Engine/Source/Editor/LandscapeEditor/Private/LandscapeEditorDetailCustomization_NewLandscape.cpp
+++ b/Engine/Source/Editor/LandscapeEditor/Private/LandscapeEditorDetailCustomization_NewLandscape.cpp
@@ -759,7 +759,9 @@ void FLandscapeEditorDetailCustomization_NewLandscape::OnNewLandscapeModeChanged
 FReply FLandscapeEditorDetailCustomization_NewLandscape::OnCreateButtonClicked()
 {
 	FEdModeLandscape* LandscapeEdMode = GetEditorMode();
-	if (LandscapeEdMode != NULL)
+	if (LandscapeEdMode != nullptr && 
+		LandscapeEdMode->GetWorld() != nullptr && 
+		LandscapeEdMode->GetWorld()->GetCurrentLevel()->bIsVisible)
 	{
 		FScopedTransaction Transaction(LOCTEXT("Undo", "Creating New Landscape"));
 
