Fixed: Unable to sculpt/paint on mirrored landscape UE-13273

diff --git a/Engine/Source/Editor/LandscapeEditor/Private/LandscapeEdModeBrushes.cpp b/Engine/Source/Editor/LandscapeEditor/Private/LandscapeEdModeBrushes.cpp
--- a/Engine/Source/Editor/LandscapeEditor/Private/LandscapeEdModeBrushes.cpp
+++ b/Engine/Source/Editor/LandscapeEditor/Private/LandscapeEdModeBrushes.cpp
@@ -102,7 +102,7 @@ public:
 		ULandscapeInfo* LandscapeInfo = EdMode->CurrentToolTarget.LandscapeInfo.Get();
 		ALandscapeProxy* Proxy = LandscapeInfo->GetLandscapeProxy();
 
-		const float ScaleXY = LandscapeInfo->DrawScale.X;
+		const float ScaleXY = FMath::Abs(LandscapeInfo->DrawScale.X);
 		const float TotalRadius = EdMode->UISettings->BrushRadius / ScaleXY;
 		const float Radius = (1.0f - EdMode->UISettings->BrushFalloff) * TotalRadius;
 		const float Falloff = EdMode->UISettings->BrushFalloff * TotalRadius;
@@ -208,7 +208,7 @@ public:
 	virtual FLandscapeBrushData ApplyBrush(const TArray<FLandscapeToolMousePosition>& InMousePositions) override
 	{
 		ULandscapeInfo* LandscapeInfo = EdMode->CurrentToolTarget.LandscapeInfo.Get();
-		const float ScaleXY = LandscapeInfo->DrawScale.X;
+		const float ScaleXY = FMath::Abs(LandscapeInfo->DrawScale.X);
 		const float TotalRadius = EdMode->UISettings->BrushRadius / ScaleXY;
 		const float Radius = (1.0f - EdMode->UISettings->BrushFalloff) * TotalRadius;
 		const float Falloff = EdMode->UISettings->BrushFalloff * TotalRadius;
@@ -569,7 +569,7 @@ public:
 				ULandscapeInfo* LandscapeInfo = Gizmo->TargetLandscapeInfo;
 				if (LandscapeInfo && LandscapeInfo->GetLandscapeProxy())
 				{
-					float ScaleXY = LandscapeInfo->DrawScale.X;
+					float ScaleXY = FMath::Abs(LandscapeInfo->DrawScale.X);
 					FMatrix LToW = LandscapeInfo->GetLandscapeProxy()->LandscapeActorToWorld().ToMatrixWithScale();
 					FMatrix WToL = LToW.InverseFast();
 
@@ -672,7 +672,7 @@ public:
 		}
 
 		Gizmo->TargetLandscapeInfo = LandscapeInfo;
-		float ScaleXY = LandscapeInfo->DrawScale.X;
+		float ScaleXY = FMath::Abs(LandscapeInfo->DrawScale.X);
 
 		// Get extent for all components
 		FIntRect Bounds;
@@ -996,7 +996,7 @@ public:
 	virtual FLandscapeBrushData ApplyBrush(const TArray<FLandscapeToolMousePosition>& MousePositions) override
 	{
 		ULandscapeInfo* LandscapeInfo = EdMode->CurrentToolTarget.LandscapeInfo.Get();
-		const float ScaleXY = LandscapeInfo->DrawScale.X;
+		const float ScaleXY = FMath::Abs(LandscapeInfo->DrawScale.X);
 		const float TotalRadius = EdMode->UISettings->BrushRadius / ScaleXY;
 		const float Radius = (1.0f - EdMode->UISettings->BrushFalloff) * TotalRadius;
 		const float Falloff = EdMode->UISettings->BrushFalloff * TotalRadius;
@@ -1106,7 +1106,7 @@ public:
 		ALandscapeProxy* Proxy = EdMode->CurrentToolTarget.LandscapeInfo.IsValid() ? EdMode->CurrentToolTarget.LandscapeInfo->GetLandscapeProxy() : nullptr;
 		if (Proxy)
 		{
-			const float ScaleXY = EdMode->CurrentToolTarget.LandscapeInfo->DrawScale.X;
+			const float ScaleXY = FMath::Abs(EdMode->CurrentToolTarget.LandscapeInfo->DrawScale.X);
 			int32 SizeX = EdMode->UISettings->AlphaTextureSizeX;
 			int32 SizeY = EdMode->UISettings->AlphaTextureSizeY;
 
@@ -1200,7 +1200,7 @@ public:
 		}
 		else
 		{
-			float ScaleXY = LandscapeInfo->DrawScale.X;
+			float ScaleXY = FMath::Abs(LandscapeInfo->DrawScale.X);
 			float Radius = EdMode->UISettings->BrushRadius / ScaleXY;
 			int32 SizeX = EdMode->UISettings->AlphaTextureSizeX;
 			int32 SizeY = EdMode->UISettings->AlphaTextureSizeY;
@@ -1304,7 +1304,7 @@ public:
 		ULandscapeInfo* LandscapeInfo = EdMode->CurrentToolTarget.LandscapeInfo.Get();
 		if (LandscapeInfo)
 		{
-			float ScaleXY = LandscapeInfo->DrawScale.X;
+			float ScaleXY = FMath::Abs(LandscapeInfo->DrawScale.X);
 			int32 SizeX = EdMode->UISettings->AlphaTextureSizeX;
 			int32 SizeY = EdMode->UISettings->AlphaTextureSizeY;
 			float Radius = EdMode->UISettings->BrushRadius / ScaleXY;
