Fix crash with landscape spline editor

diff --git a/Engine/Source/Editor/LandscapeEditor/Private/LandscapeEdModeSplineTools.cpp b/Engine/Source/Editor/LandscapeEditor/Private/LandscapeEdModeSplineTools.cpp
--- a/Engine/Source/Editor/LandscapeEditor/Private/LandscapeEdModeSplineTools.cpp
+++ b/Engine/Source/Editor/LandscapeEditor/Private/LandscapeEdModeSplineTools.cpp
@@ -305,27 +305,36 @@ public:
 		Start->ConnectedSegments.Add(FLandscapeSplineConnection(NewSegment, 0));
 		End->ConnectedSegments.Add(FLandscapeSplineConnection(NewSegment, 1));
 
+		bool bUpdatedStart = false;
+		bool bUpdatedEnd = false;
 		if (bAutoRotateStart)
 		{
 			Start->AutoCalcRotation();
 			Start->UpdateSplinePoints();
+			bUpdatedStart = true;
 		}
 		if (bAutoRotateEnd)
 		{
 			End->AutoCalcRotation();
 			End->UpdateSplinePoints();
+			bUpdatedEnd = true;
 		}
 
 		// Control points' points are currently based on connected segments, so need to be updated.
-		if (Start->Mesh != NULL)
+		if (!bUpdatedStart && Start->Mesh)
 		{
 			Start->UpdateSplinePoints();
 		}
-		if (End->Mesh != NULL)
+		if (!bUpdatedEnd && End->Mesh)
 		{
-			Start->UpdateSplinePoints();
+			End->UpdateSplinePoints();
+		}
+
+		// If we've called UpdateSplinePoints on either control point it will already have called UpdateSplinePoints on the new segment
+		if (!(bUpdatedStart || bUpdatedEnd))
+		{
+			NewSegment->UpdateSplinePoints();
 		}
-		NewSegment->UpdateSplinePoints();
 	}
 
 	void AddControlPoint(ALandscapeProxy* Landscape, const FVector& LocalLocation)
diff --git a/Engine/Source/Runtime/Landscape/Private/LandscapeSplines.cpp b/Engine/Source/Runtime/Landscape/Private/LandscapeSplines.cpp
--- a/Engine/Source/Runtime/Landscape/Private/LandscapeSplines.cpp
+++ b/Engine/Source/Runtime/Landscape/Private/LandscapeSplines.cpp
@@ -1943,12 +1943,17 @@ void ULandscapeSplineSegment::UpdateSplinePoints(bool bUpdateCollision)
 
 	// Editor mesh
 	bool bUsingEditorMesh = false;
+	FLandscapeSplineMeshEntry SplineEditorMeshEntry;
 	if (UsableMeshes.Num() == 0 && OuterSplines->SplineEditorMesh != nullptr)
 	{
-		FLandscapeSplineMeshEntry SplineEditorMeshEntry;
 		SplineEditorMeshEntry.Mesh = OuterSplines->SplineEditorMesh;
-		SplineEditorMeshEntry.Scale.X = 3;
-		SplineEditorMeshEntry.Offset.Y = 0.5f;
+		SplineEditorMeshEntry.MaterialOverrides = {};
+		SplineEditorMeshEntry.bCenterH = true;
+		SplineEditorMeshEntry.Offset = {0.0f, 0.5f};
+		SplineEditorMeshEntry.bScaleToWidth = true;
+		SplineEditorMeshEntry.Scale = {3, 1, 1};
+		SplineEditorMeshEntry.ForwardAxis = ESplineMeshAxis::X;
+		SplineEditorMeshEntry.UpAxis = ESplineMeshAxis::Z;
 		UsableMeshes.Add(&SplineEditorMeshEntry);
 		bUsingEditorMesh = true;
 	}
