Fix the strength of the landscape weightmap painting tools.

diff --git a/Engine/Source/Editor/LandscapeEditor/Classes/LandscapeEditorObject.h b/Engine/Source/Editor/LandscapeEditor/Classes/LandscapeEditorObject.h
--- a/Engine/Source/Editor/LandscapeEditor/Classes/LandscapeEditorObject.h
+++ b/Engine/Source/Editor/LandscapeEditor/Classes/LandscapeEditorObject.h
@@ -219,7 +219,7 @@ class ULandscapeEditorObject : public UObject
 
 	// Common Tool Settings:
 
-	// Strength of the tool
+	// Strength of the tool. If you're using a pen/tablet with pressure-sensing, the pressure used affects the strength of the tool.
 	UPROPERTY(Category="Tool Settings", EditAnywhere, NonTransactional, meta=(ShowForTools="Paint,Sculpt,Smooth,Flatten,Erosion,HydraErosion,Noise,Mask,CopyPaste", ClampMin="0", ClampMax="10", UIMin="0", UIMax="1"))
 	float ToolStrength;
 
diff --git a/Engine/Source/Editor/LandscapeEditor/Private/LandscapeEdModePaintTools.cpp b/Engine/Source/Editor/LandscapeEditor/Private/LandscapeEdModePaintTools.cpp
--- a/Engine/Source/Editor/LandscapeEditor/Private/LandscapeEdModePaintTools.cpp
+++ b/Engine/Source/Editor/LandscapeEditor/Private/LandscapeEdModePaintTools.cpp
@@ -73,12 +73,14 @@ protected:
 
 // 
 // FLandscapeToolPaint
+// TODO: Split heightmap sculpting into its own stroke type
 //
 template<class ToolTarget>
 class FLandscapeToolStrokePaint : public FLandscapeToolStrokePaintBase<ToolTarget>
 {
 	TMap<FIntPoint, float> TotalInfluenceMap;	// amount of time and weight the brush has spent on each vertex.
 public:
+	// Heightmap sculpt tool will continuously sculpt in the same location, weightmap paint tool doesn't
 	enum { UseContinuousApply = (ToolTarget::TargetType == ELandscapeToolTargetType::Heightmap) };
 
 	FLandscapeToolStrokePaint(FEdModeLandscape* InEdMode, const FLandscapeToolTarget& InTarget)
@@ -130,6 +132,7 @@ public:
 		{
 			// When painting weights (and not using target value mode), we use a source value that tends more
 			// to the current value as we paint over the same region multiple times.
+			// TODO: Make this frame-rate independent
 			this->Cache.GetOriginalData(X1, Y1, X2, Y2, OriginalData);
 			SourceDataArrayPtr = &OriginalData;
 
@@ -156,9 +159,12 @@ public:
 		float AdjustedStrength = ToolTarget::StrengthMultiplier(this->LandscapeInfo, UISettings->BrushRadius);
 		typename ToolTarget::CacheClass::DataType DestValue = ToolTarget::CacheClass::ClampValue(255.0f * UISettings->WeightTargetValue);
 
-		float DeltaTime = FMath::Min<float>(FApp::GetDeltaTime(), 0.1f); // Under 10 fps slow down paint speed
-		// * 3.0f to partially compensate for impact of DeltaTime on slowing the tools down compared to the old framerate-dependent version
-		float PaintStrength = UISettings->ToolStrength * Pressure * AdjustedStrength * DeltaTime * 3.0f;
+		float PaintStrength = UISettings->ToolStrength * Pressure * AdjustedStrength;
+		if (ToolTarget::TargetType != ELandscapeToolTargetType::Heightmap)
+		{
+			float DeltaTime = FMath::Min<float>(FApp::GetDeltaTime(), 0.1f); // Under 10 fps slow down paint speed
+			PaintStrength *= DeltaTime * 3.0f; // * 3.0f to partially compensate for impact of DeltaTime on slowing the tools down compared to the old framerate-dependent version
+		}
 
 		FPlane BrushPlane;
 		TArray<FVector> Normals;
