Fix the fix fixing the strength of the landscape weightmap painting tools.

diff --git a/Engine/Source/Editor/LandscapeEditor/Private/LandscapeEdModePaintTools.cpp b/Engine/Source/Editor/LandscapeEditor/Private/LandscapeEdModePaintTools.cpp
--- a/Engine/Source/Editor/LandscapeEditor/Private/LandscapeEdModePaintTools.cpp
+++ b/Engine/Source/Editor/LandscapeEditor/Private/LandscapeEdModePaintTools.cpp
@@ -160,7 +160,7 @@ public:
 		typename ToolTarget::CacheClass::DataType DestValue = ToolTarget::CacheClass::ClampValue(255.0f * UISettings->WeightTargetValue);
 
 		float PaintStrength = UISettings->ToolStrength * Pressure * AdjustedStrength;
-		if (ToolTarget::TargetType != ELandscapeToolTargetType::Heightmap)
+		if (ToolTarget::TargetType == ELandscapeToolTargetType::Heightmap)
 		{
 			float DeltaTime = FMath::Min<float>(FApp::GetDeltaTime(), 0.1f); // Under 10 fps slow down paint speed
 			PaintStrength *= DeltaTime * 3.0f; // * 3.0f to partially compensate for impact of DeltaTime on slowing the tools down compared to the old framerate-dependent version
