Added config migration path for newer versions of the engine.

diff --git a/Engine/Source/Editor/FoliageEdit/Private/FoliageEdMode.cpp b/Engine/Source/Editor/FoliageEdit/Private/FoliageEdMode.cpp
--- a/Engine/Source/Editor/FoliageEdit/Private/FoliageEdMode.cpp
+++ b/Engine/Source/Editor/FoliageEdit/Private/FoliageEdMode.cpp
@@ -2781,7 +2781,7 @@ EAxisList::Type FEdModeFoliage::GetWidgetAxisToDraw(FWidget::EWidgetMode InWidge
 void FFoliageUISettings::Load()
 {
 	FString WindowPositionString;
-	if (GConfig->GetString(TEXT("FoliageEdit"), TEXT("WindowPosition"), WindowPositionString, GEditorUserSettingsIni))
+	if (GConfig->GetString(TEXT("FoliageEdit"), TEXT("WindowPosition"), WindowPositionString, GEditorPerProjectIni))
 	{
 		TArray<FString> PositionValues;
 		if (WindowPositionString.ParseIntoArray(PositionValues, TEXT(","), true) == 4)
@@ -2793,28 +2793,28 @@ void FFoliageUISettings::Load()
 		}
 	}
 
-	GConfig->GetFloat(TEXT("FoliageEdit"), TEXT("Radius"), Radius, GEditorUserSettingsIni);
-	GConfig->GetFloat(TEXT("FoliageEdit"), TEXT("PaintDensity"), PaintDensity, GEditorUserSettingsIni);
-	GConfig->GetFloat(TEXT("FoliageEdit"), TEXT("UnpaintDensity"), UnpaintDensity, GEditorUserSettingsIni);
-	GConfig->GetBool(TEXT("FoliageEdit"), TEXT("bFilterLandscape"), bFilterLandscape, GEditorUserSettingsIni);
-	GConfig->GetBool(TEXT("FoliageEdit"), TEXT("bFilterStaticMesh"), bFilterStaticMesh, GEditorUserSettingsIni);
-	GConfig->GetBool(TEXT("FoliageEdit"), TEXT("bFilterBSP"), bFilterBSP, GEditorUserSettingsIni);
-	GConfig->GetBool(TEXT("FoliageEdit"), TEXT("bFilterTranslucent"), bFilterTranslucent, GEditorUserSettingsIni);
+	GConfig->GetFloat(TEXT("FoliageEdit"), TEXT("Radius"), Radius, GEditorPerProjectIni);
+	GConfig->GetFloat(TEXT("FoliageEdit"), TEXT("PaintDensity"), PaintDensity, GEditorPerProjectIni);
+	GConfig->GetFloat(TEXT("FoliageEdit"), TEXT("UnpaintDensity"), UnpaintDensity, GEditorPerProjectIni);
+	GConfig->GetBool(TEXT("FoliageEdit"), TEXT("bFilterLandscape"), bFilterLandscape, GEditorPerProjectIni);
+	GConfig->GetBool(TEXT("FoliageEdit"), TEXT("bFilterStaticMesh"), bFilterStaticMesh, GEditorPerProjectIni);
+	GConfig->GetBool(TEXT("FoliageEdit"), TEXT("bFilterBSP"), bFilterBSP, GEditorPerProjectIni);
+	GConfig->GetBool(TEXT("FoliageEdit"), TEXT("bFilterTranslucent"), bFilterTranslucent, GEditorPerProjectIni);
 }
 
 /** Save UI settings to ini file */
 void FFoliageUISettings::Save()
 {
 	FString WindowPositionString = FString::Printf(TEXT("%d,%d,%d,%d"), WindowX, WindowY, WindowWidth, WindowHeight);
-	GConfig->SetString(TEXT("FoliageEdit"), TEXT("WindowPosition"), *WindowPositionString, GEditorUserSettingsIni);
-
-	GConfig->SetFloat(TEXT("FoliageEdit"), TEXT("Radius"), Radius, GEditorUserSettingsIni);
-	GConfig->SetFloat(TEXT("FoliageEdit"), TEXT("PaintDensity"), PaintDensity, GEditorUserSettingsIni);
-	GConfig->SetFloat(TEXT("FoliageEdit"), TEXT("UnpaintDensity"), UnpaintDensity, GEditorUserSettingsIni);
-	GConfig->SetBool(TEXT("FoliageEdit"), TEXT("bFilterLandscape"), bFilterLandscape, GEditorUserSettingsIni);
-	GConfig->SetBool(TEXT("FoliageEdit"), TEXT("bFilterStaticMesh"), bFilterStaticMesh, GEditorUserSettingsIni);
-	GConfig->SetBool(TEXT("FoliageEdit"), TEXT("bFilterBSP"), bFilterBSP, GEditorUserSettingsIni);
-	GConfig->SetBool(TEXT("FoliageEdit"), TEXT("bFilterTranslucent"), bFilterTranslucent, GEditorUserSettingsIni);
+	GConfig->SetString(TEXT("FoliageEdit"), TEXT("WindowPosition"), *WindowPositionString, GEditorPerProjectIni);
+
+	GConfig->SetFloat(TEXT("FoliageEdit"), TEXT("Radius"), Radius, GEditorPerProjectIni);
+	GConfig->SetFloat(TEXT("FoliageEdit"), TEXT("PaintDensity"), PaintDensity, GEditorPerProjectIni);
+	GConfig->SetFloat(TEXT("FoliageEdit"), TEXT("UnpaintDensity"), UnpaintDensity, GEditorPerProjectIni);
+	GConfig->SetBool(TEXT("FoliageEdit"), TEXT("bFilterLandscape"), bFilterLandscape, GEditorPerProjectIni);
+	GConfig->SetBool(TEXT("FoliageEdit"), TEXT("bFilterStaticMesh"), bFilterStaticMesh, GEditorPerProjectIni);
+	GConfig->SetBool(TEXT("FoliageEdit"), TEXT("bFilterBSP"), bFilterBSP, GEditorPerProjectIni);
+	GConfig->SetBool(TEXT("FoliageEdit"), TEXT("bFilterTranslucent"), bFilterTranslucent, GEditorPerProjectIni);
 }
 
 #undef LOCTEXT_NAMESPACE
\ No newline at end of file
diff --git a/Engine/Source/Editor/GameProjectGeneration/Private/GameProjectUtils.cpp b/Engine/Source/Editor/GameProjectGeneration/Private/GameProjectUtils.cpp
--- a/Engine/Source/Editor/GameProjectGeneration/Private/GameProjectUtils.cpp
+++ b/Engine/Source/Editor/GameProjectGeneration/Private/GameProjectUtils.cpp
@@ -3137,7 +3137,7 @@ GameProjectUtils::EAddCodeToProjectResult GameProjectUtils::AddCodeToProject_Int
 		// Notify that we've created a brand new module
 		FSourceCodeNavigation::AccessOnNewModuleAdded().Broadcast(*GameModuleName);
 	}
-	else if (GEditor->AccessEditorUserSettings().bAutomaticallyHotReloadNewClasses)
+	else if (GetDefault<UEditorPerProjectUserSettings>()->bAutomaticallyHotReloadNewClasses)
 	{
 		FModuleStatus ModuleStatus;
 		const FName ModuleFName = *ModuleInfo.ModuleName;
diff --git a/Engine/Source/Editor/GameProjectGeneration/Private/SNewClassDialog.cpp b/Engine/Source/Editor/GameProjectGeneration/Private/SNewClassDialog.cpp
--- a/Engine/Source/Editor/GameProjectGeneration/Private/SNewClassDialog.cpp
+++ b/Engine/Source/Editor/GameProjectGeneration/Private/SNewClassDialog.cpp
@@ -1172,7 +1172,7 @@ void SNewClassDialog::FinishClicked()
 			bPreventPeriodicValidityChecksUntilNextChange = true;
 
 			// Display a nag if we didn't automatically hot-reload for the newly added class
-			const bool bWasHotReloaded = GEditor->AccessEditorUserSettings().bAutomaticallyHotReloadNewClasses;
+			const bool bWasHotReloaded = GetDefault<UEditorPerProjectUserSettings>()->bAutomaticallyHotReloadNewClasses;
 			if( bWasHotReloaded )
 			{
 				FNotificationInfo Notification( FText::Format( LOCTEXT("AddedClassSuccessNotification", "Added new class {0}"), FText::FromString(NewClassName) ) );
diff --git a/Engine/Source/Editor/GameProjectGeneration/Private/SNewProjectWizard.cpp b/Engine/Source/Editor/GameProjectGeneration/Private/SNewProjectWizard.cpp
--- a/Engine/Source/Editor/GameProjectGeneration/Private/SNewProjectWizard.cpp
+++ b/Engine/Source/Editor/GameProjectGeneration/Private/SNewProjectWizard.cpp
@@ -16,6 +16,8 @@
 #include "SHyperlink.h"
 #include "SOutputLogDialog.h"
 
+#include "Settings/EditorSettings.h"
+
 #define LOCTEXT_NAMESPACE "NewProjectWizard"
 
 FName SNewProjectWizard::TemplatePageName = TEXT("Template");
@@ -300,7 +302,7 @@ void SNewProjectWizard::Construct( const FArguments& InArgs )
 	bLastGlobalValidityCheckSuccessful = true;
 	bLastNameAndLocationValidityCheckSuccessful = true;
 	bPreventPeriodicValidityChecksUntilNextChange = false;
-	bCopyStarterContent = GEditor ? GEditor->AccessGameAgnosticSettings().bCopyStarterContentPreference : true;
+	bCopyStarterContent = GEditor ? GetDefault<UEditorSettings>()->bCopyStarterContentPreference : true;
 
 	IHardwareTargetingModule& HardwareTargeting = IHardwareTargetingModule::Get();
 
@@ -1288,7 +1290,7 @@ void SNewProjectWizard::SetDefaultProjectLocation( )
 	FString DefaultProjectFilePath;
 	
 	// First, try and use the first previously used path that still exists
-	for ( const FString& CreatedProjectPath : GEditor->GetGameAgnosticSettings().CreatedProjectPaths )
+	for ( const FString& CreatedProjectPath : GetDefault<UEditorSettings>()->CreatedProjectPaths )
 	{
 		if ( IFileManager::Get().DirectoryExists(*CreatedProjectPath) )
 		{
@@ -1458,10 +1460,11 @@ bool SNewProjectWizard::CreateProject( const FString& ProjectFile )
 		CreatedProjectPath.AppendChar('/');
 	}
 
-	GEditor->AccessGameAgnosticSettings().CreatedProjectPaths.Remove(CreatedProjectPath);
-	GEditor->AccessGameAgnosticSettings().CreatedProjectPaths.Insert(CreatedProjectPath, 0);
-	GEditor->AccessGameAgnosticSettings().bCopyStarterContentPreference = bCopyStarterContent;
-	GEditor->AccessGameAgnosticSettings().PostEditChange();
+	auto* Settings = GetMutableDefault<UEditorSettings>();
+	Settings->CreatedProjectPaths.Remove(CreatedProjectPath);
+	Settings->CreatedProjectPaths.Insert(CreatedProjectPath, 0);
+	Settings->bCopyStarterContentPreference = bCopyStarterContent;
+	Settings->PostEditChange();
 
 	return true;
 }
diff --git a/Engine/Source/Editor/GameProjectGeneration/Private/SProjectBrowser.cpp b/Engine/Source/Editor/GameProjectGeneration/Private/SProjectBrowser.cpp
--- a/Engine/Source/Editor/GameProjectGeneration/Private/SProjectBrowser.cpp
+++ b/Engine/Source/Editor/GameProjectGeneration/Private/SProjectBrowser.cpp
@@ -10,6 +10,7 @@
 #include "TargetPlatform.h"
 #include "PlatformInfo.h"
 #include "SSearchBox.h"
+#include "Settings/EditorSettings.h"
 
 #define LOCTEXT_NAMESPACE "ProjectBrowser"
 
@@ -262,7 +263,7 @@ void SProjectBrowser::Construct( const FArguments& InArgs )
 				.VAlign(VAlign_Center)
 				[
 					SNew(SCheckBox)			
-					.IsChecked(GEditor->GetGameAgnosticSettings().bLoadTheMostRecentlyLoadedProjectAtStartup ? ECheckBoxState::Checked : ECheckBoxState::Unchecked)
+					.IsChecked(GetDefault<UEditorSettings>()->bLoadTheMostRecentlyLoadedProjectAtStartup ? ECheckBoxState::Checked : ECheckBoxState::Unchecked)
 					.OnCheckStateChanged(this, &SProjectBrowser::OnAutoloadLastProjectChanged)
 					.Content()
 					[
@@ -1142,7 +1143,7 @@ FReply SProjectBrowser::OnBrowseToProjectClicked()
 
 	// Find the first valid project file to select by default
 	FString DefaultFolder = FEditorDirectories::Get().GetLastDirectory(ELastDirectory::PROJECT);
-	for ( auto ProjectIt = GEditor->GetGameAgnosticSettings().RecentlyOpenedProjectFiles.CreateConstIterator(); ProjectIt; ++ProjectIt )
+	for ( auto ProjectIt = GetDefault<UEditorSettings>()->RecentlyOpenedProjectFiles.CreateConstIterator(); ProjectIt; ++ProjectIt )
 	{
 		if ( IFileManager::Get().FileSize(**ProjectIt) > 0 )
 		{
@@ -1269,14 +1270,14 @@ void SProjectBrowser::OnFilterTextChanged(const FText& InText)
 
 void SProjectBrowser::OnAutoloadLastProjectChanged(ECheckBoxState NewState)
 {
-	UEditorGameAgnosticSettings &Settings = GEditor->AccessGameAgnosticSettings();
-	Settings.bLoadTheMostRecentlyLoadedProjectAtStartup = (NewState == ECheckBoxState::Checked);
+	UEditorSettings *Settings = GetMutableDefault<UEditorSettings>();
+	Settings->bLoadTheMostRecentlyLoadedProjectAtStartup = (NewState == ECheckBoxState::Checked);
 
-	UProperty* AutoloadProjectProperty = FindField<UProperty>(Settings.GetClass(), "bLoadTheMostRecentlyLoadedProjectAtStartup");
+	UProperty* AutoloadProjectProperty = FindField<UProperty>(Settings->GetClass(), "bLoadTheMostRecentlyLoadedProjectAtStartup");
 	if (AutoloadProjectProperty != NULL)
 	{
 		FPropertyChangedEvent PropertyUpdateStruct(AutoloadProjectProperty);
-		Settings.PostEditChangeProperty(PropertyUpdateStruct);
+		Settings->PostEditChangeProperty(PropertyUpdateStruct);
 	}
 }
 
diff --git a/Engine/Source/Editor/GameplayTagsEditor/Private/SGameplayTagWidget.cpp b/Engine/Source/Editor/GameplayTagsEditor/Private/SGameplayTagWidget.cpp
--- a/Engine/Source/Editor/GameplayTagsEditor/Private/SGameplayTagWidget.cpp
+++ b/Engine/Source/Editor/GameplayTagsEditor/Private/SGameplayTagWidget.cpp
@@ -519,7 +519,7 @@ void SGameplayTagWidget::LoadTagNodeItemExpansion( TSharedPtr<FGameplayTagNode>
 	{
 		bool bExpanded = false;
 
-		if( GConfig->GetBool(*SettingsIniSection, *(TagContainerName + Node->GetCompleteTag().ToString() + TEXT(".Expanded")), bExpanded, GEditorUserSettingsIni) )
+		if( GConfig->GetBool(*SettingsIniSection, *(TagContainerName + Node->GetCompleteTag().ToString() + TEXT(".Expanded")), bExpanded, GEditorPerProjectIni) )
 		{
 			TagTreeWidget->SetItemExpansion( Node, bExpanded );
 		}
@@ -539,7 +539,7 @@ void SGameplayTagWidget::LoadTagNodeItemExpansion( TSharedPtr<FGameplayTagNode>
 void SGameplayTagWidget::OnExpansionChanged( TSharedPtr<FGameplayTagNode> InItem, bool bIsExpanded )
 {
 	// Save the new expansion setting to ini file
-	GConfig->SetBool(*SettingsIniSection, *(TagContainerName + InItem->GetCompleteTag().ToString() + TEXT(".Expanded")), bIsExpanded, GEditorUserSettingsIni);
+	GConfig->SetBool(*SettingsIniSection, *(TagContainerName + InItem->GetCompleteTag().ToString() + TEXT(".Expanded")), bIsExpanded, GEditorPerProjectIni);
 }
 
 void SGameplayTagWidget::SetContainer(FGameplayTagContainer* OriginalContainer, FGameplayTagContainer* EditedContainer, UObject* OwnerObj)
diff --git a/Engine/Source/Editor/IntroTutorials/Private/EditorTutorialSettings.h b/Engine/Source/Editor/IntroTutorials/Private/EditorTutorialSettings.h
--- a/Engine/Source/Editor/IntroTutorials/Private/EditorTutorialSettings.h
+++ b/Engine/Source/Editor/IntroTutorials/Private/EditorTutorialSettings.h
@@ -29,7 +29,7 @@ struct FTutorialContext
 };
 
 /** Editor-wide tutorial settings */
-UCLASS(config=EditorGameAgnostic)
+UCLASS(config=EditorSettings)
 class UEditorTutorialSettings : public UObject
 {
 	GENERATED_UCLASS_BODY()
diff --git a/Engine/Source/Editor/Kismet/Private/BlueprintActionMenuUtils.cpp b/Engine/Source/Editor/Kismet/Private/BlueprintActionMenuUtils.cpp
--- a/Engine/Source/Editor/Kismet/Private/BlueprintActionMenuUtils.cpp
+++ b/Engine/Source/Editor/Kismet/Private/BlueprintActionMenuUtils.cpp
@@ -139,9 +139,9 @@ static bool BlueprintActionMenuUtilsImpl::IsUnBoundSpawner(FBlueprintActionFilte
 //------------------------------------------------------------------------------
 static bool BlueprintActionMenuUtilsImpl::IsNonFavoritedAction(FBlueprintActionFilter const& Filter, FBlueprintActionInfo& BlueprintAction)
 {
-	UEditorUserSettings& EditorUserSettings = GEditor->AccessEditorUserSettings();
+	const UEditorPerProjectUserSettings* EditorPerProjectUserSettings = GetDefault<UEditorPerProjectUserSettings>();
 	// grab the user's favorites
-	UBlueprintPaletteFavorites const* BlueprintFavorites = EditorUserSettings.BlueprintFavorites;
+	UBlueprintPaletteFavorites const* BlueprintFavorites = EditorPerProjectUserSettings->BlueprintFavorites;
 	checkSlow(BlueprintFavorites != nullptr);
 
 	return !BlueprintFavorites->IsFavorited(BlueprintAction);
diff --git a/Engine/Source/Editor/Kismet/Private/BlueprintEditorCommands.cpp b/Engine/Source/Editor/Kismet/Private/BlueprintEditorCommands.cpp
--- a/Engine/Source/Editor/Kismet/Private/BlueprintEditorCommands.cpp
+++ b/Engine/Source/Editor/Kismet/Private/BlueprintEditorCommands.cpp
@@ -290,7 +290,7 @@ void FBlueprintSpawnNodeCommands::RegisterCommands()
 	const FString ConfigSection = TEXT("BlueprintSpawnNodes");
 	const FString SettingName = TEXT("Node");
 	TArray< FString > NodeSpawns;
-	GConfig->GetArray(*ConfigSection, *SettingName, NodeSpawns, GEditorUserSettingsIni);
+	GConfig->GetArray(*ConfigSection, *SettingName, NodeSpawns, GEditorPerProjectIni);
 
 	for(int32 x = 0; x < NodeSpawns.Num(); ++x)
 	{
diff --git a/Engine/Source/Editor/Kismet/Private/SBlueprintActionMenu.cpp b/Engine/Source/Editor/Kismet/Private/SBlueprintActionMenu.cpp
--- a/Engine/Source/Editor/Kismet/Private/SBlueprintActionMenu.cpp
+++ b/Engine/Source/Editor/Kismet/Private/SBlueprintActionMenu.cpp
@@ -116,7 +116,7 @@ private:
 		bool bNoFavorites = false;
 		GConfig->GetBool(TEXT("BlueprintEditor.Palette"), TEXT("bUseLegacyLayout"), bNoFavorites, GEditorIni);
 
-		UBlueprintPaletteFavorites const* const BlueprintFavorites = GEditor->EditorUserSettings->BlueprintFavorites;
+		UBlueprintPaletteFavorites const* const BlueprintFavorites = GetDefault<UEditorPerProjectUserSettings>()->BlueprintFavorites;
 
 		EVisibility CurrentVisibility = EVisibility::Hidden;
 		if (!bNoFavorites && BlueprintFavorites && BlueprintFavorites->CanBeFavorited(ActionPtr.Pin()))
@@ -154,9 +154,9 @@ private:
 	ECheckBoxState GetFavoritedState() const
 	{
 		ECheckBoxState FavoriteState = ECheckBoxState::Unchecked;
-		if (ActionPtr.IsValid() && (GEditor != nullptr))
+		if (ActionPtr.IsValid())
 		{
-			const UEditorUserSettings& EditorSettings = GEditor->GetEditorUserSettings();
+			const UEditorPerProjectUserSettings& EditorSettings = *GetDefault<UEditorPerProjectUserSettings>();
 			if (UBlueprintPaletteFavorites* BlueprintFavorites = EditorSettings.BlueprintFavorites)
 			{
 				FavoriteState = BlueprintFavorites->IsFavorited(ActionPtr.Pin()) ? ECheckBoxState::Checked : ECheckBoxState::Unchecked;
@@ -175,11 +175,11 @@ private:
 	{
 		if (InNewState == ECheckBoxState::Checked)
 		{
-			GEditor->EditorUserSettings->BlueprintFavorites->AddFavorite(ActionPtr.Pin());
+			GetMutableDefault<UEditorPerProjectUserSettings>()->BlueprintFavorites->AddFavorite(ActionPtr.Pin());
 		}
 		else
 		{
-			GEditor->EditorUserSettings->BlueprintFavorites->RemoveFavorite(ActionPtr.Pin());
+			GetMutableDefault<UEditorPerProjectUserSettings>()->BlueprintFavorites->RemoveFavorite(ActionPtr.Pin());
 		}
 	}	
 
diff --git a/Engine/Source/Editor/Kismet/Private/SBlueprintContextTargetMenu.cpp b/Engine/Source/Editor/Kismet/Private/SBlueprintContextTargetMenu.cpp
--- a/Engine/Source/Editor/Kismet/Private/SBlueprintContextTargetMenu.cpp
+++ b/Engine/Source/Editor/Kismet/Private/SBlueprintContextTargetMenu.cpp
@@ -361,7 +361,7 @@ FContextMenuTargetProfile::FContextMenuTargetProfile(const FBlueprintActionConte
 	FString const BpConfigKey = BpSettings->GetClass()->GetPathName();
 
 	bool bOldUseTargetContextForNodeMenu = true;
-	if (GConfig->GetBool(*BpConfigKey, TEXT("bUseTargetContextForNodeMenu"), bOldUseTargetContextForNodeMenu, GEditorUserSettingsIni) && !bOldUseTargetContextForNodeMenu)
+	if (GConfig->GetBool(*BpConfigKey, TEXT("bUseTargetContextForNodeMenu"), bOldUseTargetContextForNodeMenu, GEditorPerProjectIni) && !bOldUseTargetContextForNodeMenu)
 	{
 		SavedTargetFlags = 0;
 	}
diff --git a/Engine/Source/Editor/Kismet/Private/SBlueprintFavoritesPalette.cpp b/Engine/Source/Editor/Kismet/Private/SBlueprintFavoritesPalette.cpp
--- a/Engine/Source/Editor/Kismet/Private/SBlueprintFavoritesPalette.cpp
+++ b/Engine/Source/Editor/Kismet/Private/SBlueprintFavoritesPalette.cpp
@@ -32,13 +32,13 @@ struct SBlueprintFavoritesPaletteUtils
 	 */
 	static void RemoveSelectedFavorites(FPaletteActionGetter ActionGetter)
 	{
-		UEditorUserSettings const& EditorUserSettings = GEditor->AccessEditorUserSettings();
-		if (ActionGetter.IsBound() && (EditorUserSettings.BlueprintFavorites != NULL))
+		const UEditorPerProjectUserSettings* EditorPerProjectUserSettings = GetDefault<UEditorPerProjectUserSettings>();
+		if (ActionGetter.IsBound() && (EditorPerProjectUserSettings->BlueprintFavorites != NULL))
 		{
 			TArray< TSharedPtr<FEdGraphSchemaAction> > SelectedActions;
 			ActionGetter.Execute(SelectedActions);
 
-			EditorUserSettings.BlueprintFavorites->RemoveFavorites(SelectedActions);
+			EditorPerProjectUserSettings->BlueprintFavorites->RemoveFavorites(SelectedActions);
 		}
 	}
 
@@ -48,10 +48,10 @@ struct SBlueprintFavoritesPaletteUtils
 	 */
 	static void ClearPaletteFavorites()
 	{
-		UEditorUserSettings const& EditorUserSettings = GEditor->AccessEditorUserSettings();
-		if (EditorUserSettings.BlueprintFavorites != NULL)
+		const UEditorPerProjectUserSettings* EditorPerProjectUserSettings = GetDefault<UEditorPerProjectUserSettings>();
+		if (EditorPerProjectUserSettings->BlueprintFavorites != NULL)
 		{
-			EditorUserSettings.BlueprintFavorites->ClearAllFavorites();
+			EditorPerProjectUserSettings->BlueprintFavorites->ClearAllFavorites();
 		}
 	}
 
@@ -66,10 +66,10 @@ struct SBlueprintFavoritesPaletteUtils
 	{
 		bool bIsLoaded = false;
 
-		UEditorUserSettings const& EditorUserSettings = GEditor->AccessEditorUserSettings();
-		if (EditorUserSettings.BlueprintFavorites != NULL)
+		const UEditorPerProjectUserSettings* EditorPerProjectUserSettings = GetDefault<UEditorPerProjectUserSettings>();
+		if (EditorPerProjectUserSettings->BlueprintFavorites != NULL)
 		{
-			bIsLoaded = (EditorUserSettings.BlueprintFavorites->GetCurrentProfile() == ProfileName);
+			bIsLoaded = (EditorPerProjectUserSettings->BlueprintFavorites->GetCurrentProfile() == ProfileName);
 		}
 
 		return !bIsLoaded;
@@ -83,10 +83,10 @@ struct SBlueprintFavoritesPaletteUtils
 	 */
 	static void LoadFavoritesProfile(FString ProfileName)
 	{
-		UEditorUserSettings const& EditorUserSettings = GEditor->AccessEditorUserSettings();
-		if (EditorUserSettings.BlueprintFavorites != NULL)
+		const UEditorPerProjectUserSettings* EditorPerProjectUserSettings = GetDefault<UEditorPerProjectUserSettings>();
+		if (EditorPerProjectUserSettings->BlueprintFavorites != NULL)
 		{
-			EditorUserSettings.BlueprintFavorites->LoadProfile(ProfileName);
+			EditorPerProjectUserSettings->BlueprintFavorites->LoadProfile(ProfileName);
 		}
 	}
 
@@ -100,8 +100,8 @@ struct SBlueprintFavoritesPaletteUtils
 	{
 		TArray<FString> AvailableProfiles;
 	
-		UEditorUserSettings const& EditorUserSettings = GEditor->AccessEditorUserSettings();
-		if (EditorUserSettings.BlueprintFavorites != NULL)
+		const UEditorPerProjectUserSettings* EditorPerProjectUserSettings = GetDefault<UEditorPerProjectUserSettings>();
+		if (EditorPerProjectUserSettings->BlueprintFavorites != NULL)
 		{
 			static FString const ProfilesConfigKey("Profiles");
 			GConfig->GetArray(*SBlueprintFavoritesPaletteUtils::ConfigSection, *ProfilesConfigKey, AvailableProfiles, GEditorIni);	
@@ -268,10 +268,10 @@ public:
 //------------------------------------------------------------------------------
 SBlueprintFavoritesPalette::~SBlueprintFavoritesPalette()
 {
-	UEditorUserSettings& EditorUserSettings = GEditor->AccessEditorUserSettings();
-	if (EditorUserSettings.BlueprintFavorites != NULL)
+	const UEditorPerProjectUserSettings* EditorPerProjectUserSettings = GetDefault<UEditorPerProjectUserSettings>();
+	if (EditorPerProjectUserSettings->BlueprintFavorites != NULL)
 	{
-		EditorUserSettings.BlueprintFavorites->OnFavoritesUpdated.RemoveAll(this);
+		EditorPerProjectUserSettings->BlueprintFavorites->OnFavoritesUpdated.RemoveAll(this);
 	}
 }
 
@@ -290,10 +290,10 @@ void SBlueprintFavoritesPalette::Construct(FArguments const& InArgs, TWeakPtr<FB
 
 	SBlueprintSubPalette::Construct(SuperArgs, InBlueprintEditor);
 
-	UEditorUserSettings& EditorUserSettings = GEditor->AccessEditorUserSettings();
-	if (EditorUserSettings.BlueprintFavorites != NULL)
+	const UEditorPerProjectUserSettings* EditorPerProjectUserSettings = GetDefault<UEditorPerProjectUserSettings>();
+	if (EditorPerProjectUserSettings->BlueprintFavorites != NULL)
 	{
-		EditorUserSettings.BlueprintFavorites->OnFavoritesUpdated.AddSP(this, &SBlueprintFavoritesPalette::RefreshActionsList, true);
+		EditorPerProjectUserSettings->BlueprintFavorites->OnFavoritesUpdated.AddSP(this, &SBlueprintFavoritesPalette::RefreshActionsList, true);
 	}
 }
 
diff --git a/Engine/Source/Editor/Kismet/Private/SBlueprintLibraryPalette.cpp b/Engine/Source/Editor/Kismet/Private/SBlueprintLibraryPalette.cpp
--- a/Engine/Source/Editor/Kismet/Private/SBlueprintLibraryPalette.cpp
+++ b/Engine/Source/Editor/Kismet/Private/SBlueprintLibraryPalette.cpp
@@ -32,13 +32,13 @@ struct SBlueprintLibraryPaletteUtils
 	 */
 	static void AddSelectedToFavorites(FPaletteActionGetter ActionGetter)
 	{
-		UEditorUserSettings const& EditorUserSettings = GEditor->AccessEditorUserSettings();
-		if (ActionGetter.IsBound() && (EditorUserSettings.BlueprintFavorites != NULL))
+		const UEditorPerProjectUserSettings* EditorPerProjectUserSettings = GetDefault<UEditorPerProjectUserSettings>();
+		if (ActionGetter.IsBound() && (EditorPerProjectUserSettings->BlueprintFavorites != NULL))
 		{
 			TArray< TSharedPtr<FEdGraphSchemaAction> > SelectedActions;
 			ActionGetter.Execute(SelectedActions);
 
-			EditorUserSettings.BlueprintFavorites->AddFavorites(SelectedActions);
+			EditorPerProjectUserSettings->BlueprintFavorites->AddFavorites(SelectedActions);
 		}
 	}
 
@@ -50,13 +50,13 @@ struct SBlueprintLibraryPaletteUtils
 	 */
 	static void RemoveSelectedFavorites(FPaletteActionGetter ActionGetter)
 	{
-		UEditorUserSettings const& EditorUserSettings = GEditor->AccessEditorUserSettings();
-		if (ActionGetter.IsBound() && (EditorUserSettings.BlueprintFavorites != NULL))
+		const UEditorPerProjectUserSettings* EditorPerProjectUserSettings = GetDefault<UEditorPerProjectUserSettings>();
+		if (ActionGetter.IsBound() && (EditorPerProjectUserSettings->BlueprintFavorites != NULL))
 		{
 			TArray< TSharedPtr<FEdGraphSchemaAction> > SelectedActions;
 			ActionGetter.Execute(SelectedActions);
 
-			EditorUserSettings.BlueprintFavorites->RemoveFavorites(SelectedActions);
+			EditorPerProjectUserSettings->BlueprintFavorites->RemoveFavorites(SelectedActions);
 		}
 	}
 
@@ -72,15 +72,15 @@ struct SBlueprintLibraryPaletteUtils
 	{
 		bool bCanAnyBeFavorited = false;
 
-		UEditorUserSettings const& EditorUserSettings = GEditor->AccessEditorUserSettings();
-		if (ActionGetter.IsBound() && (EditorUserSettings.BlueprintFavorites != NULL))
+		const UEditorPerProjectUserSettings* EditorPerProjectUserSettings = GetDefault<UEditorPerProjectUserSettings>();
+		if (ActionGetter.IsBound() && (EditorPerProjectUserSettings->BlueprintFavorites != NULL))
 		{
 			TArray< TSharedPtr<FEdGraphSchemaAction> > SelectedActions;
 			ActionGetter.Execute(SelectedActions);
 
 			for (TSharedPtr<FEdGraphSchemaAction> Action : SelectedActions)
 			{
-				if (EditorUserSettings.BlueprintFavorites->CanBeFavorited(Action) && !EditorUserSettings.BlueprintFavorites->IsFavorited(Action))
+				if (EditorPerProjectUserSettings->BlueprintFavorites->CanBeFavorited(Action) && !EditorPerProjectUserSettings->BlueprintFavorites->IsFavorited(Action))
 				{
 					bCanAnyBeFavorited = true;
 					break;
@@ -102,15 +102,15 @@ struct SBlueprintLibraryPaletteUtils
 	{
 		bool bCanAnyBeRemoved = false;
 
-		UEditorUserSettings const& EditorUserSettings = GEditor->AccessEditorUserSettings();
-		if (ActionGetter.IsBound() && (EditorUserSettings.BlueprintFavorites != NULL))
+		const UEditorPerProjectUserSettings* EditorPerProjectUserSettings = GetDefault<UEditorPerProjectUserSettings>();
+		if (ActionGetter.IsBound() && (EditorPerProjectUserSettings->BlueprintFavorites != NULL))
 		{
 			TArray< TSharedPtr<FEdGraphSchemaAction> > SelectedActions;
 			ActionGetter.Execute(SelectedActions);
 
 			for (TSharedPtr<FEdGraphSchemaAction> Action : SelectedActions)
 			{
-				if (EditorUserSettings.BlueprintFavorites->IsFavorited(Action))
+				if (EditorPerProjectUserSettings->BlueprintFavorites->IsFavorited(Action))
 				{
 					bCanAnyBeRemoved = true;
 					break;
diff --git a/Engine/Source/Editor/Kismet/Private/SBlueprintPalette.cpp b/Engine/Source/Editor/Kismet/Private/SBlueprintPalette.cpp
--- a/Engine/Source/Editor/Kismet/Private/SBlueprintPalette.cpp
+++ b/Engine/Source/Editor/Kismet/Private/SBlueprintPalette.cpp
@@ -1555,9 +1555,9 @@ void SBlueprintPalette::Construct(const FArguments& InArgs, TWeakPtr<FBlueprintE
 	TSharedRef<SWidget> AssetDiscoveryIndicator = EditorWidgetsModule.CreateAssetDiscoveryIndicator(EAssetDiscoveryIndicatorScaleMode::Scale_Vertical);
 
 	float FavoritesHeightRatio = 0.33f;
-	GConfig->GetFloat(*BlueprintPalette::ConfigSection, *BlueprintPalette::FavoritesHeightConfigKey, FavoritesHeightRatio, GEditorUserSettingsIni);
+	GConfig->GetFloat(*BlueprintPalette::ConfigSection, *BlueprintPalette::FavoritesHeightConfigKey, FavoritesHeightRatio, GEditorPerProjectIni);
 	float LibraryHeightRatio = 1.f - FavoritesHeightRatio;
-	GConfig->GetFloat(*BlueprintPalette::ConfigSection, *BlueprintPalette::LibraryHeightConfigKey, LibraryHeightRatio, GEditorUserSettingsIni);
+	GConfig->GetFloat(*BlueprintPalette::ConfigSection, *BlueprintPalette::LibraryHeightConfigKey, LibraryHeightRatio, GEditorPerProjectIni);
 
 	bool bUseLegacyLayout = false;
 	GConfig->GetBool(*BlueprintPalette::ConfigSection, TEXT("bUseLegacyLayout"), bUseLegacyLayout, GEditorIni);
@@ -1606,11 +1606,11 @@ void SBlueprintPalette::OnSplitterResized() const
 
 		if (SplitterSlot.GetWidget() == FavoritesWrapper)
 		{
-			GConfig->SetFloat(*BlueprintPalette::ConfigSection, *BlueprintPalette::FavoritesHeightConfigKey, SplitterSlot.SizeValue.Get(), GEditorUserSettingsIni);
+			GConfig->SetFloat(*BlueprintPalette::ConfigSection, *BlueprintPalette::FavoritesHeightConfigKey, SplitterSlot.SizeValue.Get(), GEditorPerProjectIni);
 		}
 		else if (SplitterSlot.GetWidget() == LibraryWrapper)
 		{
-			GConfig->SetFloat(*BlueprintPalette::ConfigSection, *BlueprintPalette::LibraryHeightConfigKey, SplitterSlot.SizeValue.Get(), GEditorUserSettingsIni);
+			GConfig->SetFloat(*BlueprintPalette::ConfigSection, *BlueprintPalette::LibraryHeightConfigKey, SplitterSlot.SizeValue.Get(), GEditorPerProjectIni);
 		}
 
 	}
diff --git a/Engine/Source/Editor/Kismet/Private/SCSEditorViewportClient.cpp b/Engine/Source/Editor/Kismet/Private/SCSEditorViewportClient.cpp
--- a/Engine/Source/Editor/Kismet/Private/SCSEditorViewportClient.cpp
+++ b/Engine/Source/Editor/Kismet/Private/SCSEditorViewportClient.cpp
@@ -125,7 +125,7 @@ FSCSEditorViewportClient::FSCSEditorViewportClient(TWeakPtr<FBlueprintEditor>& I
 	EngineShowFlags.SelectionOutline = GetDefault<ULevelEditorViewportSettings>()->bUseSelectionOutline;
 
 	// Set if the grid will be drawn
-	DrawHelper.bDrawGrid = GEditor->GetEditorUserSettings().bSCSEditorShowGrid;
+	DrawHelper.bDrawGrid = GetDefault<UEditorPerProjectUserSettings>()->bSCSEditorShowGrid;
 
 	// now add floor
 	EditorFloorComp = NewObject<UStaticMeshComponent>(GetTransientPackage(), TEXT("EditorFloorComp"));
@@ -143,8 +143,8 @@ FSCSEditorViewportClient::FSCSEditorViewportClient(TWeakPtr<FBlueprintEditor>& I
 	}
 
 	EditorFloorComp->SetRelativeScale3D(FVector(3.f, 3.f, 1.f));
-	EditorFloorComp->SetVisibility(GEditor->GetEditorUserSettings().bSCSEditorShowFloor);
-	EditorFloorComp->SetCollisionEnabled(GEditor->GetEditorUserSettings().bSCSEditorShowFloor? ECollisionEnabled::QueryAndPhysics : ECollisionEnabled::NoCollision);
+	EditorFloorComp->SetVisibility(GetDefault<UEditorPerProjectUserSettings>()->bSCSEditorShowFloor);
+	EditorFloorComp->SetCollisionEnabled(GetDefault<UEditorPerProjectUserSettings>()->bSCSEditorShowFloor? ECollisionEnabled::QueryAndPhysics : ECollisionEnabled::NoCollision);
 	PreviewScene->AddComponent(EditorFloorComp, FTransform::Identity);
 
 	// Turn off so that actors added to the world do not have a lifespan (so they will not auto-destroy themselves).
@@ -773,37 +773,43 @@ void FSCSEditorViewportClient::ToggleIsSimulateEnabled()
 
 bool FSCSEditorViewportClient::GetShowFloor() 
 {
-	return GEditor->GetEditorUserSettings().bSCSEditorShowFloor;
+	return GetDefault<UEditorPerProjectUserSettings>()->bSCSEditorShowFloor;
 }
 
 void FSCSEditorViewportClient::ToggleShowFloor() 
 {
-	bool bShowFloor = GEditor->AccessEditorUserSettings().bSCSEditorShowFloor;
+	auto* Settings = GetMutableDefault<UEditorPerProjectUserSettings>();
+
+	bool bShowFloor = Settings->bSCSEditorShowFloor;
 	bShowFloor = !bShowFloor;
 	
 	EditorFloorComp->SetVisibility(bShowFloor);
 	EditorFloorComp->SetCollisionEnabled(bShowFloor? ECollisionEnabled::QueryAndPhysics : ECollisionEnabled::NoCollision);
 	GetWorld()->SendAllEndOfFrameUpdates();
 
-	GEditor->AccessEditorUserSettings().bSCSEditorShowFloor = bShowFloor;
+	Settings->bSCSEditorShowFloor = bShowFloor;
+	Settings->PostEditChange();
 
 	Invalidate();
 }
 
 bool FSCSEditorViewportClient::GetShowGrid() 
 {
-	return GEditor->GetEditorUserSettings().bSCSEditorShowGrid;
+	return GetDefault<UEditorPerProjectUserSettings>()->bSCSEditorShowGrid;
 }
 
 void FSCSEditorViewportClient::ToggleShowGrid() 
 {
-	bool bShowGrid = GEditor->AccessEditorUserSettings().bSCSEditorShowGrid;
+	auto* Settings = GetMutableDefault<UEditorPerProjectUserSettings>();
+
+	bool bShowGrid = Settings->bSCSEditorShowGrid;
 	bShowGrid = !bShowGrid;
 
 	DrawHelper.bDrawGrid = bShowGrid;
 
-	GEditor->AccessEditorUserSettings().bSCSEditorShowGrid = bShowGrid;
-
+	Settings->bSCSEditorShowGrid = bShowGrid;
+	Settings->PostEditChange();
+	
 	Invalidate();
 }
 
diff --git a/Engine/Source/Editor/LandscapeEditor/Private/LandscapeEditorObject.cpp b/Engine/Source/Editor/LandscapeEditor/Private/LandscapeEditorObject.cpp
--- a/Engine/Source/Editor/LandscapeEditor/Private/LandscapeEditorObject.cpp
+++ b/Engine/Source/Editor/LandscapeEditor/Private/LandscapeEditorObject.cpp
@@ -156,105 +156,105 @@ void ULandscapeEditorObject::PostEditChangeProperty(FPropertyChangedEvent& Prope
 /** Load UI settings from ini file */
 void ULandscapeEditorObject::Load()
 {
-	GConfig->GetFloat(TEXT("LandscapeEdit"), TEXT("ToolStrength"), ToolStrength, GEditorUserSettingsIni);
-	GConfig->GetFloat(TEXT("LandscapeEdit"), TEXT("WeightTargetValue"), WeightTargetValue, GEditorUserSettingsIni);
+	GConfig->GetFloat(TEXT("LandscapeEdit"), TEXT("ToolStrength"), ToolStrength, GEditorPerProjectIni);
+	GConfig->GetFloat(TEXT("LandscapeEdit"), TEXT("WeightTargetValue"), WeightTargetValue, GEditorPerProjectIni);
 	bool InbUseWeightTargetValue = bUseWeightTargetValue;
-	GConfig->GetBool(TEXT("LandscapeEdit"), TEXT("bUseWeightTargetValue"), InbUseWeightTargetValue, GEditorUserSettingsIni);
+	GConfig->GetBool(TEXT("LandscapeEdit"), TEXT("bUseWeightTargetValue"), InbUseWeightTargetValue, GEditorPerProjectIni);
 	bUseWeightTargetValue = InbUseWeightTargetValue;
 
-	GConfig->GetFloat(TEXT("LandscapeEdit"), TEXT("BrushRadius"), BrushRadius, GEditorUserSettingsIni);
-	GConfig->GetInt(TEXT("LandscapeEdit"), TEXT("BrushComponentSize"), BrushComponentSize, GEditorUserSettingsIni);
-	GConfig->GetFloat(TEXT("LandscapeEdit"), TEXT("BrushFalloff"), BrushFalloff, GEditorUserSettingsIni);
+	GConfig->GetFloat(TEXT("LandscapeEdit"), TEXT("BrushRadius"), BrushRadius, GEditorPerProjectIni);
+	GConfig->GetInt(TEXT("LandscapeEdit"), TEXT("BrushComponentSize"), BrushComponentSize, GEditorPerProjectIni);
+	GConfig->GetFloat(TEXT("LandscapeEdit"), TEXT("BrushFalloff"), BrushFalloff, GEditorPerProjectIni);
 	bool InbUseClayBrush = bUseClayBrush;
-	GConfig->GetBool(TEXT("LandscapeEdit"), TEXT("bUseClayBrush"), InbUseClayBrush, GEditorUserSettingsIni);
+	GConfig->GetBool(TEXT("LandscapeEdit"), TEXT("bUseClayBrush"), InbUseClayBrush, GEditorPerProjectIni);
 	bUseClayBrush = InbUseClayBrush;
-	GConfig->GetFloat(TEXT("LandscapeEdit"), TEXT("AlphaBrushScale"), AlphaBrushScale, GEditorUserSettingsIni);
-	GConfig->GetBool(TEXT("LandscapeEdit"), TEXT("AlphaBrushAutoRotate"), bAlphaBrushAutoRotate, GEditorUserSettingsIni);
-	GConfig->GetFloat(TEXT("LandscapeEdit"), TEXT("AlphaBrushRotation"), AlphaBrushRotation, GEditorUserSettingsIni);
-	GConfig->GetFloat(TEXT("LandscapeEdit"), TEXT("AlphaBrushPanU"), AlphaBrushPanU, GEditorUserSettingsIni);
-	GConfig->GetFloat(TEXT("LandscapeEdit"), TEXT("AlphaBrushPanV"), AlphaBrushPanV, GEditorUserSettingsIni);
-	GConfig->GetBool(TEXT("LandscapeEdit"), TEXT("bUseWorldSpacePatternBrush"), bUseWorldSpacePatternBrush, GEditorUserSettingsIni);
-	GConfig->GetVector2D(TEXT("LandscapeEdit"), TEXT("WorldSpacePatternBrushSettings.Origin"), WorldSpacePatternBrushSettings.Origin, GEditorUserSettingsIni);
-	GConfig->GetBool(TEXT("LandscapeEdit"), TEXT("WorldSpacePatternBrushSettings.bCenterTextureOnOrigin"), WorldSpacePatternBrushSettings.bCenterTextureOnOrigin, GEditorUserSettingsIni);
-	GConfig->GetFloat(TEXT("LandscapeEdit"), TEXT("WorldSpacePatternBrushSettings.RepeatSize"), WorldSpacePatternBrushSettings.RepeatSize, GEditorUserSettingsIni);
+	GConfig->GetFloat(TEXT("LandscapeEdit"), TEXT("AlphaBrushScale"), AlphaBrushScale, GEditorPerProjectIni);
+	GConfig->GetBool(TEXT("LandscapeEdit"), TEXT("AlphaBrushAutoRotate"), bAlphaBrushAutoRotate, GEditorPerProjectIni);
+	GConfig->GetFloat(TEXT("LandscapeEdit"), TEXT("AlphaBrushRotation"), AlphaBrushRotation, GEditorPerProjectIni);
+	GConfig->GetFloat(TEXT("LandscapeEdit"), TEXT("AlphaBrushPanU"), AlphaBrushPanU, GEditorPerProjectIni);
+	GConfig->GetFloat(TEXT("LandscapeEdit"), TEXT("AlphaBrushPanV"), AlphaBrushPanV, GEditorPerProjectIni);
+	GConfig->GetBool(TEXT("LandscapeEdit"), TEXT("bUseWorldSpacePatternBrush"), bUseWorldSpacePatternBrush, GEditorPerProjectIni);
+	GConfig->GetVector2D(TEXT("LandscapeEdit"), TEXT("WorldSpacePatternBrushSettings.Origin"), WorldSpacePatternBrushSettings.Origin, GEditorPerProjectIni);
+	GConfig->GetBool(TEXT("LandscapeEdit"), TEXT("WorldSpacePatternBrushSettings.bCenterTextureOnOrigin"), WorldSpacePatternBrushSettings.bCenterTextureOnOrigin, GEditorPerProjectIni);
+	GConfig->GetFloat(TEXT("LandscapeEdit"), TEXT("WorldSpacePatternBrushSettings.RepeatSize"), WorldSpacePatternBrushSettings.RepeatSize, GEditorPerProjectIni);
 	FString AlphaTextureName = (AlphaTexture != NULL) ? AlphaTexture->GetPathName() : FString();
 	int32 InAlphaTextureChannel = AlphaTextureChannel;
-	GConfig->GetString(TEXT("LandscapeEdit"), TEXT("AlphaTextureName"), AlphaTextureName, GEditorUserSettingsIni);
-	GConfig->GetInt(TEXT("LandscapeEdit"), TEXT("AlphaTextureChannel"), InAlphaTextureChannel, GEditorUserSettingsIni);
+	GConfig->GetString(TEXT("LandscapeEdit"), TEXT("AlphaTextureName"), AlphaTextureName, GEditorPerProjectIni);
+	GConfig->GetInt(TEXT("LandscapeEdit"), TEXT("AlphaTextureChannel"), InAlphaTextureChannel, GEditorPerProjectIni);
 	AlphaTextureChannel = (EColorChannel::Type)InAlphaTextureChannel;
 	SetAlphaTexture(LoadObject<UTexture2D>(NULL, *AlphaTextureName, NULL, LOAD_NoWarn), AlphaTextureChannel);
 
 	int32 InFlattenMode = ELandscapeToolFlattenMode::Both;
-	GConfig->GetInt(TEXT("LandscapeEdit"), TEXT("FlattenMode"), InFlattenMode, GEditorUserSettingsIni);
+	GConfig->GetInt(TEXT("LandscapeEdit"), TEXT("FlattenMode"), InFlattenMode, GEditorPerProjectIni);
 	FlattenMode = (ELandscapeToolFlattenMode::Type)InFlattenMode;
 
 	bool InbUseSlopeFlatten = bUseSlopeFlatten;
-	GConfig->GetBool(TEXT("LandscapeEdit"), TEXT("bUseSlopeFlatten"), InbUseSlopeFlatten, GEditorUserSettingsIni);
+	GConfig->GetBool(TEXT("LandscapeEdit"), TEXT("bUseSlopeFlatten"), InbUseSlopeFlatten, GEditorPerProjectIni);
 	bUseSlopeFlatten = InbUseSlopeFlatten;
 
 	bool InbPickValuePerApply = bPickValuePerApply;
-	GConfig->GetBool(TEXT("LandscapeEdit"), TEXT("bPickValuePerApply"), InbPickValuePerApply, GEditorUserSettingsIni);
+	GConfig->GetBool(TEXT("LandscapeEdit"), TEXT("bPickValuePerApply"), InbPickValuePerApply, GEditorPerProjectIni);
 	bPickValuePerApply = InbPickValuePerApply;
 
 	bool InbUseFlattenTarget = bUseFlattenTarget;
-	GConfig->GetBool(TEXT("LandscapeEdit"), TEXT("bUseFlattenTarget"), InbUseFlattenTarget, GEditorUserSettingsIni);
+	GConfig->GetBool(TEXT("LandscapeEdit"), TEXT("bUseFlattenTarget"), InbUseFlattenTarget, GEditorPerProjectIni);
 	bUseFlattenTarget = InbUseFlattenTarget;
-	GConfig->GetFloat(TEXT("LandscapeEdit"), TEXT("FlattenTarget"), FlattenTarget, GEditorUserSettingsIni);
+	GConfig->GetFloat(TEXT("LandscapeEdit"), TEXT("FlattenTarget"), FlattenTarget, GEditorPerProjectIni);
 
-	GConfig->GetFloat(TEXT("LandscapeEdit"), TEXT("RampWidth"), RampWidth, GEditorUserSettingsIni);
-	GConfig->GetFloat(TEXT("LandscapeEdit"), TEXT("RampSideFalloff"), RampSideFalloff, GEditorUserSettingsIni);
+	GConfig->GetFloat(TEXT("LandscapeEdit"), TEXT("RampWidth"), RampWidth, GEditorPerProjectIni);
+	GConfig->GetFloat(TEXT("LandscapeEdit"), TEXT("RampSideFalloff"), RampSideFalloff, GEditorPerProjectIni);
 
-	GConfig->GetInt(TEXT("LandscapeEdit"), TEXT("ErodeThresh"), ErodeThresh, GEditorUserSettingsIni);
-	GConfig->GetInt(TEXT("LandscapeEdit"), TEXT("ErodeIterationNum"), ErodeIterationNum, GEditorUserSettingsIni);
-	GConfig->GetInt(TEXT("LandscapeEdit"), TEXT("ErodeSurfaceThickness"), ErodeSurfaceThickness, GEditorUserSettingsIni);
+	GConfig->GetInt(TEXT("LandscapeEdit"), TEXT("ErodeThresh"), ErodeThresh, GEditorPerProjectIni);
+	GConfig->GetInt(TEXT("LandscapeEdit"), TEXT("ErodeIterationNum"), ErodeIterationNum, GEditorPerProjectIni);
+	GConfig->GetInt(TEXT("LandscapeEdit"), TEXT("ErodeSurfaceThickness"), ErodeSurfaceThickness, GEditorPerProjectIni);
 	int32 InErosionNoiseMode = ErosionNoiseMode;
-	GConfig->GetInt(TEXT("LandscapeEdit"), TEXT("ErosionNoiseMode"), InErosionNoiseMode, GEditorUserSettingsIni);
+	GConfig->GetInt(TEXT("LandscapeEdit"), TEXT("ErosionNoiseMode"), InErosionNoiseMode, GEditorPerProjectIni);
 	ErosionNoiseMode = (ELandscapeToolErosionMode::Type)InErosionNoiseMode;
-	GConfig->GetFloat(TEXT("LandscapeEdit"), TEXT("ErosionNoiseScale"), ErosionNoiseScale, GEditorUserSettingsIni);
+	GConfig->GetFloat(TEXT("LandscapeEdit"), TEXT("ErosionNoiseScale"), ErosionNoiseScale, GEditorPerProjectIni);
 
-	GConfig->GetInt(TEXT("LandscapeEdit"), TEXT("RainAmount"), RainAmount, GEditorUserSettingsIni);
-	GConfig->GetFloat(TEXT("LandscapeEdit"), TEXT("SedimentCapacity"), SedimentCapacity, GEditorUserSettingsIni);
-	GConfig->GetInt(TEXT("LandscapeEdit"), TEXT("HErodeIterationNum"), HErodeIterationNum, GEditorUserSettingsIni);
+	GConfig->GetInt(TEXT("LandscapeEdit"), TEXT("RainAmount"), RainAmount, GEditorPerProjectIni);
+	GConfig->GetFloat(TEXT("LandscapeEdit"), TEXT("SedimentCapacity"), SedimentCapacity, GEditorPerProjectIni);
+	GConfig->GetInt(TEXT("LandscapeEdit"), TEXT("HErodeIterationNum"), HErodeIterationNum, GEditorPerProjectIni);
 	int32 InRainDistMode = RainDistMode;
-	GConfig->GetInt(TEXT("LandscapeEdit"), TEXT("RainDistNoiseMode"), InRainDistMode, GEditorUserSettingsIni);
+	GConfig->GetInt(TEXT("LandscapeEdit"), TEXT("RainDistNoiseMode"), InRainDistMode, GEditorPerProjectIni);
 	RainDistMode = (ELandscapeToolHydroErosionMode::Type)InRainDistMode;
-	GConfig->GetFloat(TEXT("LandscapeEdit"), TEXT("RainDistScale"), RainDistScale, GEditorUserSettingsIni);
-	GConfig->GetFloat(TEXT("LandscapeEdit"), TEXT("HErosionDetailScale"), HErosionDetailScale, GEditorUserSettingsIni);
+	GConfig->GetFloat(TEXT("LandscapeEdit"), TEXT("RainDistScale"), RainDistScale, GEditorPerProjectIni);
+	GConfig->GetFloat(TEXT("LandscapeEdit"), TEXT("HErosionDetailScale"), HErosionDetailScale, GEditorPerProjectIni);
 	bool InbHErosionDetailSmooth = bHErosionDetailSmooth;
-	GConfig->GetBool(TEXT("LandscapeEdit"), TEXT("bHErosionDetailSmooth"), InbHErosionDetailSmooth, GEditorUserSettingsIni);
+	GConfig->GetBool(TEXT("LandscapeEdit"), TEXT("bHErosionDetailSmooth"), InbHErosionDetailSmooth, GEditorPerProjectIni);
 	bHErosionDetailSmooth = InbHErosionDetailSmooth;
 
 	int32 InNoiseMode = NoiseMode;
-	GConfig->GetInt(TEXT("LandscapeEdit"), TEXT("NoiseMode"), InNoiseMode, GEditorUserSettingsIni);
+	GConfig->GetInt(TEXT("LandscapeEdit"), TEXT("NoiseMode"), InNoiseMode, GEditorPerProjectIni);
 	NoiseMode = (ELandscapeToolNoiseMode::Type)InNoiseMode;
-	GConfig->GetFloat(TEXT("LandscapeEdit"), TEXT("NoiseScale"), NoiseScale, GEditorUserSettingsIni);
+	GConfig->GetFloat(TEXT("LandscapeEdit"), TEXT("NoiseScale"), NoiseScale, GEditorPerProjectIni);
 
-	GConfig->GetInt(TEXT("LandscapeEdit"), TEXT("SmoothFilterKernelSize"), SmoothFilterKernelSize, GEditorUserSettingsIni);
-	GConfig->GetFloat(TEXT("LandscapeEdit"), TEXT("DetailScale"), DetailScale, GEditorUserSettingsIni);
+	GConfig->GetInt(TEXT("LandscapeEdit"), TEXT("SmoothFilterKernelSize"), SmoothFilterKernelSize, GEditorPerProjectIni);
+	GConfig->GetFloat(TEXT("LandscapeEdit"), TEXT("DetailScale"), DetailScale, GEditorPerProjectIni);
 	bool InbDetailSmooth = bDetailSmooth;
-	GConfig->GetBool(TEXT("LandscapeEdit"), TEXT("bDetailSmooth"), InbDetailSmooth, GEditorUserSettingsIni);
+	GConfig->GetBool(TEXT("LandscapeEdit"), TEXT("bDetailSmooth"), InbDetailSmooth, GEditorPerProjectIni);
 	bDetailSmooth = InbDetailSmooth;
 
-	GConfig->GetFloat(TEXT("LandscapeEdit"), TEXT("MaximumValueRadius"), MaximumValueRadius, GEditorUserSettingsIni);
+	GConfig->GetFloat(TEXT("LandscapeEdit"), TEXT("MaximumValueRadius"), MaximumValueRadius, GEditorPerProjectIni);
 
 	bool InbSmoothGizmoBrush = bSmoothGizmoBrush;
-	GConfig->GetBool(TEXT("LandscapeEdit"), TEXT("bSmoothGizmoBrush"), InbSmoothGizmoBrush, GEditorUserSettingsIni);
+	GConfig->GetBool(TEXT("LandscapeEdit"), TEXT("bSmoothGizmoBrush"), InbSmoothGizmoBrush, GEditorPerProjectIni);
 	bSmoothGizmoBrush = InbSmoothGizmoBrush;
 
 	int32 InPasteMode = ELandscapeToolNoiseMode::Both; 
-	GConfig->GetInt(TEXT("LandscapeEdit"), TEXT("PasteMode"), InPasteMode, GEditorUserSettingsIni);
+	GConfig->GetInt(TEXT("LandscapeEdit"), TEXT("PasteMode"), InPasteMode, GEditorPerProjectIni);
 	//PasteMode = (ELandscapeToolPasteMode::Type)InPasteMode;
 	SetPasteMode((ELandscapeToolPasteMode::Type)InPasteMode);
 
 	int32 InConvertMode = ResizeLandscape_ConvertMode;
-	GConfig->GetInt(TEXT("LandscapeEdit"), TEXT("ConvertMode"), InConvertMode, GEditorUserSettingsIni);
+	GConfig->GetInt(TEXT("LandscapeEdit"), TEXT("ConvertMode"), InConvertMode, GEditorPerProjectIni);
 	ResizeLandscape_ConvertMode = (ELandscapeConvertMode::Type)InConvertMode;
 
 	// Region
-	//GConfig->GetBool(TEXT("LandscapeEdit"), TEXT("bUseSelectedRegion"), bUseSelectedRegion, GEditorUserSettingsIni);
-	//GConfig->GetBool(TEXT("LandscapeEdit"), TEXT("bUseNegativeMask"), bUseNegativeMask, GEditorUserSettingsIni);
+	//GConfig->GetBool(TEXT("LandscapeEdit"), TEXT("bUseSelectedRegion"), bUseSelectedRegion, GEditorPerProjectIni);
+	//GConfig->GetBool(TEXT("LandscapeEdit"), TEXT("bUseNegativeMask"), bUseNegativeMask, GEditorPerProjectIni);
 	bool InbApplyToAllTargets = bApplyToAllTargets;
-	GConfig->GetBool(TEXT("LandscapeEdit"), TEXT("bApplyToAllTargets"), InbApplyToAllTargets, GEditorUserSettingsIni);
+	GConfig->GetBool(TEXT("LandscapeEdit"), TEXT("bApplyToAllTargets"), InbApplyToAllTargets, GEditorPerProjectIni);
 	bApplyToAllTargets = InbApplyToAllTargets;
 
 	// Set EditRenderMode
@@ -273,7 +273,7 @@ void ULandscapeEditorObject::Load()
 	}
 
 	FString NewLandscapeMaterialName = (NewLandscape_Material != NULL) ? NewLandscape_Material->GetPathName() : FString();
-	GConfig->GetString(TEXT("LandscapeEdit"), TEXT("NewLandscapeMaterialName"), NewLandscapeMaterialName, GEditorUserSettingsIni);
+	GConfig->GetString(TEXT("LandscapeEdit"), TEXT("NewLandscapeMaterialName"), NewLandscapeMaterialName, GEditorPerProjectIni);
 	NewLandscape_Material = LoadObject<UMaterialInterface>(NULL, *NewLandscapeMaterialName, NULL, LOAD_NoWarn);
 
 	RefreshImportLayersList();
@@ -282,66 +282,66 @@ void ULandscapeEditorObject::Load()
 /** Save UI settings to ini file */
 void ULandscapeEditorObject::Save()
 {
-	GConfig->SetFloat(TEXT("LandscapeEdit"), TEXT("ToolStrength"), ToolStrength, GEditorUserSettingsIni);
-	GConfig->SetFloat(TEXT("LandscapeEdit"), TEXT("WeightTargetValue"), WeightTargetValue, GEditorUserSettingsIni);
-	GConfig->SetBool(TEXT("LandscapeEdit"), TEXT("bUseWeightTargetValue"), bUseWeightTargetValue, GEditorUserSettingsIni);
-
-	GConfig->SetFloat(TEXT("LandscapeEdit"), TEXT("BrushRadius"), BrushRadius, GEditorUserSettingsIni);
-	GConfig->SetInt(TEXT("LandscapeEdit"), TEXT("BrushComponentSize"), BrushComponentSize, GEditorUserSettingsIni);
-	GConfig->SetFloat(TEXT("LandscapeEdit"), TEXT("BrushFalloff"), BrushFalloff, GEditorUserSettingsIni);
-	GConfig->SetBool(TEXT("LandscapeEdit"), TEXT("bUseClayBrush"), bUseClayBrush, GEditorUserSettingsIni);
-	GConfig->SetFloat(TEXT("LandscapeEdit"), TEXT("AlphaBrushScale"), AlphaBrushScale, GEditorUserSettingsIni);
-	GConfig->SetBool(TEXT("LandscapeEdit"), TEXT("AlphaBrushAutoRotate"), bAlphaBrushAutoRotate, GEditorUserSettingsIni);
-	GConfig->SetFloat(TEXT("LandscapeEdit"), TEXT("AlphaBrushRotation"), AlphaBrushRotation, GEditorUserSettingsIni);
-	GConfig->SetFloat(TEXT("LandscapeEdit"), TEXT("AlphaBrushPanU"), AlphaBrushPanU, GEditorUserSettingsIni);
-	GConfig->SetFloat(TEXT("LandscapeEdit"), TEXT("AlphaBrushPanV"), AlphaBrushPanV, GEditorUserSettingsIni);
-	GConfig->SetVector2D(TEXT("LandscapeEdit"), TEXT("WorldSpacePatternBrushSettings.Origin"), WorldSpacePatternBrushSettings.Origin, GEditorUserSettingsIni);
-	GConfig->SetBool(TEXT("LandscapeEdit"), TEXT("WorldSpacePatternBrushSettings.bCenterTextureOnOrigin"), WorldSpacePatternBrushSettings.bCenterTextureOnOrigin, GEditorUserSettingsIni);
-	GConfig->SetFloat(TEXT("LandscapeEdit"), TEXT("WorldSpacePatternBrushSettings.RepeatSize"), WorldSpacePatternBrushSettings.RepeatSize, GEditorUserSettingsIni);
+	GConfig->SetFloat(TEXT("LandscapeEdit"), TEXT("ToolStrength"), ToolStrength, GEditorPerProjectIni);
+	GConfig->SetFloat(TEXT("LandscapeEdit"), TEXT("WeightTargetValue"), WeightTargetValue, GEditorPerProjectIni);
+	GConfig->SetBool(TEXT("LandscapeEdit"), TEXT("bUseWeightTargetValue"), bUseWeightTargetValue, GEditorPerProjectIni);
+
+	GConfig->SetFloat(TEXT("LandscapeEdit"), TEXT("BrushRadius"), BrushRadius, GEditorPerProjectIni);
+	GConfig->SetInt(TEXT("LandscapeEdit"), TEXT("BrushComponentSize"), BrushComponentSize, GEditorPerProjectIni);
+	GConfig->SetFloat(TEXT("LandscapeEdit"), TEXT("BrushFalloff"), BrushFalloff, GEditorPerProjectIni);
+	GConfig->SetBool(TEXT("LandscapeEdit"), TEXT("bUseClayBrush"), bUseClayBrush, GEditorPerProjectIni);
+	GConfig->SetFloat(TEXT("LandscapeEdit"), TEXT("AlphaBrushScale"), AlphaBrushScale, GEditorPerProjectIni);
+	GConfig->SetBool(TEXT("LandscapeEdit"), TEXT("AlphaBrushAutoRotate"), bAlphaBrushAutoRotate, GEditorPerProjectIni);
+	GConfig->SetFloat(TEXT("LandscapeEdit"), TEXT("AlphaBrushRotation"), AlphaBrushRotation, GEditorPerProjectIni);
+	GConfig->SetFloat(TEXT("LandscapeEdit"), TEXT("AlphaBrushPanU"), AlphaBrushPanU, GEditorPerProjectIni);
+	GConfig->SetFloat(TEXT("LandscapeEdit"), TEXT("AlphaBrushPanV"), AlphaBrushPanV, GEditorPerProjectIni);
+	GConfig->SetVector2D(TEXT("LandscapeEdit"), TEXT("WorldSpacePatternBrushSettings.Origin"), WorldSpacePatternBrushSettings.Origin, GEditorPerProjectIni);
+	GConfig->SetBool(TEXT("LandscapeEdit"), TEXT("WorldSpacePatternBrushSettings.bCenterTextureOnOrigin"), WorldSpacePatternBrushSettings.bCenterTextureOnOrigin, GEditorPerProjectIni);
+	GConfig->SetFloat(TEXT("LandscapeEdit"), TEXT("WorldSpacePatternBrushSettings.RepeatSize"), WorldSpacePatternBrushSettings.RepeatSize, GEditorPerProjectIni);
 	const FString AlphaTextureName = (AlphaTexture != NULL) ? AlphaTexture->GetPathName() : FString();
-	GConfig->SetString(TEXT("LandscapeEdit"), TEXT("AlphaTextureName"), *AlphaTextureName, GEditorUserSettingsIni);
-	GConfig->SetInt(TEXT("LandscapeEdit"), TEXT("AlphaTextureChannel"), (int32)AlphaTextureChannel, GEditorUserSettingsIni);
-
-	GConfig->SetInt(TEXT("LandscapeEdit"), TEXT("FlattenMode"), FlattenMode, GEditorUserSettingsIni);
-	GConfig->SetBool(TEXT("LandscapeEdit"), TEXT("bUseSlopeFlatten"), bUseSlopeFlatten, GEditorUserSettingsIni);
-	GConfig->SetBool(TEXT("LandscapeEdit"), TEXT("bPickValuePerApply"), bPickValuePerApply, GEditorUserSettingsIni);
-	GConfig->SetBool(TEXT("LandscapeEdit"), TEXT("bUseFlattenTarget"), bUseFlattenTarget, GEditorUserSettingsIni);
-	GConfig->SetFloat(TEXT("LandscapeEdit"), TEXT("FlattenTarget"), FlattenTarget, GEditorUserSettingsIni);
-
-	GConfig->SetFloat(TEXT("LandscapeEdit"), TEXT("RampWidth"), RampWidth, GEditorUserSettingsIni);
-	GConfig->SetFloat(TEXT("LandscapeEdit"), TEXT("RampSideFalloff"), RampSideFalloff, GEditorUserSettingsIni);
-
-	GConfig->SetInt(TEXT("LandscapeEdit"), TEXT("ErodeThresh"), ErodeThresh, GEditorUserSettingsIni);
-	GConfig->SetInt(TEXT("LandscapeEdit"), TEXT("ErodeIterationNum"), ErodeIterationNum, GEditorUserSettingsIni);
-	GConfig->SetInt(TEXT("LandscapeEdit"), TEXT("ErodeSurfaceThickness"), ErodeSurfaceThickness, GEditorUserSettingsIni);
-	GConfig->SetInt(TEXT("LandscapeEdit"), TEXT("ErosionNoiseMode"), (int32)ErosionNoiseMode, GEditorUserSettingsIni);
-	GConfig->SetFloat(TEXT("LandscapeEdit"), TEXT("ErosionNoiseScale"), ErosionNoiseScale, GEditorUserSettingsIni);
-
-	GConfig->SetInt(TEXT("LandscapeEdit"), TEXT("RainAmount"), RainAmount, GEditorUserSettingsIni);
-	GConfig->SetFloat(TEXT("LandscapeEdit"), TEXT("SedimentCapacity"), SedimentCapacity, GEditorUserSettingsIni);
-	GConfig->SetInt(TEXT("LandscapeEdit"), TEXT("HErodeIterationNum"), ErodeIterationNum, GEditorUserSettingsIni);
-	GConfig->SetInt(TEXT("LandscapeEdit"), TEXT("RainDistMode"), (int32)RainDistMode, GEditorUserSettingsIni);
-	GConfig->SetFloat(TEXT("LandscapeEdit"), TEXT("RainDistScale"), RainDistScale, GEditorUserSettingsIni);
-	GConfig->SetFloat(TEXT("LandscapeEdit"), TEXT("HErosionDetailScale"), HErosionDetailScale, GEditorUserSettingsIni);
-	GConfig->SetBool(TEXT("LandscapeEdit"), TEXT("bHErosionDetailSmooth"), bHErosionDetailSmooth, GEditorUserSettingsIni);
-
-	GConfig->SetInt(TEXT("LandscapeEdit"), TEXT("NoiseMode"), (int32)NoiseMode, GEditorUserSettingsIni);
-	GConfig->SetFloat(TEXT("LandscapeEdit"), TEXT("NoiseScale"), NoiseScale, GEditorUserSettingsIni);
-	GConfig->SetInt(TEXT("LandscapeEdit"), TEXT("SmoothFilterKernelSize"), SmoothFilterKernelSize, GEditorUserSettingsIni);
-	GConfig->SetFloat(TEXT("LandscapeEdit"), TEXT("DetailScale"), DetailScale, GEditorUserSettingsIni);
-	GConfig->SetBool(TEXT("LandscapeEdit"), TEXT("bDetailSmooth"), bDetailSmooth, GEditorUserSettingsIni);
-
-	GConfig->SetFloat(TEXT("LandscapeEdit"), TEXT("MaximumValueRadius"), MaximumValueRadius, GEditorUserSettingsIni);
-
-	GConfig->SetBool(TEXT("LandscapeEdit"), TEXT("bSmoothGizmoBrush"), bSmoothGizmoBrush, GEditorUserSettingsIni);
-	GConfig->SetInt(TEXT("LandscapeEdit"), TEXT("PasteMode"), (int32)PasteMode, GEditorUserSettingsIni);
-	GConfig->SetInt(TEXT("LandscapeEdit"), TEXT("ConvertMode"), (int32)ResizeLandscape_ConvertMode, GEditorUserSettingsIni);
-	//GConfig->SetBool(TEXT("LandscapeEdit"), TEXT("bUseSelectedRegion"), bUseSelectedRegion, GEditorUserSettingsIni);
-	//GConfig->SetBool(TEXT("LandscapeEdit"), TEXT("bUseNegativeMask"), bUseNegativeMask, GEditorUserSettingsIni);
-	GConfig->SetBool(TEXT("LandscapeEdit"), TEXT("bApplyToAllTargets"), bApplyToAllTargets, GEditorUserSettingsIni);
+	GConfig->SetString(TEXT("LandscapeEdit"), TEXT("AlphaTextureName"), *AlphaTextureName, GEditorPerProjectIni);
+	GConfig->SetInt(TEXT("LandscapeEdit"), TEXT("AlphaTextureChannel"), (int32)AlphaTextureChannel, GEditorPerProjectIni);
+
+	GConfig->SetInt(TEXT("LandscapeEdit"), TEXT("FlattenMode"), FlattenMode, GEditorPerProjectIni);
+	GConfig->SetBool(TEXT("LandscapeEdit"), TEXT("bUseSlopeFlatten"), bUseSlopeFlatten, GEditorPerProjectIni);
+	GConfig->SetBool(TEXT("LandscapeEdit"), TEXT("bPickValuePerApply"), bPickValuePerApply, GEditorPerProjectIni);
+	GConfig->SetBool(TEXT("LandscapeEdit"), TEXT("bUseFlattenTarget"), bUseFlattenTarget, GEditorPerProjectIni);
+	GConfig->SetFloat(TEXT("LandscapeEdit"), TEXT("FlattenTarget"), FlattenTarget, GEditorPerProjectIni);
+
+	GConfig->SetFloat(TEXT("LandscapeEdit"), TEXT("RampWidth"), RampWidth, GEditorPerProjectIni);
+	GConfig->SetFloat(TEXT("LandscapeEdit"), TEXT("RampSideFalloff"), RampSideFalloff, GEditorPerProjectIni);
+
+	GConfig->SetInt(TEXT("LandscapeEdit"), TEXT("ErodeThresh"), ErodeThresh, GEditorPerProjectIni);
+	GConfig->SetInt(TEXT("LandscapeEdit"), TEXT("ErodeIterationNum"), ErodeIterationNum, GEditorPerProjectIni);
+	GConfig->SetInt(TEXT("LandscapeEdit"), TEXT("ErodeSurfaceThickness"), ErodeSurfaceThickness, GEditorPerProjectIni);
+	GConfig->SetInt(TEXT("LandscapeEdit"), TEXT("ErosionNoiseMode"), (int32)ErosionNoiseMode, GEditorPerProjectIni);
+	GConfig->SetFloat(TEXT("LandscapeEdit"), TEXT("ErosionNoiseScale"), ErosionNoiseScale, GEditorPerProjectIni);
+
+	GConfig->SetInt(TEXT("LandscapeEdit"), TEXT("RainAmount"), RainAmount, GEditorPerProjectIni);
+	GConfig->SetFloat(TEXT("LandscapeEdit"), TEXT("SedimentCapacity"), SedimentCapacity, GEditorPerProjectIni);
+	GConfig->SetInt(TEXT("LandscapeEdit"), TEXT("HErodeIterationNum"), ErodeIterationNum, GEditorPerProjectIni);
+	GConfig->SetInt(TEXT("LandscapeEdit"), TEXT("RainDistMode"), (int32)RainDistMode, GEditorPerProjectIni);
+	GConfig->SetFloat(TEXT("LandscapeEdit"), TEXT("RainDistScale"), RainDistScale, GEditorPerProjectIni);
+	GConfig->SetFloat(TEXT("LandscapeEdit"), TEXT("HErosionDetailScale"), HErosionDetailScale, GEditorPerProjectIni);
+	GConfig->SetBool(TEXT("LandscapeEdit"), TEXT("bHErosionDetailSmooth"), bHErosionDetailSmooth, GEditorPerProjectIni);
+
+	GConfig->SetInt(TEXT("LandscapeEdit"), TEXT("NoiseMode"), (int32)NoiseMode, GEditorPerProjectIni);
+	GConfig->SetFloat(TEXT("LandscapeEdit"), TEXT("NoiseScale"), NoiseScale, GEditorPerProjectIni);
+	GConfig->SetInt(TEXT("LandscapeEdit"), TEXT("SmoothFilterKernelSize"), SmoothFilterKernelSize, GEditorPerProjectIni);
+	GConfig->SetFloat(TEXT("LandscapeEdit"), TEXT("DetailScale"), DetailScale, GEditorPerProjectIni);
+	GConfig->SetBool(TEXT("LandscapeEdit"), TEXT("bDetailSmooth"), bDetailSmooth, GEditorPerProjectIni);
+
+	GConfig->SetFloat(TEXT("LandscapeEdit"), TEXT("MaximumValueRadius"), MaximumValueRadius, GEditorPerProjectIni);
+
+	GConfig->SetBool(TEXT("LandscapeEdit"), TEXT("bSmoothGizmoBrush"), bSmoothGizmoBrush, GEditorPerProjectIni);
+	GConfig->SetInt(TEXT("LandscapeEdit"), TEXT("PasteMode"), (int32)PasteMode, GEditorPerProjectIni);
+	GConfig->SetInt(TEXT("LandscapeEdit"), TEXT("ConvertMode"), (int32)ResizeLandscape_ConvertMode, GEditorPerProjectIni);
+	//GConfig->SetBool(TEXT("LandscapeEdit"), TEXT("bUseSelectedRegion"), bUseSelectedRegion, GEditorPerProjectIni);
+	//GConfig->SetBool(TEXT("LandscapeEdit"), TEXT("bUseNegativeMask"), bUseNegativeMask, GEditorPerProjectIni);
+	GConfig->SetBool(TEXT("LandscapeEdit"), TEXT("bApplyToAllTargets"), bApplyToAllTargets, GEditorPerProjectIni);
 
 	const FString NewLandscapeMaterialName = (NewLandscape_Material != NULL) ? NewLandscape_Material->GetPathName() : FString();
-	GConfig->SetString(TEXT("LandscapeEdit"), TEXT("NewLandscapeMaterialName"), *NewLandscapeMaterialName, GEditorUserSettingsIni);
+	GConfig->SetString(TEXT("LandscapeEdit"), TEXT("NewLandscapeMaterialName"), *NewLandscapeMaterialName, GEditorPerProjectIni);
 }
 
 // Region
diff --git a/Engine/Source/Editor/LevelEditor/Private/LevelEditor.cpp b/Engine/Source/Editor/LevelEditor/Private/LevelEditor.cpp
--- a/Engine/Source/Editor/LevelEditor/Private/LevelEditor.cpp
+++ b/Engine/Source/Editor/LevelEditor/Private/LevelEditor.cpp
@@ -50,7 +50,7 @@ public:
 	void Construct(const FArguments& InArgs)
 	{
 		FString OptionalBranchPrefix;
-		GConfig->GetString(TEXT("LevelEditor"), TEXT("ProjectNameWatermarkPrefix"), /*out*/ OptionalBranchPrefix, GEditorUserSettingsIni);
+		GConfig->GetString(TEXT("LevelEditor"), TEXT("ProjectNameWatermarkPrefix"), /*out*/ OptionalBranchPrefix, GEditorPerProjectIni);
 
 		FFormatNamedArguments Args;
 		Args.Add(TEXT("Branch"), FText::FromString(OptionalBranchPrefix));
diff --git a/Engine/Source/Editor/LevelEditor/Private/LevelEditorActions.cpp b/Engine/Source/Editor/LevelEditor/Private/LevelEditorActions.cpp
--- a/Engine/Source/Editor/LevelEditor/Private/LevelEditorActions.cpp
+++ b/Engine/Source/Editor/LevelEditor/Private/LevelEditorActions.cpp
@@ -458,8 +458,9 @@ void FLevelEditorActionCallbacks::AttachToSocketSelection(const FName SocketName
 
 void FLevelEditorActionCallbacks::SetMaterialQualityLevel( EMaterialQualityLevel::Type NewQualityLevel )
 {
-	GUnrealEd->AccessEditorUserSettings().MaterialQualityLevel = NewQualityLevel;
-	GUnrealEd->AccessEditorUserSettings().PostEditChange();
+	auto* Settings = GetMutableDefault<UEditorPerProjectUserSettings>();
+	Settings->MaterialQualityLevel = NewQualityLevel;
+	Settings->PostEditChange();
 
 	//Ensure the material quality cvar is also set.
 	static IConsoleVariable* MaterialQualityLevelVar = IConsoleManager::Get().FindConsoleVariable(TEXT("r.MaterialQualityLevel"));
@@ -502,10 +503,10 @@ bool FLevelEditorActionCallbacks::IsFeatureLevelPreviewChecked(ERHIFeatureLevel:
 
 void FLevelEditorActionCallbacks::ConfigureLightingBuildOptions( const FLightingBuildOptions& Options )
 {
-	GConfig->SetBool( TEXT("LightingBuildOptions"), TEXT("OnlyBuildSelected"),		Options.bOnlyBuildSelected,			GEditorUserSettingsIni );
-	GConfig->SetBool( TEXT("LightingBuildOptions"), TEXT("OnlyBuildCurrentLevel"),	Options.bOnlyBuildCurrentLevel,		GEditorUserSettingsIni );
-	GConfig->SetBool( TEXT("LightingBuildOptions"), TEXT("OnlyBuildSelectedLevels"),Options.bOnlyBuildSelectedLevels,	GEditorUserSettingsIni );
-	GConfig->SetBool( TEXT("LightingBuildOptions"), TEXT("OnlyBuildVisibility"),	Options.bOnlyBuildVisibility,		GEditorUserSettingsIni );
+	GConfig->SetBool( TEXT("LightingBuildOptions"), TEXT("OnlyBuildSelected"),		Options.bOnlyBuildSelected,			GEditorPerProjectIni );
+	GConfig->SetBool( TEXT("LightingBuildOptions"), TEXT("OnlyBuildCurrentLevel"),	Options.bOnlyBuildCurrentLevel,		GEditorPerProjectIni );
+	GConfig->SetBool( TEXT("LightingBuildOptions"), TEXT("OnlyBuildSelectedLevels"),Options.bOnlyBuildSelectedLevels,	GEditorPerProjectIni );
+	GConfig->SetBool( TEXT("LightingBuildOptions"), TEXT("OnlyBuildVisibility"),	Options.bOnlyBuildVisibility,		GEditorPerProjectIni );
 }
 
 
@@ -566,7 +567,7 @@ void FLevelEditorActionCallbacks::BuildLightingOnly_VisibilityOnly_Execute()
 bool FLevelEditorActionCallbacks::LightingBuildOptions_UseErrorColoring_IsChecked()
 {
 	bool bUseErrorColoring = false;
-	GConfig->GetBool(TEXT("LightingBuildOptions"), TEXT("UseErrorColoring"), bUseErrorColoring, GEditorUserSettingsIni);
+	GConfig->GetBool(TEXT("LightingBuildOptions"), TEXT("UseErrorColoring"), bUseErrorColoring, GEditorPerProjectIni);
 	return bUseErrorColoring;
 }
 
@@ -574,15 +575,15 @@ bool FLevelEditorActionCallbacks::LightingBuildOptions_UseErrorColoring_IsChecke
 void FLevelEditorActionCallbacks::LightingBuildOptions_UseErrorColoring_Toggled()
 {
 	bool bUseErrorColoring = false;
-	GConfig->GetBool( TEXT("LightingBuildOptions"), TEXT("UseErrorColoring"), bUseErrorColoring, GEditorUserSettingsIni );
-	GConfig->SetBool( TEXT("LightingBuildOptions"), TEXT("UseErrorColoring"), !bUseErrorColoring, GEditorUserSettingsIni );
+	GConfig->GetBool( TEXT("LightingBuildOptions"), TEXT("UseErrorColoring"), bUseErrorColoring, GEditorPerProjectIni );
+	GConfig->SetBool( TEXT("LightingBuildOptions"), TEXT("UseErrorColoring"), !bUseErrorColoring, GEditorPerProjectIni );
 }
 
 
 bool FLevelEditorActionCallbacks::LightingBuildOptions_ShowLightingStats_IsChecked()
 {
 	bool bShowLightingBuildInfo = false;
-	GConfig->GetBool(TEXT("LightingBuildOptions"), TEXT("ShowLightingBuildInfo"), bShowLightingBuildInfo, GEditorUserSettingsIni);
+	GConfig->GetBool(TEXT("LightingBuildOptions"), TEXT("ShowLightingBuildInfo"), bShowLightingBuildInfo, GEditorPerProjectIni);
 	return bShowLightingBuildInfo;
 }
 
@@ -590,8 +591,8 @@ bool FLevelEditorActionCallbacks::LightingBuildOptions_ShowLightingStats_IsCheck
 void FLevelEditorActionCallbacks::LightingBuildOptions_ShowLightingStats_Toggled()
 {
 	bool bShowLightingBuildInfo = false;
-	GConfig->GetBool( TEXT("LightingBuildOptions"), TEXT("ShowLightingBuildInfo"), bShowLightingBuildInfo, GEditorUserSettingsIni );
-	GConfig->SetBool( TEXT("LightingBuildOptions"), TEXT("ShowLightingBuildInfo"), !bShowLightingBuildInfo, GEditorUserSettingsIni );
+	GConfig->GetBool( TEXT("LightingBuildOptions"), TEXT("ShowLightingBuildInfo"), bShowLightingBuildInfo, GEditorPerProjectIni );
+	GConfig->SetBool( TEXT("LightingBuildOptions"), TEXT("ShowLightingBuildInfo"), !bShowLightingBuildInfo, GEditorPerProjectIni );
 }
 
 
@@ -630,13 +631,13 @@ void FLevelEditorActionCallbacks::BuildLODsOnly_Execute()
 bool FLevelEditorActionCallbacks::IsLightingQualityChecked( ELightingBuildQuality TestQuality )
 {
 	int32 CurrentQualityLevel;
-	GConfig->GetInt(TEXT("LightingBuildOptions"), TEXT("QualityLevel"), CurrentQualityLevel, GEditorUserSettingsIni);
+	GConfig->GetInt(TEXT("LightingBuildOptions"), TEXT("QualityLevel"), CurrentQualityLevel, GEditorPerProjectIni);
 	return TestQuality == CurrentQualityLevel;
 }
 
 void FLevelEditorActionCallbacks::SetLightingQuality( ELightingBuildQuality NewQuality )
 {
-	GConfig->SetInt(TEXT("LightingBuildOptions"), TEXT("QualityLevel"), (int32)NewQuality, GEditorUserSettingsIni);
+	GConfig->SetInt(TEXT("LightingBuildOptions"), TEXT("QualityLevel"), (int32)NewQuality, GEditorPerProjectIni);
 }
 
 void FLevelEditorActionCallbacks::SetLightingToolShowBounds()
@@ -2146,9 +2147,11 @@ bool FLevelEditorActionCallbacks::OnGetTransformWidgetVisibility()
 
 void FLevelEditorActionCallbacks::OnAllowTranslucentSelection()
 {
+	auto* Settings = GetMutableDefault<UEditorPerProjectUserSettings>();
+
 	// Toggle 'allow select translucent'
-	GEditor->AccessEditorUserSettings().bAllowSelectTranslucent = !GEditor->AccessEditorUserSettings().bAllowSelectTranslucent;
-	GEditor->AccessEditorUserSettings().PostEditChange();
+	Settings->bAllowSelectTranslucent = !Settings->bAllowSelectTranslucent;
+	Settings->PostEditChange();
 
 	// Need to refresh hit proxies as we changed what should be rendered into them
 	GUnrealEd->RedrawAllViewports();
@@ -2156,7 +2159,7 @@ void FLevelEditorActionCallbacks::OnAllowTranslucentSelection()
 
 bool FLevelEditorActionCallbacks::OnIsAllowTranslucentSelectionEnabled()
 {
-	return GEditor->GetEditorUserSettings().bAllowSelectTranslucent == true;
+	return GetDefault<UEditorPerProjectUserSettings>()->bAllowSelectTranslucent == true;
 }
 
 void FLevelEditorActionCallbacks::OnAllowGroupSelection()
diff --git a/Engine/Source/Editor/LevelEditor/Private/LevelViewportLayout.cpp b/Engine/Source/Editor/LevelEditor/Private/LevelViewportLayout.cpp
--- a/Engine/Source/Editor/LevelEditor/Private/LevelViewportLayout.cpp
+++ b/Engine/Source/Editor/LevelEditor/Private/LevelViewportLayout.cpp
@@ -236,9 +236,9 @@ void FLevelViewportLayout::InitCommonLayoutFromString( const FString& SpecificLa
 	{
 		const FString& IniSection = FLayoutSaveRestore::GetAdditionalLayoutConfigIni();
 
-		GConfig->GetBool(*IniSection, *(SpecificLayoutString + TEXT(".bIsMaximized")), bShouldBeMaximized, GEditorUserSettingsIni);
-		GConfig->GetBool(*IniSection, *(SpecificLayoutString + TEXT(".bIsImmersive")), bShouldBeImmersive, GEditorUserSettingsIni);
-		GConfig->GetInt(*IniSection, *(SpecificLayoutString + TEXT(".MaximizedViewportID")), MaximizedViewportID, GEditorUserSettingsIni);
+		GConfig->GetBool(*IniSection, *(SpecificLayoutString + TEXT(".bIsMaximized")), bShouldBeMaximized, GEditorPerProjectIni);
+		GConfig->GetBool(*IniSection, *(SpecificLayoutString + TEXT(".bIsImmersive")), bShouldBeImmersive, GEditorPerProjectIni);
+		GConfig->GetInt(*IniSection, *(SpecificLayoutString + TEXT(".MaximizedViewportID")), MaximizedViewportID, GEditorPerProjectIni);
 	}
 	// Replacement layouts (those selected by the user via a command) don't start maximized so the layout can be seen clearly.
 	if (!bIsReplacement && bIsMaximizeSupported && MaximizedViewportID >= 0 && MaximizedViewportID < Viewports.Num() && (bShouldBeMaximized || bShouldBeImmersive))
@@ -268,9 +268,9 @@ void FLevelViewportLayout::SaveCommonLayoutString( const FString& SpecificLayout
 		}
 	}
 
-	GConfig->SetBool(*IniSection, *(SpecificLayoutString + TEXT(".bIsMaximized")), bIsMaximizeSupported && bIsMaximized, GEditorUserSettingsIni);
-	GConfig->SetBool(*IniSection, *(SpecificLayoutString + TEXT(".bIsImmersive")), bIsImmersive, GEditorUserSettingsIni);
-	GConfig->SetInt(*IniSection, *(SpecificLayoutString + TEXT(".MaximizedViewportID")), MaximizedViewportID, GEditorUserSettingsIni);
+	GConfig->SetBool(*IniSection, *(SpecificLayoutString + TEXT(".bIsMaximized")), bIsMaximizeSupported && bIsMaximized, GEditorPerProjectIni);
+	GConfig->SetBool(*IniSection, *(SpecificLayoutString + TEXT(".bIsImmersive")), bIsImmersive, GEditorPerProjectIni);
+	GConfig->SetInt(*IniSection, *(SpecificLayoutString + TEXT(".MaximizedViewportID")), MaximizedViewportID, GEditorPerProjectIni);
 }
 
 void FLevelViewportLayout::RequestMaximizeViewport( TSharedRef<class SLevelViewport> ViewportToMaximize, const bool bWantMaximize, const bool bWantImmersive, const bool bAllowAnimation )
diff --git a/Engine/Source/Editor/LevelEditor/Private/LevelViewportLayout2x2.cpp b/Engine/Source/Editor/LevelEditor/Private/LevelViewportLayout2x2.cpp
--- a/Engine/Source/Editor/LevelEditor/Private/LevelViewportLayout2x2.cpp
+++ b/Engine/Source/Editor/LevelEditor/Private/LevelViewportLayout2x2.cpp
@@ -38,7 +38,7 @@ TSharedRef<SWidget> FLevelViewportLayout2x2::MakeViewportLayout(const FString& L
 		{
 			FString PercentageString;
 			FVector2D NewPercentage = ViewportLayout2x2Defs::DefaultSplitterPercentages;
-			if(GConfig->GetString(*IniSection, *(SpecificLayoutString + FString::Printf(TEXT(".Percentages%i"), i)), PercentageString, GEditorUserSettingsIni))
+			if(GConfig->GetString(*IniSection, *(SpecificLayoutString + FString::Printf(TEXT(".Percentages%i"), i)), PercentageString, GEditorPerProjectIni))
 			{
 				NewPercentage.InitFromString(PercentageString);
 			}						
@@ -120,7 +120,7 @@ void FLevelViewportLayout2x2::SaveLayoutString(const FString& LayoutString) cons
 		SplitterWidget->GetSplitterPercentages(Percentages);
 		for (int32 i = 0; i < Percentages.Num(); ++i)
 		{
-			GConfig->SetString(*IniSection, *(SpecificLayoutString + FString::Printf(TEXT(".Percentages%i"), i)), *Percentages[i].ToString(), GEditorUserSettingsIni);
+			GConfig->SetString(*IniSection, *(SpecificLayoutString + FString::Printf(TEXT(".Percentages%i"), i)), *Percentages[i].ToString(), GEditorPerProjectIni);
 		}
 
 		SaveCommonLayoutString(SpecificLayoutString);
diff --git a/Engine/Source/Editor/LevelEditor/Private/LevelViewportLayoutFourPanes.cpp b/Engine/Source/Editor/LevelEditor/Private/LevelViewportLayoutFourPanes.cpp
--- a/Engine/Source/Editor/LevelEditor/Private/LevelViewportLayoutFourPanes.cpp
+++ b/Engine/Source/Editor/LevelEditor/Private/LevelViewportLayoutFourPanes.cpp
@@ -36,15 +36,15 @@ TSharedRef<SWidget> FLevelViewportLayoutFourPanes::MakeViewportLayout(const FStr
 		ViewportKey3 = SpecificLayoutString + TEXT(".Viewport3");
 
 		FString PercentageString;
-		if(GConfig->GetString(*IniSection, *(SpecificLayoutString + TEXT(".Percentage0")), PercentageString, GEditorUserSettingsIni))
+		if(GConfig->GetString(*IniSection, *(SpecificLayoutString + TEXT(".Percentage0")), PercentageString, GEditorPerProjectIni))
 		{
 			TTypeFromString<float>::FromString(PrimarySplitterPercentage, *PercentageString);
 		}		
-		if(GConfig->GetString(*IniSection, *(SpecificLayoutString + TEXT(".Percentage1")), PercentageString, GEditorUserSettingsIni))
+		if(GConfig->GetString(*IniSection, *(SpecificLayoutString + TEXT(".Percentage1")), PercentageString, GEditorPerProjectIni))
 		{
 			TTypeFromString<float>::FromString(SecondarySplitterPercentage0, *PercentageString);
 		}	
-		if(GConfig->GetString(*IniSection, *(SpecificLayoutString + TEXT(".Percentage2")), PercentageString, GEditorUserSettingsIni))
+		if(GConfig->GetString(*IniSection, *(SpecificLayoutString + TEXT(".Percentage2")), PercentageString, GEditorPerProjectIni))
 		{
 			TTypeFromString<float>::FromString(SecondarySplitterPercentage1, *PercentageString);
 		}	
@@ -74,9 +74,9 @@ void FLevelViewportLayoutFourPanes::SaveLayoutString(const FString& LayoutString
 		float SecondaryPercentage0 = SecondarySplitterWidget->SlotAt(0).SizeValue.Get();
 		float SecondaryPercentage1 = SecondarySplitterWidget->SlotAt(1).SizeValue.Get();
 
-		GConfig->SetString(*IniSection, *(SpecificLayoutString + TEXT(".Percentage0")), *TTypeToString<float>::ToString(PrimaryPercentage), GEditorUserSettingsIni);
-		GConfig->SetString(*IniSection, *(SpecificLayoutString + TEXT(".Percentage1")), *TTypeToString<float>::ToString(SecondaryPercentage0), GEditorUserSettingsIni);
-		GConfig->SetString(*IniSection, *(SpecificLayoutString + TEXT(".Percentage2")), *TTypeToString<float>::ToString(SecondaryPercentage1), GEditorUserSettingsIni);
+		GConfig->SetString(*IniSection, *(SpecificLayoutString + TEXT(".Percentage0")), *TTypeToString<float>::ToString(PrimaryPercentage), GEditorPerProjectIni);
+		GConfig->SetString(*IniSection, *(SpecificLayoutString + TEXT(".Percentage1")), *TTypeToString<float>::ToString(SecondaryPercentage0), GEditorPerProjectIni);
+		GConfig->SetString(*IniSection, *(SpecificLayoutString + TEXT(".Percentage2")), *TTypeToString<float>::ToString(SecondaryPercentage1), GEditorPerProjectIni);
 
 		SaveCommonLayoutString(SpecificLayoutString);
 	}
diff --git a/Engine/Source/Editor/LevelEditor/Private/LevelViewportLayoutThreePanes.cpp b/Engine/Source/Editor/LevelEditor/Private/LevelViewportLayoutThreePanes.cpp
--- a/Engine/Source/Editor/LevelEditor/Private/LevelViewportLayoutThreePanes.cpp
+++ b/Engine/Source/Editor/LevelEditor/Private/LevelViewportLayoutThreePanes.cpp
@@ -37,11 +37,11 @@ TSharedRef<SWidget> FLevelViewportLayoutThreePanes::MakeViewportLayout(const FSt
 		ViewportKey2 = SpecificLayoutString + TEXT(".Viewport2");
 
 		FString PercentageString;
-		if(GConfig->GetString(*IniSection, *(SpecificLayoutString + TEXT(".Percentage0")), PercentageString, GEditorUserSettingsIni))
+		if(GConfig->GetString(*IniSection, *(SpecificLayoutString + TEXT(".Percentage0")), PercentageString, GEditorPerProjectIni))
 		{
 			TTypeFromString<float>::FromString(PrimarySplitterPercentage, *PercentageString);
 		}		
-		if(GConfig->GetString(*IniSection, *(SpecificLayoutString + TEXT(".Percentage1")), PercentageString, GEditorUserSettingsIni))
+		if(GConfig->GetString(*IniSection, *(SpecificLayoutString + TEXT(".Percentage1")), PercentageString, GEditorPerProjectIni))
 		{
 			TTypeFromString<float>::FromString(SecondarySplitterPercentage, *PercentageString);
 		}	
@@ -70,8 +70,8 @@ void FLevelViewportLayoutThreePanes::SaveLayoutString(const FString& LayoutStrin
 		check (SecondarySplitterWidget->GetChildren()->Num() == 2);
 		float SecondaryPercentage = SecondarySplitterWidget->SlotAt(0).SizeValue.Get();
 
-		GConfig->SetString(*IniSection, *(SpecificLayoutString + TEXT(".Percentage0")), *TTypeToString<float>::ToString(PrimaryPercentage), GEditorUserSettingsIni);
-		GConfig->SetString(*IniSection, *(SpecificLayoutString + TEXT(".Percentage1")), *TTypeToString<float>::ToString(SecondaryPercentage), GEditorUserSettingsIni);
+		GConfig->SetString(*IniSection, *(SpecificLayoutString + TEXT(".Percentage0")), *TTypeToString<float>::ToString(PrimaryPercentage), GEditorPerProjectIni);
+		GConfig->SetString(*IniSection, *(SpecificLayoutString + TEXT(".Percentage1")), *TTypeToString<float>::ToString(SecondaryPercentage), GEditorPerProjectIni);
 
 		SaveCommonLayoutString(SpecificLayoutString);
 	}
diff --git a/Engine/Source/Editor/LevelEditor/Private/LevelViewportLayoutTwoPanes.cpp b/Engine/Source/Editor/LevelEditor/Private/LevelViewportLayoutTwoPanes.cpp
--- a/Engine/Source/Editor/LevelEditor/Private/LevelViewportLayoutTwoPanes.cpp
+++ b/Engine/Source/Editor/LevelEditor/Private/LevelViewportLayoutTwoPanes.cpp
@@ -40,7 +40,7 @@ TSharedRef<SWidget> TLevelViewportLayoutTwoPanes<TOrientation>::MakeViewportLayo
 
 
 		FString PercentageString;
-		if(GConfig->GetString(*IniSection, *(SpecificLayoutString + TEXT(".Percentage")), PercentageString, GEditorUserSettingsIni))
+		if(GConfig->GetString(*IniSection, *(SpecificLayoutString + TEXT(".Percentage")), PercentageString, GEditorPerProjectIni))
 		{
 			TTypeFromString<float>::FromString(SplitterPercentage, *PercentageString);
 		}						
@@ -92,7 +92,7 @@ void TLevelViewportLayoutTwoPanes<TOrientation>::SaveLayoutString(const FString&
 		check (SplitterWidget->GetChildren()->Num() == 2);
 		float Percentage = SplitterWidget->SlotAt(0).SizeValue.Get();
 
-		GConfig->SetString(*IniSection, *(SpecificLayoutString + TEXT(".Percentage")), *TTypeToString<float>::ToString(Percentage), GEditorUserSettingsIni);
+		GConfig->SetString(*IniSection, *(SpecificLayoutString + TEXT(".Percentage")), *TTypeToString<float>::ToString(Percentage), GEditorPerProjectIni);
 
 		SaveCommonLayoutString(SpecificLayoutString);
 	}
diff --git a/Engine/Source/Editor/LevelEditor/Private/LevelViewportTabContent.cpp b/Engine/Source/Editor/LevelEditor/Private/LevelViewportTabContent.cpp
--- a/Engine/Source/Editor/LevelEditor/Private/LevelViewportTabContent.cpp
+++ b/Engine/Source/Editor/LevelEditor/Private/LevelViewportTabContent.cpp
@@ -48,7 +48,7 @@ void FLevelViewportTabContent::Initialize(TSharedPtr<class SLevelEditor> InParen
 
 	FString LayoutTypeString;
 	if(LayoutString.IsEmpty() ||
-		!GConfig->GetString(*IniSection, *(InLayoutString + TEXT(".LayoutType")), LayoutTypeString, GEditorUserSettingsIni))
+		!GConfig->GetString(*IniSection, *(InLayoutString + TEXT(".LayoutType")), LayoutTypeString, GEditorPerProjectIni))
 	{
 		LayoutTypeString = LevelViewportConfigurationNames::FourPanes2x2.ToString();
 	}
@@ -123,7 +123,7 @@ void FLevelViewportTabContent::SaveLayoutString(const FString InLayoutString) co
 		FString LayoutTypeString = ActiveLevelViewportLayout->GetLayoutTypeName().ToString();
 
 		const FString& IniSection = FLayoutSaveRestore::GetAdditionalLayoutConfigIni();
-		GConfig->SetString(*IniSection, *(InLayoutString + TEXT(".LayoutType")), *LayoutTypeString, GEditorUserSettingsIni);
+		GConfig->SetString(*IniSection, *(InLayoutString + TEXT(".LayoutType")), *LayoutTypeString, GEditorPerProjectIni);
 
 		ActiveLevelViewportLayout->SaveLayoutString(InLayoutString);
 	}
diff --git a/Engine/Source/Editor/LevelEditor/Private/SLevelEditor.cpp b/Engine/Source/Editor/LevelEditor/Private/SLevelEditor.cpp
--- a/Engine/Source/Editor/LevelEditor/Private/SLevelEditor.cpp
+++ b/Engine/Source/Editor/LevelEditor/Private/SLevelEditor.cpp
@@ -266,7 +266,7 @@ SLevelEditor::~SLevelEditor()
 	LevelEditorModule.OnNotificationBarChanged().RemoveAll( this );
 	
 	GetMutableDefault<UEditorExperimentalSettings>()->OnSettingChanged().RemoveAll( this );
-	GEditor->AccessEditorUserSettings().OnUserSettingChanged().RemoveAll( this );
+	GetMutableDefault<UEditorPerProjectUserSettings>()->OnUserSettingChanged().RemoveAll( this );
 	FEditorModeRegistry::Get().OnRegisteredModesChanged().RemoveAll( this );
 
 	FEditorDelegates::MapChange.RemoveAll(this);
diff --git a/Engine/Source/Editor/LevelEditor/Private/SLevelEditorModeContent.cpp b/Engine/Source/Editor/LevelEditor/Private/SLevelEditorModeContent.cpp
--- a/Engine/Source/Editor/LevelEditor/Private/SLevelEditorModeContent.cpp
+++ b/Engine/Source/Editor/LevelEditor/Private/SLevelEditorModeContent.cpp
@@ -15,7 +15,7 @@
 SLevelEditorModeContent::~SLevelEditorModeContent()
 {
 	GLevelEditorModeTools().OnEditorModeChanged().RemoveAll( this );
-	GEditor->AccessEditorUserSettings().OnUserSettingChanged().RemoveAll( this );
+	GetMutableDefault<UEditorPerProjectUserSettings>()->OnUserSettingChanged().RemoveAll( this );
 }
 
 void SLevelEditorModeContent::Construct( const FArguments& InArgs, const TSharedRef< class ILevelEditor >& InOwningLevelEditor, const TSharedRef< class SDockTab >& InOwningDocTab, FEdMode* InEditorMode )
@@ -26,7 +26,7 @@ void SLevelEditorModeContent::Construct( const FArguments& InArgs, const TShared
 
 	InOwningDocTab->SetOnTabClosed( SDockTab::FOnTabClosedCallback::CreateSP(this, &SLevelEditorModeContent::HandleParentClosed ) );
 	GLevelEditorModeTools().OnEditorModeChanged().AddSP( this, &SLevelEditorModeContent::HandleEditorModeChanged );
-	GEditor->AccessEditorUserSettings().OnUserSettingChanged().AddSP( this, &SLevelEditorModeContent::HandleUserSettingsChange );
+	GetMutableDefault<UEditorPerProjectUserSettings>()->OnUserSettingChanged().AddSP( this, &SLevelEditorModeContent::HandleUserSettingsChange );
 
 	ChildSlot
 	[
diff --git a/Engine/Source/Editor/LevelEditor/Private/SLevelEditorToolBox.cpp b/Engine/Source/Editor/LevelEditor/Private/SLevelEditorToolBox.cpp
--- a/Engine/Source/Editor/LevelEditor/Private/SLevelEditorToolBox.cpp
+++ b/Engine/Source/Editor/LevelEditor/Private/SLevelEditorToolBox.cpp
@@ -13,7 +13,7 @@
 
 SLevelEditorToolBox::~SLevelEditorToolBox()
 {
-	GEditor->AccessEditorUserSettings().OnUserSettingChanged().RemoveAll( this );
+	GetMutableDefault<UEditorPerProjectUserSettings>()->OnUserSettingChanged().RemoveAll( this );
 }
 
 void SLevelEditorToolBox::Construct( const FArguments& InArgs, const TSharedRef< class ILevelEditor >& OwningLevelEditor )
@@ -23,7 +23,7 @@ void SLevelEditorToolBox::Construct( const FArguments& InArgs, const TSharedRef<
 	// Important: We use a raw binding here because we are releasing our binding in our destructor (where a weak pointer would be invalid)
 	// It's imperative that our delegate is removed in the destructor for the level editor module to play nicely with reloading.
 
-	GEditor->AccessEditorUserSettings().OnUserSettingChanged().AddRaw( this, &SLevelEditorToolBox::HandleUserSettingsChange );
+	GetMutableDefault<UEditorPerProjectUserSettings>()->OnUserSettingChanged().AddRaw( this, &SLevelEditorToolBox::HandleUserSettingsChange );
 
 	ChildSlot
 	[
diff --git a/Engine/Source/Editor/LevelEditor/Private/SLevelViewport.cpp b/Engine/Source/Editor/LevelEditor/Private/SLevelViewport.cpp
--- a/Engine/Source/Editor/LevelEditor/Private/SLevelViewport.cpp
+++ b/Engine/Source/Editor/LevelEditor/Private/SLevelViewport.cpp
@@ -1944,7 +1944,7 @@ FLevelEditorViewportInstanceSettings SLevelViewport::LoadLegacyConfigFromIni(con
 
 	{
 		int32 ViewportTypeAsInt = ViewportInstanceSettings.ViewportType;
-		GConfig->GetInt(*IniSection, *(InConfigKey + TEXT(".Type")), ViewportTypeAsInt, GEditorUserSettingsIni);
+		GConfig->GetInt(*IniSection, *(InConfigKey + TEXT(".Type")), ViewportTypeAsInt, GEditorPerProjectIni);
 		ViewportInstanceSettings.ViewportType = (ViewportTypeAsInt == -1 || ViewportTypeAsInt == 255) ? LVT_None : static_cast<ELevelViewportType>(ViewportTypeAsInt); // LVT_None used to be -1 or 255
 
 		if(ViewportInstanceSettings.ViewportType == LVT_None)
@@ -1953,18 +1953,18 @@ FLevelEditorViewportInstanceSettings SLevelViewport::LoadLegacyConfigFromIni(con
 		}
 	}
 	
-	GConfig->GetString(*IniSection, *(InConfigKey + TEXT(".EditorShowFlags")), ViewportInstanceSettings.EditorShowFlagsString, GEditorUserSettingsIni);
-	GConfig->GetString(*IniSection, *(InConfigKey + TEXT(".GameShowFlags")), ViewportInstanceSettings.GameShowFlagsString, GEditorUserSettingsIni);
+	GConfig->GetString(*IniSection, *(InConfigKey + TEXT(".EditorShowFlags")), ViewportInstanceSettings.EditorShowFlagsString, GEditorPerProjectIni);
+	GConfig->GetString(*IniSection, *(InConfigKey + TEXT(".GameShowFlags")), ViewportInstanceSettings.GameShowFlagsString, GEditorPerProjectIni);
 		
 	// A single view mode index has been deprecated in favor of separate perspective and orthographic settings
 	EViewModeIndex LegacyViewModeIndex = VMI_Unknown;
 	{
 		int32 LegacyVMIAsInt = VMI_Unknown;
-		GConfig->GetInt(*IniSection, *(InConfigKey+TEXT(".ViewModeIndex")), LegacyVMIAsInt, GEditorUserSettingsIni);
+		GConfig->GetInt(*IniSection, *(InConfigKey+TEXT(".ViewModeIndex")), LegacyVMIAsInt, GEditorPerProjectIni);
 		LegacyViewModeIndex = (LegacyVMIAsInt == -1) ? VMI_Unknown : static_cast<EViewModeIndex>(LegacyVMIAsInt); // VMI_Unknown used to be -1
 	}
 
-	if(!GConfig->GetInt(*IniSection, *(InConfigKey+TEXT(".PerspViewModeIndex")), (int32&)ViewportInstanceSettings.PerspViewModeIndex, GEditorUserSettingsIni))
+	if(!GConfig->GetInt(*IniSection, *(InConfigKey+TEXT(".PerspViewModeIndex")), (int32&)ViewportInstanceSettings.PerspViewModeIndex, GEditorPerProjectIni))
 	{
 		if(ViewportInstanceSettings.ViewportType == LVT_Perspective)
 		{
@@ -1978,7 +1978,7 @@ FLevelEditorViewportInstanceSettings SLevelViewport::LoadLegacyConfigFromIni(con
 		}
 	}
 
-	if(!GConfig->GetInt(*IniSection, *(InConfigKey+TEXT(".OrthoViewModeIndex")), (int32&)ViewportInstanceSettings.OrthoViewModeIndex, GEditorUserSettingsIni))
+	if(!GConfig->GetInt(*IniSection, *(InConfigKey+TEXT(".OrthoViewModeIndex")), (int32&)ViewportInstanceSettings.OrthoViewModeIndex, GEditorPerProjectIni))
 	{
 		// Default to Brush Wireframe for an orthographic viewport
 		ViewportInstanceSettings.OrthoViewModeIndex = (ViewportInstanceSettings.ViewportType != LVT_Perspective && LegacyViewModeIndex != VMI_Unknown) ? LegacyViewModeIndex : VMI_BrushWireframe;
@@ -1986,7 +1986,7 @@ FLevelEditorViewportInstanceSettings SLevelViewport::LoadLegacyConfigFromIni(con
 
 	{
 		FString BufferVisualizationModeString;
-		if(GConfig->GetString(*IniSection, *(InConfigKey+TEXT(".BufferVisualizationMode")), BufferVisualizationModeString, GEditorUserSettingsIni))
+		if(GConfig->GetString(*IniSection, *(InConfigKey+TEXT(".BufferVisualizationMode")), BufferVisualizationModeString, GEditorPerProjectIni))
 		{
 			ViewportInstanceSettings.BufferVisualizationMode = *BufferVisualizationModeString;
 		}
@@ -1994,16 +1994,16 @@ FLevelEditorViewportInstanceSettings SLevelViewport::LoadLegacyConfigFromIni(con
 	
 	{
 		FString ExposureSettingsString;
-		if(GConfig->GetString(*IniSection, *(InConfigKey + TEXT(".ExposureSettings")), ExposureSettingsString, GEditorUserSettingsIni))
+		if(GConfig->GetString(*IniSection, *(InConfigKey + TEXT(".ExposureSettings")), ExposureSettingsString, GEditorPerProjectIni))
 		{
 			ViewportInstanceSettings.ExposureSettings.SetFromString(*ExposureSettingsString);
 		}
 	}
 
-	GConfig->GetBool(*IniSection, *(InConfigKey + TEXT(".bIsRealtime")), ViewportInstanceSettings.bIsRealtime, GEditorUserSettingsIni);
-	GConfig->GetBool(*IniSection, *(InConfigKey + TEXT(".bWantStats")), ViewportInstanceSettings.bShowStats, GEditorUserSettingsIni);
-	GConfig->GetBool(*IniSection, *(InConfigKey + TEXT(".bWantFPS")), ViewportInstanceSettings.bShowFPS_DEPRECATED, GEditorUserSettingsIni);
-	GConfig->GetFloat(*IniSection, *(InConfigKey + TEXT(".FOVAngle")), ViewportInstanceSettings.FOVAngle, GEditorUserSettingsIni);
+	GConfig->GetBool(*IniSection, *(InConfigKey + TEXT(".bIsRealtime")), ViewportInstanceSettings.bIsRealtime, GEditorPerProjectIni);
+	GConfig->GetBool(*IniSection, *(InConfigKey + TEXT(".bWantStats")), ViewportInstanceSettings.bShowStats, GEditorPerProjectIni);
+	GConfig->GetBool(*IniSection, *(InConfigKey + TEXT(".bWantFPS")), ViewportInstanceSettings.bShowFPS_DEPRECATED, GEditorPerProjectIni);
+	GConfig->GetFloat(*IniSection, *(InConfigKey + TEXT(".FOVAngle")), ViewportInstanceSettings.FOVAngle, GEditorPerProjectIni);
 
 	return ViewportInstanceSettings;
 }
diff --git a/Engine/Source/Editor/LevelEditor/Private/SLevelViewportToolBar.cpp b/Engine/Source/Editor/LevelEditor/Private/SLevelViewportToolBar.cpp
--- a/Engine/Source/Editor/LevelEditor/Private/SLevelViewportToolBar.cpp
+++ b/Engine/Source/Editor/LevelEditor/Private/SLevelViewportToolBar.cpp
@@ -608,7 +608,7 @@ TSharedRef<SWidget> SLevelViewportToolBar::GenerateDevicePreviewMenu() const
 	{
 		// Build the menu from the contents of the game ini
 		//@todo This should probably be using GConfig->GetText [10/21/2013 justin.sargent]
-		if ( GConfig->GetString( *INISection, *FString::Printf( TEXT("%s%d"), *INIKeyBase, ItemIdx ), CurItem, GEditorUserSettingsIni ) )
+		if ( GConfig->GetString( *INISection, *FString::Printf( TEXT("%s%d"), *INIKeyBase, ItemIdx ), CurItem, GEditorPerProjectIni ) )
 		{
 			const FName PlatformIcon = UIManager->GetDeviceIconName( CurItem );
 
diff --git a/Engine/Source/Editor/LevelEditor/Private/SSurfaceProperties.cpp b/Engine/Source/Editor/LevelEditor/Private/SSurfaceProperties.cpp
--- a/Engine/Source/Editor/LevelEditor/Private/SSurfaceProperties.cpp
+++ b/Engine/Source/Editor/LevelEditor/Private/SSurfaceProperties.cpp
@@ -51,8 +51,8 @@ void SSurfaceProperties::Construct( const FArguments& InArgs )
 	bPreserveScaleRatio = false;
 	bUseRelativeScaling = false;
 
-	GConfig->GetBool(TEXT("SelectionDetails"), TEXT("PreserveScaleRatio"), bPreserveScaleRatio, GEditorUserSettingsIni);
-	GConfig->GetBool(TEXT("SelectionDetails"), TEXT("UseRelativeScaling"), bUseRelativeScaling, GEditorUserSettingsIni);
+	GConfig->GetBool(TEXT("SelectionDetails"), TEXT("PreserveScaleRatio"), bPreserveScaleRatio, GEditorPerProjectIni);
+	GConfig->GetBool(TEXT("SelectionDetails"), TEXT("UseRelativeScaling"), bUseRelativeScaling, GEditorPerProjectIni);
 
 	static const float ScalingValues[] = { 1.0f / 16, 1.0f / 8, 1.0f / 4, 1.0f / 2, 1, 2, 4, 8, 16 };
 	for(int Idx = 0; Idx < ARRAY_COUNT(ScalingValues); Idx++)
@@ -880,7 +880,7 @@ void SSurfaceProperties::OnCustomPanValueCommitted( int32 NewValue, ETextCommit:
 void SSurfaceProperties::OnScaleLabelClicked( )
 {
 	bUseRelativeScaling = !bUseRelativeScaling;
-	GConfig->SetBool(TEXT("SurfaceSelection"), TEXT("UseRelativeScaling"), bUseRelativeScaling, GEditorUserSettingsIni);
+	GConfig->SetBool(TEXT("SurfaceSelection"), TEXT("UseRelativeScaling"), bUseRelativeScaling, GEditorPerProjectIni);
 }
 
 FText SSurfaceProperties::GetScalingLabel() const
@@ -949,7 +949,7 @@ void SSurfaceProperties::OnPreserveScaleRatioToggled( ECheckBoxState NewState )
 {
 	bPreserveScaleRatio = (NewState == ECheckBoxState::Checked) ? true : false;
 	CachedScalingValueV = CachedScalingValueU;
-	GConfig->SetBool(TEXT("SurfaceSelection"), TEXT("PreserveScaleRatio"), bPreserveScaleRatio, GEditorUserSettingsIni);
+	GConfig->SetBool(TEXT("SurfaceSelection"), TEXT("PreserveScaleRatio"), bPreserveScaleRatio, GEditorPerProjectIni);
 }
 
 void SSurfaceProperties::OnCustomRotateValueCommitted(int32 NewValue, ETextCommit::Type CommitInfo)
diff --git a/Engine/Source/Editor/MainFrame/Private/Frame/MainFrameActions.cpp b/Engine/Source/Editor/MainFrame/Private/Frame/MainFrameActions.cpp
--- a/Engine/Source/Editor/MainFrame/Private/Frame/MainFrameActions.cpp
+++ b/Engine/Source/Editor/MainFrame/Private/Frame/MainFrameActions.cpp
@@ -11,6 +11,7 @@
 #include "SourceCodeNavigation.h"
 #include "SOutputLogDialog.h"
 
+#include "Settings/EditorSettings.h"
 
 #define LOCTEXT_NAMESPACE "MainFrameActions"
 
@@ -228,7 +229,7 @@ void FMainFrameActionCallbacks::CacheProjectNames()
 	ProjectNames.Empty();
 
 	// The switch project menu is filled with recently opened project files
-	ProjectNames = GEditor->GetGameAgnosticSettings().RecentlyOpenedProjectFiles;
+	ProjectNames = GetDefault<UEditorSettings>()->RecentlyOpenedProjectFiles;
 }
 
 void FMainFrameActionCallbacks::NewProject( bool bAllowProjectOpening, bool bAllowProjectCreate )
@@ -716,7 +717,7 @@ void FMainFrameActionCallbacks::ResetLayout()
 	}
 
 	// make a backup
-	GEditor->SaveEditorUserSettings();
+	GetMutableDefault<UEditorPerProjectUserSettings>()->SaveConfig();
 
 	FString BackupEditorLayoutIni = FString::Printf(TEXT("%s_Backup.ini"), *FPaths::GetBaseFilename(GEditorLayoutIni, false));
 
diff --git a/Engine/Source/Editor/MainFrame/Private/Frame/MainFrameHandler.cpp b/Engine/Source/Editor/MainFrame/Private/Frame/MainFrameHandler.cpp
--- a/Engine/Source/Editor/MainFrame/Private/Frame/MainFrameHandler.cpp
+++ b/Engine/Source/Editor/MainFrame/Private/Frame/MainFrameHandler.cpp
@@ -62,7 +62,7 @@ void FMainFrameHandler::ShutDownEditor()
 	// Delete user settings, if requested
 	if (FUnrealEdMisc::Get().IsDeletePreferences())
 	{
-		IFileManager::Get().Delete(*GEditorUserSettingsIni);
+		IFileManager::Get().Delete(*GEditorPerProjectIni);
 	}
 
 	// Take a screenshot of this project for the project browser
diff --git a/Engine/Source/Editor/MainFrame/Private/Frame/RootWindowLocation.h b/Engine/Source/Editor/MainFrame/Private/Frame/RootWindowLocation.h
--- a/Engine/Source/Editor/MainFrame/Private/Frame/RootWindowLocation.h
+++ b/Engine/Source/Editor/MainFrame/Private/Frame/RootWindowLocation.h
@@ -76,9 +76,9 @@ public:
 	 */
 	void SaveToIni( )
 	{
-		GConfig->SetString( TEXT("RootWindow"), TEXT("ScreenPosition"), *ScreenPosition.ToString(), GEditorUserSettingsIni );
-		GConfig->SetString( TEXT("RootWindow"), TEXT("WindowSize"), *WindowSize.ToString(), GEditorUserSettingsIni );
-		GConfig->SetBool( TEXT("RootWindow"), TEXT("InitiallyMaximized"), InitiallyMaximized, GEditorUserSettingsIni );
+		GConfig->SetString( TEXT("RootWindow"), TEXT("ScreenPosition"), *ScreenPosition.ToString(), GEditorPerProjectIni );
+		GConfig->SetString( TEXT("RootWindow"), TEXT("WindowSize"), *WindowSize.ToString(), GEditorPerProjectIni );
+		GConfig->SetBool( TEXT("RootWindow"), TEXT("InitiallyMaximized"), InitiallyMaximized, GEditorPerProjectIni );
 	}
 
 
@@ -88,7 +88,7 @@ private:
 	{
 		FVector2D ReturnValue = DefaultValue;
 		FString ValueAsString;
-		if ( GConfig->GetString(TEXT("RootWindow"), SettingName, ValueAsString, GEditorUserSettingsIni) && ReturnValue.InitFromString(ValueAsString) )
+		if ( GConfig->GetString(TEXT("RootWindow"), SettingName, ValueAsString, GEditorPerProjectIni) && ReturnValue.InitFromString(ValueAsString) )
 		{
 			// Successfully loaded setting
 			return ReturnValue;
@@ -102,7 +102,7 @@ private:
 	static bool BoolFromSettings( const TCHAR* SettingName, bool DefaultValue )
 	{
 		bool ReturnValue;
-		if ( GConfig->GetBool(TEXT("RootWindow"), SettingName, ReturnValue, GEditorUserSettingsIni) )
+		if ( GConfig->GetBool(TEXT("RootWindow"), SettingName, ReturnValue, GEditorPerProjectIni) )
 		{
 			return ReturnValue;
 		}
diff --git a/Engine/Source/Editor/MainFrame/Private/MainFrameModule.cpp b/Engine/Source/Editor/MainFrame/Private/MainFrameModule.cpp
--- a/Engine/Source/Editor/MainFrame/Private/MainFrameModule.cpp
+++ b/Engine/Source/Editor/MainFrame/Private/MainFrameModule.cpp
@@ -347,7 +347,7 @@ TSharedRef<SWidget> FMainFrameModule::MakeDeveloperTools() const
 		/** @return Returns true if frame rate and memory should be displayed in the UI */
 		static EVisibility ShouldShowFrameRateAndMemory()
 		{
-			return GEditor->GetEditorUserSettings().bShowFrameRateAndMemory ? EVisibility::SelfHitTestInvisible : EVisibility::Collapsed;
+			return GetDefault<UEditorPerProjectUserSettings>()->bShowFrameRateAndMemory ? EVisibility::SelfHitTestInvisible : EVisibility::Collapsed;
 		}
 	};
 
@@ -638,7 +638,7 @@ void FMainFrameModule::StartupModule( )
 	ModuleCompileStartTime = 0.0f;
 
 	// migrate old layout settings
-	FLayoutSaveRestore::MigrateConfig(GEditorUserSettingsIni, GEditorLayoutIni);
+	FLayoutSaveRestore::MigrateConfig(GEditorPerProjectIni, GEditorLayoutIni);
 }
 
 
diff --git a/Engine/Source/Editor/MaterialEditor/Private/MaterialEditor.cpp b/Engine/Source/Editor/MaterialEditor/Private/MaterialEditor.cpp
--- a/Engine/Source/Editor/MaterialEditor/Private/MaterialEditor.cpp
+++ b/Engine/Source/Editor/MaterialEditor/Private/MaterialEditor.cpp
@@ -1090,7 +1090,7 @@ void FMaterialEditor::LoadEditorSettings()
 
 	// Primitive type
 	int32 PrimType;
-	if(GConfig->GetInt(TEXT("MaterialEditor"), TEXT("PrimType"), PrimType, GEditorUserSettingsIni))
+	if(GConfig->GetInt(TEXT("MaterialEditor"), TEXT("PrimType"), PrimType, GEditorPerProjectIni))
 	{
 		Viewport->OnSetPreviewPrimitive((EThumbnailPrimType)PrimType);
 	}
@@ -1115,7 +1115,7 @@ void FMaterialEditor::SaveEditorSettings()
 		EditorOptions->SaveConfig();
 	}
 
-	GConfig->SetInt(TEXT("MaterialEditor"), TEXT("PrimType"), Viewport->PreviewPrimType, GEditorUserSettingsIni);
+	GConfig->SetInt(TEXT("MaterialEditor"), TEXT("PrimType"), Viewport->PreviewPrimType, GEditorPerProjectIni);
 }
 
 FText FMaterialEditor::GetCodeViewText() const
diff --git a/Engine/Source/Editor/MaterialEditor/Private/MaterialEditorActions.cpp b/Engine/Source/Editor/MaterialEditor/Private/MaterialEditorActions.cpp
--- a/Engine/Source/Editor/MaterialEditor/Private/MaterialEditorActions.cpp
+++ b/Engine/Source/Editor/MaterialEditor/Private/MaterialEditorActions.cpp
@@ -80,7 +80,7 @@ void FMaterialEditorSpawnNodeCommands::RegisterCommands()
 	const FString ConfigSection = TEXT("MaterialEditorSpawnNodes");
 	const FString SettingName = TEXT("Node");
 	TArray< FString > NodeSpawns;
-	GConfig->GetArray(*ConfigSection, *SettingName, NodeSpawns, GEditorUserSettingsIni);
+	GConfig->GetArray(*ConfigSection, *SettingName, NodeSpawns, GEditorPerProjectIni);
 
 	for(int32 x = 0; x < NodeSpawns.Num(); ++x)
 	{
diff --git a/Engine/Source/Editor/MaterialEditor/Private/MaterialInstanceEditor.cpp b/Engine/Source/Editor/MaterialEditor/Private/MaterialInstanceEditor.cpp
--- a/Engine/Source/Editor/MaterialEditor/Private/MaterialInstanceEditor.cpp
+++ b/Engine/Source/Editor/MaterialEditor/Private/MaterialInstanceEditor.cpp
@@ -849,10 +849,10 @@ void FMaterialInstanceEditor::SaveSettings()
 		PreviewVC->PreviewScene.SaveSettings(TEXT("MaterialInstanceEditor"));
 	}
 
-	GConfig->SetBool(TEXT("MaterialInstanceEditor"), TEXT("bShowGrid"), PreviewVC->IsTogglePreviewGridChecked(), GEditorUserSettingsIni);
-	GConfig->SetBool(TEXT("MaterialInstanceEditor"), TEXT("bDrawGrid"), PreviewVC->IsRealtime(), GEditorUserSettingsIni);
-	GConfig->SetInt(TEXT("MaterialInstanceEditor"), TEXT("PrimType"), PreviewVC->PreviewPrimType, GEditorUserSettingsIni);
-	GConfig->SetBool(TEXT("MaterialInstanceEditor"), TEXT("bWantsMobileStats"), IsToggleMobileStatsChecked(), GEditorUserSettingsIni);
+	GConfig->SetBool(TEXT("MaterialInstanceEditor"), TEXT("bShowGrid"), PreviewVC->IsTogglePreviewGridChecked(), GEditorPerProjectIni);
+	GConfig->SetBool(TEXT("MaterialInstanceEditor"), TEXT("bDrawGrid"), PreviewVC->IsRealtime(), GEditorPerProjectIni);
+	GConfig->SetInt(TEXT("MaterialInstanceEditor"), TEXT("PrimType"), PreviewVC->PreviewPrimType, GEditorPerProjectIni);
+	GConfig->SetBool(TEXT("MaterialInstanceEditor"), TEXT("bWantsMobileStats"), IsToggleMobileStatsChecked(), GEditorPerProjectIni);
 }
 
 void FMaterialInstanceEditor::LoadSettings()
@@ -861,10 +861,10 @@ void FMaterialInstanceEditor::LoadSettings()
 	bool bShowGrid=false;
 	int32 PrimType=static_cast<EThumbnailPrimType>( TPT_Sphere );
 	bool bWantsMobileStats=false;
-	GConfig->GetBool(TEXT("MaterialInstanceEditor"), TEXT("bShowGrid"), bShowGrid, GEditorUserSettingsIni);
-	GConfig->GetBool(TEXT("MaterialInstanceEditor"), TEXT("bDrawGrid"), bRealtime, GEditorUserSettingsIni);
-	GConfig->GetInt(TEXT("MaterialInstanceEditor"), TEXT("PrimType"), PrimType, GEditorUserSettingsIni);
-	GConfig->GetBool(TEXT("MaterialInstanceEditor"), TEXT("bWantsMobileStats"), bWantsMobileStats, GEditorUserSettingsIni);
+	GConfig->GetBool(TEXT("MaterialInstanceEditor"), TEXT("bShowGrid"), bShowGrid, GEditorPerProjectIni);
+	GConfig->GetBool(TEXT("MaterialInstanceEditor"), TEXT("bDrawGrid"), bRealtime, GEditorPerProjectIni);
+	GConfig->GetInt(TEXT("MaterialInstanceEditor"), TEXT("PrimType"), PrimType, GEditorPerProjectIni);
+	GConfig->GetBool(TEXT("MaterialInstanceEditor"), TEXT("bWantsMobileStats"), bWantsMobileStats, GEditorPerProjectIni);
 
 	if (bWantsMobileStats)
 	{
