Fix potential crash in FLandscapeBrushComponent::LeaveBrush reported by licensee (UE-14279)

diff --git a/Engine/Source/Editor/LandscapeEditor/Private/LandscapeEdModeBrushes.cpp b/Engine/Source/Editor/LandscapeEditor/Private/LandscapeEdModeBrushes.cpp
--- a/Engine/Source/Editor/LandscapeEditor/Private/LandscapeEdModeBrushes.cpp
+++ b/Engine/Source/Editor/LandscapeEditor/Private/LandscapeEdModeBrushes.cpp
@@ -80,7 +80,7 @@ public:
 	{
 		for (TSet<ULandscapeComponent*>::TIterator It(BrushMaterialComponents); It; ++It)
 		{
-			if ((*It)->EditToolRenderData != nullptr)
+			if ((*It) != nullptr && (*It)->EditToolRenderData != nullptr)
 			{
 				(*It)->EditToolRenderData->Update(nullptr);
 			}
@@ -126,11 +126,19 @@ public:
 		TSet<ULandscapeComponent*> RemovedComponents = BrushMaterialComponents.Difference(NewComponents);
 		for (ULandscapeComponent* RemovedComponent : RemovedComponents)
 		{
-			BrushMaterialFreeInstances.Push(BrushMaterialInstanceMap.FindAndRemoveChecked(RemovedComponent));
-
-			if (ensure(RemovedComponent->EditToolRenderData != nullptr))
+			if (RemovedComponent == nullptr)
 			{
-				RemovedComponent->EditToolRenderData->Update(nullptr);
+				// can occur if a component was removed by a user tool
+				BrushMaterialInstanceMap.Remove(nullptr);
+			}
+			else
+			{
+				BrushMaterialFreeInstances.Push(BrushMaterialInstanceMap.FindAndRemoveChecked(RemovedComponent));
+
+				if (ensure(RemovedComponent->EditToolRenderData != nullptr))
+				{
+					RemovedComponent->EditToolRenderData->Update(nullptr);
+				}
 			}
 		}
 
@@ -337,6 +345,7 @@ public:
 	// FGCObject interface
 	virtual void AddReferencedObjects(FReferenceCollector& Collector) override
 	{
+		Collector.AddReferencedObjects(BrushMaterialComponents);
 		Collector.AddReferencedObject(BrushMaterial);
 	}
 
@@ -346,7 +355,7 @@ public:
 	{
 		for (TSet<ULandscapeComponent*>::TIterator It(BrushMaterialComponents); It; ++It)
 		{
-			if ((*It)->EditToolRenderData != nullptr)
+			if ((*It) != nullptr && (*It)->EditToolRenderData != nullptr)
 			{
 				(*It)->EditToolRenderData->Update(nullptr);
 			}
@@ -414,7 +423,7 @@ public:
 		TSet<ULandscapeComponent*> RemovedComponents = BrushMaterialComponents.Difference(NewComponents);
 		for (ULandscapeComponent* RemovedComponent : RemovedComponents)
 		{
-			if (RemovedComponent->EditToolRenderData != nullptr)
+			if (RemovedComponent != nullptr && RemovedComponent->EditToolRenderData != nullptr)
 			{
 				RemovedComponent->EditToolRenderData->Update(nullptr);
 			}
@@ -531,6 +540,7 @@ public:
 	// FGCObject interface
 	virtual void AddReferencedObjects(FReferenceCollector& Collector) override
 	{
+		Collector.AddReferencedObjects(BrushMaterialComponents);
 		Collector.AddReferencedObject(BrushMaterial);
 	}
 
@@ -551,7 +561,7 @@ public:
 	{
 		for (TSet<ULandscapeComponent*>::TIterator It(BrushMaterialComponents); It; ++It)
 		{
-			if ((*It)->EditToolRenderData != nullptr)
+			if ((*It) != nullptr && (*It)->EditToolRenderData != nullptr)
 			{
 				(*It)->EditToolRenderData->Update(nullptr);
 			}
@@ -619,7 +629,7 @@ public:
 					TSet<ULandscapeComponent*> RemovedComponents = BrushMaterialComponents.Difference(NewComponents);
 					for (ULandscapeComponent* RemovedComponent : RemovedComponents)
 					{
-						if (RemovedComponent->EditToolRenderData != nullptr)
+						if (RemovedComponent != nullptr && RemovedComponent->EditToolRenderData != nullptr)
 						{
 							RemovedComponent->EditToolRenderData->UpdateGizmo(nullptr);
 						}
