Fix landscape layer import being averaged with existing data

diff --git a/Engine/Source/Editor/LandscapeEditor/Private/LandscapeEdMode.cpp b/Engine/Source/Editor/LandscapeEditor/Private/LandscapeEdMode.cpp
--- a/Engine/Source/Editor/LandscapeEditor/Private/LandscapeEdMode.cpp
+++ b/Engine/Source/Editor/LandscapeEditor/Private/LandscapeEdMode.cpp
@@ -2457,7 +2457,7 @@ void FEdModeLandscape::ImportData(const FLandscapeTargetListInfo& TargetInfo, co
 			{
 				if (RawData->Num() == SizeX * SizeY)
 				{
-					FAlphamapAccessor<false, true> AlphamapAccessor(LandscapeInfo, TargetInfo.LayerInfoObj.Get());
+					FAlphamapAccessor<false, false> AlphamapAccessor(LandscapeInfo, TargetInfo.LayerInfoObj.Get());
 					AlphamapAccessor.SetData(MinX, MinY, MaxX, MaxY, RawData->GetData(), ELandscapeLayerPaintingRestriction::None);
 				}
 				else
