Fixed: Foliage painted on landscape does not snap to it when user sculpts landscape
UE-14723

diff --git a/Engine/Source/Editor/LandscapeEditor/Private/LandscapeEdModeTools.h b/Engine/Source/Editor/LandscapeEditor/Private/LandscapeEdModeTools.h
--- a/Engine/Source/Editor/LandscapeEditor/Private/LandscapeEdModeTools.h
+++ b/Engine/Source/Editor/LandscapeEditor/Private/LandscapeEdModeTools.h
@@ -625,7 +625,8 @@ struct FHeightmapAccessor
 			bool bUpdateFoliage = false;
 			for (ULandscapeComponent* Component : Components)
 			{
-				if (AInstancedFoliageActor::HasFoliageAttached(Component))
+				ULandscapeHeightfieldCollisionComponent* CollisionComponent = Component->CollisionComponent.Get();
+				if (CollisionComponent && AInstancedFoliageActor::HasFoliageAttached(CollisionComponent))
 				{
 					bUpdateFoliage = true;
 					break;
@@ -793,7 +794,8 @@ struct FXYOffsetmapAccessor
 			bool bUpdateFoliage = false;
 			for (ULandscapeComponent* Component : Components)
 			{
-				if (AInstancedFoliageActor::HasFoliageAttached(Component))
+				ULandscapeHeightfieldCollisionComponent* CollisionComponent = Component->CollisionComponent.Get();
+				if (CollisionComponent && AInstancedFoliageActor::HasFoliageAttached(CollisionComponent))
 				{
 					bUpdateFoliage = true;
 					break;
