Fixed: Editor crash after attempting to create layer info in new landscape window
UE-15159

diff --git a/Engine/Source/Editor/LandscapeEditor/Private/LandscapeEditorObject.cpp b/Engine/Source/Editor/LandscapeEditor/Private/LandscapeEditorObject.cpp
--- a/Engine/Source/Editor/LandscapeEditor/Private/LandscapeEditorObject.cpp
+++ b/Engine/Source/Editor/LandscapeEditor/Private/LandscapeEditorObject.cpp
@@ -146,8 +146,7 @@ void ULandscapeEditorObject::PostEditChangeProperty(FPropertyChangedEvent& Prope
 
 	if (PropertyChangedEvent.MemberProperty == NULL ||
 		PropertyChangedEvent.MemberProperty->GetFName() == GET_MEMBER_NAME_CHECKED(ULandscapeEditorObject, NewLandscape_Material) ||
-		PropertyChangedEvent.MemberProperty->GetFName() == GET_MEMBER_NAME_CHECKED(ULandscapeEditorObject, ImportLandscape_HeightmapFilename) ||
-		PropertyChangedEvent.MemberProperty->GetFName() == GET_MEMBER_NAME_CHECKED(ULandscapeEditorObject, ImportLandscape_Layers))
+		PropertyChangedEvent.MemberProperty->GetFName() == GET_MEMBER_NAME_CHECKED(ULandscapeEditorObject, ImportLandscape_HeightmapFilename))
 	{
 		RefreshImportLayersList();
 	}
