Fixed: Unable to undo after landscape creation with certain settings
UE-7597

diff --git a/Engine/Source/Editor/LandscapeEditor/Classes/LandscapeEditorObject.h b/Engine/Source/Editor/LandscapeEditor/Classes/LandscapeEditorObject.h
--- a/Engine/Source/Editor/LandscapeEditor/Classes/LandscapeEditorObject.h
+++ b/Engine/Source/Editor/LandscapeEditor/Classes/LandscapeEditorObject.h
@@ -555,12 +555,22 @@ class ULandscapeEditorObject : public UObject
 	void ClearImportLandscapeData() { ImportLandscape_Data.Empty(); }
 
 	void RefreshImportLayersList();
+	
+	int32 ClampLandscapeSize(int32 InComponenetsCount) const
+	{
+		// Max size is either whole components below 8192 verts, or 32 components 
+		return FMath::Clamp(InComponenetsCount, 1, FMath::Min(32, FMath::FloorToInt(8191 / (NewLandscape_SectionsPerComponent * NewLandscape_QuadsPerSection))));
+	}
+	
+	int32 CalcComponentsCount(int32 InResolution) const
+	{
+		return ClampLandscapeSize(InResolution / (NewLandscape_SectionsPerComponent * NewLandscape_QuadsPerSection));
+	}
 
 	void NewLandscape_ClampSize()
 	{
-		// Max size is either whole components below 8192 verts, or 32 components
-		NewLandscape_ComponentCount.X = FMath::Clamp(NewLandscape_ComponentCount.X, 1, FMath::Min(32, FMath::FloorToInt(8191 / (NewLandscape_SectionsPerComponent * NewLandscape_QuadsPerSection))));
-		NewLandscape_ComponentCount.Y = FMath::Clamp(NewLandscape_ComponentCount.Y, 1, FMath::Min(32, FMath::FloorToInt(8191 / (NewLandscape_SectionsPerComponent * NewLandscape_QuadsPerSection))));
+		NewLandscape_ComponentCount.X = ClampLandscapeSize(NewLandscape_ComponentCount.X);
+		NewLandscape_ComponentCount.Y = ClampLandscapeSize(NewLandscape_ComponentCount.Y);
 	}
 
 	void UpdateComponentCount()
diff --git a/Engine/Source/Editor/LandscapeEditor/Private/LandscapeEditorDetailCustomization_NewLandscape.cpp b/Engine/Source/Editor/LandscapeEditor/Private/LandscapeEditorDetailCustomization_NewLandscape.cpp
--- a/Engine/Source/Editor/LandscapeEditor/Private/LandscapeEditorDetailCustomization_NewLandscape.cpp
+++ b/Engine/Source/Editor/LandscapeEditor/Private/LandscapeEditorDetailCustomization_NewLandscape.cpp
@@ -371,12 +371,13 @@ void FLandscapeEditorDetailCustomization_NewLandscape::CustomizeDetails(IDetailL
 		+ SHorizontalBox::Slot()
 		.FillWidth(1)
 		[
-			SNew(SSpinBox<int32>)
+			SNew(SNumericEntryBox<int32>)
 			.Font(DetailBuilder.GetDetailFont())
 			.MinValue(1)
 			.MaxValue(8192)
 			.MinSliderValue(1)
 			.MaxSliderValue(8192)
+			.AllowSpin(true)
 			//.MinSliderValue(TAttribute<TOptional<int32> >(this, &FLandscapeEditorDetailCustomization_NewLandscape::GetMinLandscapeResolution))
 			//.MaxSliderValue(TAttribute<TOptional<int32> >(this, &FLandscapeEditorDetailCustomization_NewLandscape::GetMaxLandscapeResolution))
 			.Value(this, &FLandscapeEditorDetailCustomization_NewLandscape::GetLandscapeResolutionX)
@@ -396,12 +397,13 @@ void FLandscapeEditorDetailCustomization_NewLandscape::CustomizeDetails(IDetailL
 		.FillWidth(1)
 		.Padding(0,0,12,0) // Line up with the other properties due to having no reset to default button
 		[
-			SNew(SSpinBox<int32>)
+			SNew(SNumericEntryBox<int32>)
 			.Font(DetailBuilder.GetDetailFont())
 			.MinValue(1)
 			.MaxValue(8192)
 			.MinSliderValue(1)
 			.MaxSliderValue(8192)
+			.AllowSpin(true)
 			//.MinSliderValue(TAttribute<TOptional<int32> >(this, &FLandscapeEditorDetailCustomization_NewLandscape::GetMinLandscapeResolution))
 			//.MaxSliderValue(TAttribute<TOptional<int32> >(this, &FLandscapeEditorDetailCustomization_NewLandscape::GetMaxLandscapeResolution))
 			.Value(this, &FLandscapeEditorDetailCustomization_NewLandscape::GetLandscapeResolutionY)
@@ -591,7 +593,7 @@ FText FLandscapeEditorDetailCustomization_NewLandscape::GetSectionsPerComponent(
 	return FText::Format(SectionsPerComponent == 1 ? LOCTEXT("1x1Section", "{Width}\u00D7{Height} Section") : LOCTEXT("NxNSections", "{Width}\u00D7{Height} Sections"), Args);
 }
 
-int32 FLandscapeEditorDetailCustomization_NewLandscape::GetLandscapeResolutionX() const
+TOptional<int32> FLandscapeEditorDetailCustomization_NewLandscape::GetLandscapeResolutionX() const
 {
 	FEdModeLandscape* LandscapeEdMode = GetEditorMode();
 	if (LandscapeEdMode != NULL)
@@ -607,13 +609,17 @@ void FLandscapeEditorDetailCustomization_NewLandscape::OnChangeLandscapeResoluti
 	FEdModeLandscape* LandscapeEdMode = GetEditorMode();
 	if (LandscapeEdMode != NULL)
 	{
-		if (!GEditor->IsTransactionActive())
+		int32 NewComponentCountX = LandscapeEdMode->UISettings->CalcComponentsCount(NewValue);
+		if (NewComponentCountX != LandscapeEdMode->UISettings->NewLandscape_ComponentCount.X)
 		{
-			GEditor->BeginTransaction(LOCTEXT("ChangeResolutionX_Transaction", "Change Landscape Resolution X"));
+			if (!GEditor->IsTransactionActive())
+			{
+				GEditor->BeginTransaction(LOCTEXT("ChangeResolutionX_Transaction", "Change Landscape Resolution X"));
+			}
+
+			LandscapeEdMode->UISettings->Modify();
+			LandscapeEdMode->UISettings->NewLandscape_ComponentCount.X = NewComponentCountX;
 		}
-		LandscapeEdMode->UISettings->Modify();
-		LandscapeEdMode->UISettings->NewLandscape_ComponentCount.X = (NewValue / (LandscapeEdMode->UISettings->NewLandscape_SectionsPerComponent * LandscapeEdMode->UISettings->NewLandscape_QuadsPerSection));
-		LandscapeEdMode->UISettings->NewLandscape_ClampSize();
 	}
 }
 
@@ -627,13 +633,12 @@ void FLandscapeEditorDetailCustomization_NewLandscape::OnCommitLandscapeResoluti
 			GEditor->BeginTransaction(LOCTEXT("ChangeResolutionX_Transaction", "Change Landscape Resolution X"));
 		}
 		LandscapeEdMode->UISettings->Modify();
-		LandscapeEdMode->UISettings->NewLandscape_ComponentCount.X = (NewValue / (LandscapeEdMode->UISettings->NewLandscape_SectionsPerComponent * LandscapeEdMode->UISettings->NewLandscape_QuadsPerSection));
-		LandscapeEdMode->UISettings->NewLandscape_ClampSize();
+		LandscapeEdMode->UISettings->NewLandscape_ComponentCount.X = LandscapeEdMode->UISettings->CalcComponentsCount(NewValue);
 		GEditor->EndTransaction();
 	}
 }
 
-int32 FLandscapeEditorDetailCustomization_NewLandscape::GetLandscapeResolutionY() const
+TOptional<int32> FLandscapeEditorDetailCustomization_NewLandscape::GetLandscapeResolutionY() const
 {
 	FEdModeLandscape* LandscapeEdMode = GetEditorMode();
 	if (LandscapeEdMode != NULL)
@@ -649,13 +654,17 @@ void FLandscapeEditorDetailCustomization_NewLandscape::OnChangeLandscapeResoluti
 	FEdModeLandscape* LandscapeEdMode = GetEditorMode();
 	if (LandscapeEdMode != NULL)
 	{
-		if (!GEditor->IsTransactionActive())
+		int32 NewComponentCountY = LandscapeEdMode->UISettings->CalcComponentsCount(NewValue);
+		if (NewComponentCountY != LandscapeEdMode->UISettings->NewLandscape_ComponentCount.Y)
 		{
-			GEditor->BeginTransaction(LOCTEXT("ChangeResolutionY_Transaction", "Change Landscape Resolution Y"));
+			if (!GEditor->IsTransactionActive())
+			{
+				GEditor->BeginTransaction(LOCTEXT("ChangeResolutionY_Transaction", "Change Landscape Resolution Y"));
+			}
+			
+			LandscapeEdMode->UISettings->Modify();
+			LandscapeEdMode->UISettings->NewLandscape_ComponentCount.Y = NewComponentCountY;
 		}
-		LandscapeEdMode->UISettings->Modify();
-		LandscapeEdMode->UISettings->NewLandscape_ComponentCount.Y = (NewValue / (LandscapeEdMode->UISettings->NewLandscape_SectionsPerComponent * LandscapeEdMode->UISettings->NewLandscape_QuadsPerSection));
-		LandscapeEdMode->UISettings->NewLandscape_ClampSize();
 	}
 }
 
@@ -669,8 +678,7 @@ void FLandscapeEditorDetailCustomization_NewLandscape::OnCommitLandscapeResoluti
 			GEditor->BeginTransaction(LOCTEXT("ChangeResolutionY_Transaction", "Change Landscape Resolution Y"));
 		}
 		LandscapeEdMode->UISettings->Modify();
-		LandscapeEdMode->UISettings->NewLandscape_ComponentCount.Y = (NewValue / (LandscapeEdMode->UISettings->NewLandscape_SectionsPerComponent * LandscapeEdMode->UISettings->NewLandscape_QuadsPerSection));
-		LandscapeEdMode->UISettings->NewLandscape_ClampSize();
+		LandscapeEdMode->UISettings->NewLandscape_ComponentCount.Y = LandscapeEdMode->UISettings->CalcComponentsCount(NewValue);
 		GEditor->EndTransaction();
 	}
 }
diff --git a/Engine/Source/Editor/LandscapeEditor/Private/LandscapeEditorDetailCustomization_NewLandscape.h b/Engine/Source/Editor/LandscapeEditor/Private/LandscapeEditorDetailCustomization_NewLandscape.h
--- a/Engine/Source/Editor/LandscapeEditor/Private/LandscapeEditorDetailCustomization_NewLandscape.h
+++ b/Engine/Source/Editor/LandscapeEditor/Private/LandscapeEditorDetailCustomization_NewLandscape.h
@@ -35,11 +35,11 @@ public:
 	static void OnChangeSectionsPerComponent(TSharedRef<IPropertyHandle> PropertyHandle, int32 NewSize);
 	static FText GetSectionsPerComponent(TSharedRef<IPropertyHandle> PropertyHandle);
 
-	int32 GetLandscapeResolutionX() const;
+	TOptional<int32> GetLandscapeResolutionX() const;
 	void OnChangeLandscapeResolutionX(int32 NewValue);
 	void OnCommitLandscapeResolutionX(int32 NewValue, ETextCommit::Type CommitInfo);
 
-	int32 GetLandscapeResolutionY() const;
+	TOptional<int32> GetLandscapeResolutionY() const;
 	void OnChangeLandscapeResolutionY(int32 NewValue);
 	void OnCommitLandscapeResolutionY(int32 NewValue, ETextCommit::Type CommitInfo);
 
