Fixed: Crash turning off the visibility of a second sub level with a landscape in it
UE-15410

diff --git a/Engine/Source/Editor/LandscapeEditor/Private/LandscapeEditorDetails.cpp b/Engine/Source/Editor/LandscapeEditor/Private/LandscapeEditorDetails.cpp
--- a/Engine/Source/Editor/LandscapeEditor/Private/LandscapeEditorDetails.cpp
+++ b/Engine/Source/Editor/LandscapeEditor/Private/LandscapeEditorDetails.cpp
@@ -208,9 +208,17 @@ EVisibility FLandscapeEditorDetails::GetTargetLandscapeSelectorVisibility()
 FText FLandscapeEditorDetails::GetTargetLandscapeName()
 {
 	FEdModeLandscape* LandscapeEdMode = GetEditorMode();
-	if (LandscapeEdMode && LandscapeEdMode->CurrentToolTarget.LandscapeInfo.IsValid())
+	if (LandscapeEdMode)
 	{
-		return FText::FromString(LandscapeEdMode->CurrentToolTarget.LandscapeInfo->GetLandscapeProxy()->GetActorLabel());
+		ULandscapeInfo* Info = LandscapeEdMode->CurrentToolTarget.LandscapeInfo.Get();
+		if (Info)
+		{
+			ALandscapeProxy* Proxy = Info->GetLandscapeProxy();
+			if (Proxy)
+			{
+				return FText::FromString(Proxy->GetActorLabel());
+			}
+		}
 	}
 
 	return FText();
