Fixed int -> int32 in landscape code
Also fixed a few oddities following CL 1896753 with ints in the landscape details customization

diff --git a/Engine/Source/Editor/LandscapeEditor/Private/LandscapeEdModeComponentTools.cpp b/Engine/Source/Editor/LandscapeEditor/Private/LandscapeEdModeComponentTools.cpp
--- a/Engine/Source/Editor/LandscapeEditor/Private/LandscapeEdModeComponentTools.cpp
+++ b/Engine/Source/Editor/LandscapeEditor/Private/LandscapeEdModeComponentTools.cpp
@@ -327,7 +327,7 @@ public:
 			}
 
 			// Check the LayerInfoObjects are same package with Landscape
-			for (int i = 0; i < LandscapeInfo->Layers.Num(); ++i)
+			for (int32 i = 0; i < LandscapeInfo->Layers.Num(); ++i)
 			{
 				ULandscapeLayerInfoObject* LayerInfo = LandscapeInfo->Layers[i].LayerInfoObj;
 				if (LayerInfo && LayerInfo->GetOutermost() == Landscape->GetOutermost())
diff --git a/Engine/Source/Editor/LandscapeEditor/Private/LandscapeEdModeTools.h b/Engine/Source/Editor/LandscapeEditor/Private/LandscapeEdModeTools.h
--- a/Engine/Source/Editor/LandscapeEditor/Private/LandscapeEdModeTools.h
+++ b/Engine/Source/Editor/LandscapeEditor/Private/LandscapeEdModeTools.h
@@ -157,7 +157,7 @@ inline void LowPassFilter(int32 X1, int32 Y1, int32 X2, int32 Y2, FLandscapeBrus
 	int32 FFTWidth = X2 - X1 - 1;
 	int32 FFTHeight = Y2 - Y1 - 1;
 
-	const int NDims = 2;
+	const int32 NDims = 2;
 	const int32 Dims[NDims] = { FFTHeight - FFTHeight % 2, FFTWidth - FFTWidth % 2 };
 	kiss_fftnd_cfg stf = kiss_fftnd_alloc(Dims, NDims, 0, NULL, NULL),
 		sti = kiss_fftnd_alloc(Dims, NDims, 1, NULL, NULL);
@@ -165,12 +165,12 @@ inline void LowPassFilter(int32 X1, int32 Y1, int32 X2, int32 Y2, FLandscapeBrus
 	kiss_fft_cpx *buf = (kiss_fft_cpx *)KISS_FFT_MALLOC(sizeof(kiss_fft_cpx) * Dims[0] * Dims[1]);
 	kiss_fft_cpx *out = (kiss_fft_cpx *)KISS_FFT_MALLOC(sizeof(kiss_fft_cpx) * Dims[0] * Dims[1]);
 
-	for (int Y = Y1 + 1; Y <= Y2 - 1 - FFTHeight % 2; Y++)
+	for (int32 Y = Y1 + 1; Y <= Y2 - 1 - FFTHeight % 2; Y++)
 	{
 		auto* DataScanline = Data.GetData() + (Y - Y1) * (X2 - X1 + 1) + (0 - X1);
 		auto* bufScanline = buf + (Y - (Y1 + 1)) * Dims[1] + (0 - (X1 + 1));
 
-		for (int X = X1 + 1; X <= X2 - 1 - FFTWidth % 2; X++)
+		for (int32 X = X1 + 1; X <= X2 - 1 - FFTWidth % 2; X++)
 		{
 			bufScanline[X].r = DataScanline[X];
 			bufScanline[X].i = 0;
@@ -181,10 +181,10 @@ inline void LowPassFilter(int32 X1, int32 Y1, int32 X2, int32 Y2, FLandscapeBrus
 	kiss_fftnd(stf, buf, out);
 
 	int32 CenterPos[2] = { Dims[0] >> 1, Dims[1] >> 1 };
-	for (int Y = 0; Y < Dims[0]; Y++)
+	for (int32 Y = 0; Y < Dims[0]; Y++)
 	{
 		float DistFromCenter = 0.0f;
-		for (int X = 0; X < Dims[1]; X++)
+		for (int32 X = 0; X < Dims[1]; X++)
 		{
 			if (Y < CenterPos[0])
 			{
diff --git a/Engine/Source/Editor/LandscapeEditor/Private/LandscapeEditorDetailCustomization_Base.cpp b/Engine/Source/Editor/LandscapeEditor/Private/LandscapeEditorDetailCustomization_Base.cpp
--- a/Engine/Source/Editor/LandscapeEditor/Private/LandscapeEditorDetailCustomization_Base.cpp
+++ b/Engine/Source/Editor/LandscapeEditor/Private/LandscapeEditorDetailCustomization_Base.cpp
@@ -41,14 +41,8 @@ bool FLandscapeEditorDetailCustomization_Base::IsBrushSetActive(FName BrushSetNa
 
 TOptional<int32> FLandscapeEditorDetailCustomization_Base::OnGetValue(TSharedRef<IPropertyHandle> PropertyHandle)
 {
-	// Try getting as float, if that fails then get as int
-	float FloatValue = 0.0f;
-	int32 IntValue = 0.0f;
-	if (PropertyHandle->GetValue(FloatValue) == FPropertyAccess::Success)
-	{
-		return TOptional<int32>(FloatValue);
-	}
-	else if (ensure(PropertyHandle->GetValue(IntValue) == FPropertyAccess::Success))
+	int32 IntValue = 0;
+	if (ensure(PropertyHandle->GetValue(IntValue) == FPropertyAccess::Success))
 	{
 		return TOptional<int32>(IntValue);
 	}
@@ -59,21 +53,13 @@ TOptional<int32> FLandscapeEditorDetailCustomization_Base::OnGetValue(TSharedRef
 
 void FLandscapeEditorDetailCustomization_Base::OnValueChanged(int32 NewValue, TSharedRef<IPropertyHandle> PropertyHandle)
 {
-	// Try setting as float, if that fails then set as int
-	EPropertyValueSetFlags::Type Flags = EPropertyValueSetFlags::InteractiveChange;
-	if (PropertyHandle->SetValue(NewValue, Flags) != FPropertyAccess::Success)
-	{
-		ensure(PropertyHandle->SetValue(NewValue, Flags) == FPropertyAccess::Success);
-	}
+	const EPropertyValueSetFlags::Type Flags = EPropertyValueSetFlags::InteractiveChange;
+	ensure(PropertyHandle->SetValue(NewValue, Flags) == FPropertyAccess::Success);
 }
 
 void FLandscapeEditorDetailCustomization_Base::OnValueCommitted(int32 NewValue, ETextCommit::Type CommitType, TSharedRef<IPropertyHandle> PropertyHandle)
 {
-	// Try setting as float, if that fails then set as int
-	if (PropertyHandle->SetValue(NewValue) != FPropertyAccess::Success)
-	{
-		ensure(PropertyHandle->SetValue(NewValue) == FPropertyAccess::Success);
-	}
+	ensure(PropertyHandle->SetValue(NewValue) == FPropertyAccess::Success);
 }
 
 //////////////////////////////////////////////////////////////////////////
diff --git a/Engine/Source/Editor/LandscapeEditor/Private/LandscapeEditorDetailCustomization_CopyPaste.cpp b/Engine/Source/Editor/LandscapeEditor/Private/LandscapeEditorDetailCustomization_CopyPaste.cpp
--- a/Engine/Source/Editor/LandscapeEditor/Private/LandscapeEditorDetailCustomization_CopyPaste.cpp
+++ b/Engine/Source/Editor/LandscapeEditor/Private/LandscapeEditorDetailCustomization_CopyPaste.cpp
@@ -379,7 +379,7 @@ FReply FLandscapeEditorDetailCustomization_CopyPaste::OnGizmoExportButtonClicked
 
 			// Local set for export
 			TSet<ULandscapeLayerInfoObject*> LayerInfoSet;
-			for (int i = 0; i < Gizmo->LayerInfos.Num(); i++)
+			for (int32 i = 0; i < Gizmo->LayerInfos.Num(); i++)
 			{
 				if (LandscapeEdMode->CurrentToolTarget.TargetType == ELandscapeToolTargetType::Weightmap && LandscapeEdMode->CurrentToolTarget.LayerInfo == Gizmo->LayerInfos[i])
 				{
diff --git a/Engine/Source/Editor/WorldBrowser/Private/Tiles/STiledLandscapeImportDlg.cpp b/Engine/Source/Editor/WorldBrowser/Private/Tiles/STiledLandscapeImportDlg.cpp
--- a/Engine/Source/Editor/WorldBrowser/Private/Tiles/STiledLandscapeImportDlg.cpp
+++ b/Engine/Source/Editor/WorldBrowser/Private/Tiles/STiledLandscapeImportDlg.cpp
@@ -98,7 +98,7 @@ void STiledLandcapeImportDlg::Construct(const FArguments& InArgs, TSharedPtr<SWi
 						.LabelPadding(0)
 						.Label()
 						[
-							SNumericEntryBox<int>::BuildLabel( LOCTEXT("X_Label", "X"), FLinearColor::White, SNumericEntryBox<int>::RedLabelBackgroundColor )
+							SNumericEntryBox<int32>::BuildLabel( LOCTEXT("X_Label", "X"), FLinearColor::White, SNumericEntryBox<int32>::RedLabelBackgroundColor )
 						]
 					]
 					
@@ -112,7 +112,7 @@ void STiledLandcapeImportDlg::Construct(const FArguments& InArgs, TSharedPtr<SWi
 						.LabelPadding(0)
 						.Label()
 						[
-							SNumericEntryBox<float>::BuildLabel( LOCTEXT("Y_Label", "Y"), FLinearColor::White, SNumericEntryBox<int>::GreenLabelBackgroundColor )
+							SNumericEntryBox<int32>::BuildLabel( LOCTEXT("Y_Label", "Y"), FLinearColor::White, SNumericEntryBox<int32>::GreenLabelBackgroundColor )
 						]
 					]
 				]
diff --git a/Engine/Source/Runtime/Landscape/Private/LandscapeGizmoActor.cpp b/Engine/Source/Runtime/Landscape/Private/LandscapeGizmoActor.cpp
--- a/Engine/Source/Runtime/Landscape/Private/LandscapeGizmoActor.cpp
+++ b/Engine/Source/Runtime/Landscape/Private/LandscapeGizmoActor.cpp
@@ -1166,7 +1166,7 @@ void ALandscapeGizmoActiveActor::ImportFromClipboard()
 				while( !FParse::Command(&Str,TEXT("Region=")) )
 				{
 					FParse::Next(&Str);
-					int i = 0;
+					int32 i = 0;
 					while (!FChar::IsWhitespace(*Str))
 					{
 						StrBuf[i++] = *Str;
diff --git a/Engine/Source/Runtime/Landscape/Private/LandscapeRender.cpp b/Engine/Source/Runtime/Landscape/Private/LandscapeRender.cpp
--- a/Engine/Source/Runtime/Landscape/Private/LandscapeRender.cpp
+++ b/Engine/Source/Runtime/Landscape/Private/LandscapeRender.cpp
@@ -713,7 +713,7 @@ void FLandscapeComponentSceneProxy::CreateRenderThreadResources()
 		SharedBuffers->AdjacencyIndexBuffers->AddRef();
 
 		// Delayed Initialize for IndexBuffers
-		for (int i = 0; i < SharedBuffers->NumIndexBuffers; i++)
+		for (int32 i = 0; i < SharedBuffers->NumIndexBuffers; i++)
 		{
 			SharedBuffers->IndexBuffers[i]->InitResource();
 		}
@@ -1866,7 +1866,7 @@ void FLandscapeSharedBuffers::CreateIndexBuffers(ERHIFeatureLevel::Type InFeatur
 		else
 		{
 			// non-ES2 version
-			int SubOffset = 0;
+			int32 SubOffset = 0;
 			for (int32 SubY = 0; SubY < NumSubsections; SubY++)
 			{
 				for (int32 SubX = 0; SubX < NumSubsections; SubX++)
@@ -2143,7 +2143,7 @@ FLandscapeSharedAdjacencyIndexBuffer::FLandscapeSharedAdjacencyIndexBuffer(FLand
 
 FLandscapeSharedAdjacencyIndexBuffer::~FLandscapeSharedAdjacencyIndexBuffer()
 {
-	for (int i = 0; i < IndexBuffers.Num(); ++i)
+	for (int32 i = 0; i < IndexBuffers.Num(); ++i)
 	{
 		IndexBuffers[i]->ReleaseResource();
 		delete IndexBuffers[i];
