Reorganize Landscape Promotion test grouping

diff --git a/Engine/Source/Editor/LandscapeEditor/Private/Tests/LandscapeAutomationTests.cpp b/Engine/Source/Editor/LandscapeEditor/Private/Tests/LandscapeAutomationTests.cpp
--- a/Engine/Source/Editor/LandscapeEditor/Private/Tests/LandscapeAutomationTests.cpp
+++ b/Engine/Source/Editor/LandscapeEditor/Private/Tests/LandscapeAutomationTests.cpp
@@ -117,7 +117,7 @@ bool FEndModifyLandscapeCommand::Update()
 /**
 * Landscape creation / edit test
 */
-IMPLEMENT_SIMPLE_AUTOMATION_TEST(FLandscapeEditorTest, "System.Promotion.Editor Promotion Pass.Step 1 Main Editor Test.New Landscape", EAutomationTestFlags::ATF_Editor | EAutomationTestFlags::ATF_NonNullRHI );
+IMPLEMENT_SIMPLE_AUTOMATION_TEST(FLandscapeEditorTest, "System.Promotion.Editor.Landscape Editor", EAutomationTestFlags::ATF_Editor | EAutomationTestFlags::ATF_NonNullRHI );
 bool FLandscapeEditorTest::RunTest(const FString& Parameters)
 {
 	//New level
