Fixed crash in the landscape flatten and ramp tools caused by not exposing a reference to an owned asset to the garbage collector and having it be GC'd
Moved the pointer to LandscapeInfo into the base tool stroke class because *every* tool stroke referenced it and it made it easier to expose to GC

diff --git a/Engine/Source/Editor/LandscapeEditor/Private/LandscapeEdModeComponentTools.cpp b/Engine/Source/Editor/LandscapeEditor/Private/LandscapeEdModeComponentTools.cpp
--- a/Engine/Source/Editor/LandscapeEditor/Private/LandscapeEdModeComponentTools.cpp
+++ b/Engine/Source/Editor/LandscapeEditor/Private/LandscapeEdModeComponentTools.cpp
@@ -27,8 +27,8 @@ class FLandscapeToolStrokeSelect : public FLandscapeToolStrokeBase
 
 public:
 	FLandscapeToolStrokeSelect(FEdModeLandscape* InEdMode, const FLandscapeToolTarget& InTarget)
-		: bInitializedComponentInvert(false)
-		, LandscapeInfo(InTarget.LandscapeInfo.Get())
+		: FLandscapeToolStrokeBase(InEdMode, InTarget)
+		, bInitializedComponentInvert(false)
 		, Cache(InTarget)
 	{
 	}
@@ -170,7 +170,6 @@ public:
 	}
 
 protected:
-	ULandscapeInfo* LandscapeInfo;
 	FLandscapeDataCache Cache;
 };
 
@@ -212,10 +211,8 @@ public:
 class FLandscapeToolStrokeVisibility : public FLandscapeToolStrokeBase
 {
 public:
-	enum { UseContinuousApply = false };
-
 	FLandscapeToolStrokeVisibility(FEdModeLandscape* InEdMode, const FLandscapeToolTarget& InTarget)
-		: LandscapeInfo(InTarget.LandscapeInfo.Get())
+		: FLandscapeToolStrokeBase(InEdMode, InTarget)
 		, Cache(InTarget)
 	{
 	}
@@ -268,7 +265,6 @@ public:
 	}
 
 protected:
-	ULandscapeInfo* LandscapeInfo;
 	FLandscapeVisCache Cache;
 };
 
@@ -295,13 +291,11 @@ public:
 	}
 };
 
-class FLandscapeToolStrokeMoveToLevel
+class FLandscapeToolStrokeMoveToLevel : public FLandscapeToolStrokeBase
 {
 public:
-	enum { UseContinuousApply = false };
-
 	FLandscapeToolStrokeMoveToLevel(FEdModeLandscape* InEdMode, const FLandscapeToolTarget& InTarget)
-		: LandscapeInfo(InTarget.LandscapeInfo.Get())
+		: FLandscapeToolStrokeBase(InEdMode, InTarget)
 	{
 	}
 
@@ -727,9 +721,6 @@ public:
 			}
 		}
 	}
-
-protected:
-	ULandscapeInfo* LandscapeInfo;
 };
 
 // 
@@ -755,8 +746,7 @@ class FLandscapeToolStrokeAddComponent : public FLandscapeToolStrokeBase
 {
 public:
 	FLandscapeToolStrokeAddComponent(FEdModeLandscape* InEdMode, const FLandscapeToolTarget& InTarget)
-		: EdMode(InEdMode)
-		, LandscapeInfo(InTarget.LandscapeInfo.Get())
+		: FLandscapeToolStrokeBase(InEdMode, InTarget)
 		, HeightCache(InTarget)
 		, XYOffsetCache(InTarget)
 	{
@@ -879,13 +869,52 @@ public:
 
 			EdMode->LandscapeRenderAddCollision = nullptr;
 
+			// Add/update "add collision" around the newly added components
+			{
+				// Top row
+				int32 ComponentIndexY = ComponentIndexY1 - 1;
+				for (int32 ComponentIndexX = ComponentIndexX1 - 1; ComponentIndexX <= ComponentIndexX2 + 1; ++ComponentIndexX)
+				{
+					if (!LandscapeInfo->XYtoComponentMap.FindRef(FIntPoint(ComponentIndexX, ComponentIndexY)))
+					{
+						LandscapeInfo->UpdateAddCollision(FIntPoint(ComponentIndexX, ComponentIndexY));
+					}
+				}
+
+				// Sides
+				for (ComponentIndexY = ComponentIndexY1; ComponentIndexY <= ComponentIndexY2; ++ComponentIndexY)
+				{
+					// Left
+					int32 ComponentIndexX = ComponentIndexX1 - 1;
+					if (!LandscapeInfo->XYtoComponentMap.FindRef(FIntPoint(ComponentIndexX, ComponentIndexY)))
+					{
+						LandscapeInfo->UpdateAddCollision(FIntPoint(ComponentIndexX, ComponentIndexY));
+					}
+
+					// Right
+					ComponentIndexX = ComponentIndexX1 + 1;
+					if (!LandscapeInfo->XYtoComponentMap.FindRef(FIntPoint(ComponentIndexX, ComponentIndexY)))
+					{
+						LandscapeInfo->UpdateAddCollision(FIntPoint(ComponentIndexX, ComponentIndexY));
+					}
+				}
+
+				// Bottom row
+				ComponentIndexY = ComponentIndexY2 + 1;
+				for (int32 ComponentIndexX = ComponentIndexX1 - 1; ComponentIndexX <= ComponentIndexX2 + 1; ++ComponentIndexX)
+				{
+					if (!LandscapeInfo->XYtoComponentMap.FindRef(FIntPoint(ComponentIndexX, ComponentIndexY)))
+					{
+						LandscapeInfo->UpdateAddCollision(FIntPoint(ComponentIndexX, ComponentIndexY));
+					}
+				}
+			}
+
 			GEngine->BroadcastOnActorMoved(Landscape);
 		}
 	}
 
 protected:
-	FEdModeLandscape* EdMode;
-	ULandscapeInfo* LandscapeInfo;
 	FLandscapeHeightCache HeightCache;
 	FLandscapeXYOffsetCache<true> XYOffsetCache;
 };
@@ -926,7 +955,7 @@ class FLandscapeToolStrokeDeleteComponent : public FLandscapeToolStrokeBase
 {
 public:
 	FLandscapeToolStrokeDeleteComponent(FEdModeLandscape* InEdMode, const FLandscapeToolTarget& InTarget)
-		: LandscapeInfo(InTarget.LandscapeInfo.Get())
+		: FLandscapeToolStrokeBase(InEdMode, InTarget)
 	{
 	}
 
@@ -1085,9 +1114,6 @@ public:
 			GEngine->BroadcastLevelActorListChanged();
 		}
 	}
-
-protected:
-	ULandscapeInfo* LandscapeInfo;
 };
 
 // 
@@ -1113,8 +1139,7 @@ class FLandscapeToolStrokeCopy : public FLandscapeToolStrokeBase
 {
 public:
 	FLandscapeToolStrokeCopy(FEdModeLandscape* InEdMode, const FLandscapeToolTarget& InTarget)
-		: EdMode(InEdMode)
-		, LandscapeInfo(InTarget.LandscapeInfo.Get())
+		: FLandscapeToolStrokeBase(InEdMode, InTarget)
 		, Cache(InTarget)
 		, HeightCache(InTarget)
 		, WeightCache(InTarget)
@@ -1409,8 +1434,6 @@ public:
 	}
 
 protected:
-	FEdModeLandscape* EdMode;
-	ULandscapeInfo* LandscapeInfo;
 	typename ToolTarget::CacheClass Cache;
 	FLandscapeHeightCache HeightCache;
 	FLandscapeFullWeightCache WeightCache;
@@ -1471,8 +1494,7 @@ class FLandscapeToolStrokePaste : public FLandscapeToolStrokeBase
 {
 public:
 	FLandscapeToolStrokePaste(FEdModeLandscape* InEdMode, const FLandscapeToolTarget& InTarget)
-		: EdMode(InEdMode)
-		, LandscapeInfo(InTarget.LandscapeInfo.Get())
+		: FLandscapeToolStrokeBase(InEdMode, InTarget)
 		, Cache(InTarget)
 		, HeightCache(InTarget)
 		, WeightCache(InTarget)
@@ -1727,8 +1749,6 @@ public:
 	}
 
 protected:
-	FEdModeLandscape* EdMode;
-	ULandscapeInfo* LandscapeInfo;
 	typename ToolTarget::CacheClass Cache;
 	FLandscapeHeightCache HeightCache;
 	FLandscapeFullWeightCache WeightCache;
diff --git a/Engine/Source/Editor/LandscapeEditor/Private/LandscapeEdModeErosionTools.cpp b/Engine/Source/Editor/LandscapeEditor/Private/LandscapeEdModeErosionTools.cpp
--- a/Engine/Source/Editor/LandscapeEditor/Private/LandscapeEdModeErosionTools.cpp
+++ b/Engine/Source/Editor/LandscapeEditor/Private/LandscapeEdModeErosionTools.cpp
@@ -21,7 +21,7 @@ class FLandscapeToolStrokeErosionBase : public FLandscapeToolStrokeBase
 {
 public:
 	FLandscapeToolStrokeErosionBase(FEdModeLandscape* InEdMode, const FLandscapeToolTarget& InTarget)
-		: LandscapeInfo(InTarget.LandscapeInfo.Get())
+		: FLandscapeToolStrokeBase(InEdMode, InTarget)
 		, HeightCache(InTarget)
 		, WeightCache(InTarget)
 		, bWeightApplied(InTarget.TargetType != ELandscapeToolTargetType::Heightmap)
@@ -29,7 +29,6 @@ public:
 	}
 
 protected:
-	ULandscapeInfo* LandscapeInfo;
 	FLandscapeHeightCache HeightCache;
 	FLandscapeFullWeightCache WeightCache;
 	bool bWeightApplied;
diff --git a/Engine/Source/Editor/LandscapeEditor/Private/LandscapeEdModePaintTools.cpp b/Engine/Source/Editor/LandscapeEditor/Private/LandscapeEdModePaintTools.cpp
--- a/Engine/Source/Editor/LandscapeEditor/Private/LandscapeEdModePaintTools.cpp
+++ b/Engine/Source/Editor/LandscapeEditor/Private/LandscapeEdModePaintTools.cpp
@@ -59,14 +59,13 @@ class FLandscapeToolStrokePaintBase : public FLandscapeToolStrokeBase
 {
 public:
 	FLandscapeToolStrokePaintBase(FEdModeLandscape* InEdMode, const FLandscapeToolTarget& InTarget)
-		: Cache(InTarget)
-		, LandscapeInfo(InTarget.LandscapeInfo.Get())
+		: FLandscapeToolStrokeBase(InEdMode, InTarget)
+		, Cache(InTarget)
 	{
 	}
 
 protected:
 	typename ToolTarget::CacheClass Cache;
-	ULandscapeInfo* LandscapeInfo;
 };
 
 
@@ -613,7 +612,7 @@ public:
 };
 
 template<class ToolTarget>
-class FLandscapeToolFlatten : public FLandscapeToolPaintBase<ToolTarget, FLandscapeToolStrokeFlatten<ToolTarget>>
+class FLandscapeToolFlatten : public FLandscapeToolPaintBase < ToolTarget, FLandscapeToolStrokeFlatten<ToolTarget> >
 {
 protected:
 	UStaticMesh* PlaneMesh;
@@ -628,6 +627,12 @@ public:
 		check(PlaneMesh);
 	}
 
+	virtual void AddReferencedObjects(FReferenceCollector& Collector) override
+	{
+		Collector.AddReferencedObject(PlaneMesh);
+		Collector.AddReferencedObject(MeshComponent);
+	}
+
 	virtual const TCHAR* GetToolName() override { return TEXT("Flatten"); }
 	virtual FText GetDisplayName() override { return NSLOCTEXT("UnrealEd", "LandscapeMode_Flatten", "Flatten"); };
 
diff --git a/Engine/Source/Editor/LandscapeEditor/Private/LandscapeEdModeRampTool.cpp b/Engine/Source/Editor/LandscapeEditor/Private/LandscapeEdModeRampTool.cpp
--- a/Engine/Source/Editor/LandscapeEditor/Private/LandscapeEdModeRampTool.cpp
+++ b/Engine/Source/Editor/LandscapeEditor/Private/LandscapeEdModeRampTool.cpp
@@ -100,6 +100,11 @@ public:
 		check(SpriteTexture);
 	}
 
+	virtual void AddReferencedObjects(FReferenceCollector& Collector) override
+	{
+		Collector.AddReferencedObject(SpriteTexture);
+	}
+
 	virtual const TCHAR* GetToolName() override { return TEXT("Ramp"); }
 	virtual FText GetDisplayName() override { return NSLOCTEXT("UnrealEd", "LandscapeMode_Ramp", "Ramp"); };
 
diff --git a/Engine/Source/Editor/LandscapeEditor/Private/LandscapeEdModeSplineTools.cpp b/Engine/Source/Editor/LandscapeEditor/Private/LandscapeEdModeSplineTools.cpp
--- a/Engine/Source/Editor/LandscapeEditor/Private/LandscapeEdModeSplineTools.cpp
+++ b/Engine/Source/Editor/LandscapeEditor/Private/LandscapeEdModeSplineTools.cpp
@@ -46,6 +46,14 @@ public:
 		GEditor->RegisterForUndo(this);
 	}
 
+	virtual void AddReferencedObjects(FReferenceCollector& Collector) override
+	{
+		Collector.AddReferencedObject(LandscapeInfo);
+		Collector.AddReferencedObjects(SelectedSplineControlPoints);
+		Collector.AddReferencedObjects(SelectedSplineSegments);
+		Collector.AddReferencedObject(DraggingTangent_Segment);
+	}
+
 	~FLandscapeToolSplines()
 	{
 		// GEditor is invalid at shutdown as the object system is unloaded before the landscape module.
diff --git a/Engine/Source/Editor/LandscapeEditor/Private/LandscapeEdModeTools.h b/Engine/Source/Editor/LandscapeEditor/Private/LandscapeEdModeTools.h
--- a/Engine/Source/Editor/LandscapeEditor/Private/LandscapeEdModeTools.h
+++ b/Engine/Source/Editor/LandscapeEditor/Private/LandscapeEdModeTools.h
@@ -1247,17 +1247,35 @@ struct FWeightmapToolTarget
 };
 
 /**
-* FLandscapeToolStrokeBase - base class for tool strokes (used by FLandscapeToolBase)
-*/
+ * FLandscapeToolStrokeBase - base class for tool strokes (used by FLandscapeToolBase)
+ */
 
-class FLandscapeToolStrokeBase
+class FLandscapeToolStrokeBase : protected FGCObject
 {
 public:
 	// Whether to call Apply() every frame even if the mouse hasn't moved
 	enum { UseContinuousApply = false };
 
-	// Signature of Apply() method:
+	// This is also the expected signature of derived class constructor used by FLandscapeToolBase
+	FLandscapeToolStrokeBase(FEdModeLandscape* InEdMode, const FLandscapeToolTarget& InTarget)
+		: EdMode(InEdMode)
+		, Target(InTarget)
+		, LandscapeInfo(InTarget.LandscapeInfo.Get())
+	{
+	}
+
+	// Signature of Apply() method for derived strokes
 	// void Apply(FEditorViewportClient* ViewportClient, FLandscapeBrush* Brush, const ULandscapeEditorObject* UISettings, const TArray<FLandscapeToolMousePosition>& MousePositions);
+
+	virtual void AddReferencedObjects(FReferenceCollector& Collector) override
+	{
+		Collector.AddReferencedObject(LandscapeInfo);
+	}
+
+protected:
+	FEdModeLandscape* EdMode;
+	const FLandscapeToolTarget& Target;
+	ULandscapeInfo* LandscapeInfo;
 };
 
 
diff --git a/Engine/Source/Editor/LandscapeEditor/Private/LandscapeEdModeXYOffset.cpp b/Engine/Source/Editor/LandscapeEditor/Private/LandscapeEdModeXYOffset.cpp
--- a/Engine/Source/Editor/LandscapeEditor/Private/LandscapeEdModeXYOffset.cpp
+++ b/Engine/Source/Editor/LandscapeEditor/Private/LandscapeEdModeXYOffset.cpp
@@ -54,7 +54,7 @@ class FLandscapeToolStrokeRetopologize : public FLandscapeToolStrokeBase
 {
 public:
 	FLandscapeToolStrokeRetopologize(FEdModeLandscape* InEdMode, const FLandscapeToolTarget& InTarget)
-		: LandscapeInfo(InTarget.LandscapeInfo.Get())
+		: FLandscapeToolStrokeBase(InEdMode, InTarget)
 		, Cache(InTarget)
 	{
 	}
@@ -547,7 +547,6 @@ public:
 	}
 
 protected:
-	ULandscapeInfo* LandscapeInfo;
 	FLandscapeXYOffsetCache<false> Cache;
 };
 
diff --git a/Engine/Source/Editor/LandscapeEditor/Public/LandscapeEdMode.h b/Engine/Source/Editor/LandscapeEditor/Public/LandscapeEdMode.h
--- a/Engine/Source/Editor/LandscapeEditor/Public/LandscapeEdMode.h
+++ b/Engine/Source/Editor/LandscapeEditor/Public/LandscapeEdMode.h
@@ -221,7 +221,7 @@ enum class ELandscapeToolType
 /**
  * FLandscapeTool
  */
-class FLandscapeTool
+class FLandscapeTool : public FGCObject
 {
 public:
 	virtual void EnterTool() {}
@@ -265,6 +265,10 @@ public:
 	virtual ELandscapeToolType GetToolType() { return ELandscapeToolType::Normal; }
 	virtual ELandscapeToolTargetTypeMask::Type GetSupportedTargetTypes() { return ELandscapeToolTargetTypeMask::NA; };
 
+	// FGCObject interface
+	virtual void AddReferencedObjects(FReferenceCollector& Collector) override {}
+
+public:
 	int32					PreviousBrushIndex;
 	TArray<FName>			ValidBrushes;
 };
