Fixed smooth tool exploding when tool strength >1

diff --git a/Engine/Source/Editor/LandscapeEditor/Private/LandscapeEdModePaintTools.cpp b/Engine/Source/Editor/LandscapeEditor/Private/LandscapeEdModePaintTools.cpp
--- a/Engine/Source/Editor/LandscapeEditor/Private/LandscapeEdModePaintTools.cpp
+++ b/Engine/Source/Editor/LandscapeEditor/Private/LandscapeEdModePaintTools.cpp
@@ -384,10 +384,12 @@ public:
 		TArray<typename ToolTarget::CacheClass::DataType> Data;
 		this->Cache.GetCachedData(X1, Y1, X2, Y2, Data);
 
+		const float ToolStrength = FMath::Clamp<float>(UISettings->ToolStrength * Pressure, 0.0f, 1.0f);
+
 		// Apply the brush
 		if (UISettings->bDetailSmooth)
 		{
-			LowPassFilter<typename ToolTarget::CacheClass::DataType>(X1, Y1, X2, Y2, BrushInfo, Data, UISettings->DetailScale, UISettings->ToolStrength * Pressure);
+			LowPassFilter<typename ToolTarget::CacheClass::DataType>(X1, Y1, X2, Y2, BrushInfo, Data, UISettings->DetailScale, ToolStrength);
 		}
 		else
 		{
@@ -433,7 +435,7 @@ public:
 
 						FilterValue /= FilterSamplingNumber;
 
-						DataScanline[X] = FMath::Lerp(DataScanline[X], (typename ToolTarget::CacheClass::DataType)FilterValue, BrushValue * UISettings->ToolStrength * Pressure);
+						DataScanline[X] = FMath::Lerp(DataScanline[X], (typename ToolTarget::CacheClass::DataType)FilterValue, BrushValue * ToolStrength);
 					}
 				}
 			}
