Fix smooth tool's flattening effect

diff --git a/Engine/Source/Editor/LandscapeEditor/Private/LandscapeEdModePaintTools.cpp b/Engine/Source/Editor/LandscapeEditor/Private/LandscapeEdModePaintTools.cpp
--- a/Engine/Source/Editor/LandscapeEditor/Private/LandscapeEdModePaintTools.cpp
+++ b/Engine/Source/Editor/LandscapeEditor/Private/LandscapeEdModePaintTools.cpp
@@ -379,26 +379,39 @@ public:
 			Y2 += 1;
 		}
 
-		this->Cache.CacheData(X1, Y1, X2, Y2);
-
-		TArray<typename ToolTarget::CacheClass::DataType> Data;
-		this->Cache.GetCachedData(X1, Y1, X2, Y2, Data);
-
 		const float ToolStrength = FMath::Clamp<float>(UISettings->ToolStrength * Pressure, 0.0f, 1.0f);
 
 		// Apply the brush
 		if (UISettings->bDetailSmooth)
 		{
+			this->Cache.CacheData(X1, Y1, X2, Y2);
+
+			TArray<typename ToolTarget::CacheClass::DataType> Data;
+			this->Cache.GetCachedData(X1, Y1, X2, Y2, Data);
+
 			LowPassFilter<typename ToolTarget::CacheClass::DataType>(X1, Y1, X2, Y2, BrushInfo, Data, UISettings->DetailScale, ToolStrength);
+
+			this->Cache.SetCachedData(X1, Y1, X2, Y2, Data, UISettings->PaintingRestriction);
 		}
 		else
 		{
-			int32 FilterRadius = UISettings->SmoothFilterKernelSize;
+			const int32 FilterRadius = UISettings->SmoothFilterKernelSize;
+
+			// Expand the cached area by the filter radius so we can run the filter at full size on the edge of the brush
+			const int32 CacheX1 = X1 - FilterRadius;
+			const int32 CacheY1 = Y1 - FilterRadius;
+			const int32 CacheX2 = X2 + FilterRadius;
+			const int32 CacheY2 = Y2 + FilterRadius;
+
+			this->Cache.CacheData(CacheX1, CacheY1, CacheX2, CacheY2);
+
+			TArray<typename ToolTarget::CacheClass::DataType> Data;
+			this->Cache.GetCachedData(CacheX1, CacheY1, CacheX2, CacheY2, Data);
 
 			for (int32 Y = BrushInfo.GetBounds().Min.Y; Y < BrushInfo.GetBounds().Max.Y; Y++)
 			{
 				const float* BrushScanline = BrushInfo.GetDataPtr(FIntPoint(0, Y));
-				auto* DataScanline = Data.GetData() + (Y - Y1) * (X2 - X1 + 1) + (0 - X1);
+				auto* DataScanline = Data.GetData() + (Y - CacheY1) * (CacheX2 - CacheX1 + 1) + (0 - CacheX1);
 
 				for (int32 X = BrushInfo.GetBounds().Min.X; X < BrushInfo.GetBounds().Max.X; X++)
 				{
@@ -406,30 +419,24 @@ public:
 
 					if (BrushValue > 0.0f)
 					{
+						// needs to be ~12 bits larger than ToolTarget::CacheClass::DataType (for max FilterRadius (31))
+						// the editor is 64-bit native so just go the whole hog :)
 						int64 FilterValue = 0;
-						int32 FilterSamplingNumber = 0;
 
-						// The previous version of this clamped to X1/Y1/X2/Y2 *inside* the loop, which were always expanded by one from the BrushInfo
-						// So the Find on the BrushInfo always gave null for anything that would have been clamped and caused it to do nothing
-						// So now just skip iterating those by clamping first
+						const int32 SampleX1 = X - FilterRadius;
+						const int32 SampleY1 = Y - FilterRadius;
+						const int32 SampleX2 = X + FilterRadius;
+						const int32 SampleY2 = Y + FilterRadius;
 
-						const int32 SampleX1 = FMath::Max<int32>(X - FilterRadius, BrushInfo.GetBounds().Min.X);
-						const int32 SampleY1 = FMath::Max<int32>(Y - FilterRadius, BrushInfo.GetBounds().Min.Y);
-						const int32 SampleX2 = FMath::Min<int32>(X + FilterRadius + 1, BrushInfo.GetBounds().Max.X);
-						const int32 SampleY2 = FMath::Min<int32>(Y + FilterRadius + 1, BrushInfo.GetBounds().Max.Y);
-						for (int32 SampleY = SampleY1; SampleY < SampleY2; SampleY++)
+						const int32 FilterSamplingNumber = FMath::Square(FilterRadius + 1 + FilterRadius);
+
+						for (int32 SampleY = SampleY1; SampleY <= SampleY2; SampleY++)
 						{
-							const float* SampleBrushScanline = BrushInfo.GetDataPtr(FIntPoint(0, SampleY));
-							auto* SampleDataScanline = Data.GetData() + (SampleY - Y1) * (X2 - X1 + 1) + (0 - X1);
+							auto* SampleDataScanline = Data.GetData() + (SampleY - CacheY1) * (CacheX2 - CacheX1 + 1) + (0 - CacheX1);
 
-							for (int32 SampleX = SampleX1; SampleX < SampleX2; SampleX++)
+							for (int32 SampleX = SampleX1; SampleX <= SampleX2; SampleX++)
 							{
-								const float SampleBrushValue = SampleBrushScanline[SampleX];
-								if (SampleBrushValue > 0.0f)
-								{
-									FilterValue += SampleDataScanline[SampleX];
-									FilterSamplingNumber++;
-								}
+								FilterValue += SampleDataScanline[SampleX];
 							}
 						}
 
@@ -439,9 +446,10 @@ public:
 					}
 				}
 			}
-		}
 
-		this->Cache.SetCachedData(X1, Y1, X2, Y2, Data, UISettings->PaintingRestriction);
+			// TODO: only need to set the modified area, not the extra data we cached for the filter, but the interface to this takes an entire TArray
+			this->Cache.SetCachedData(CacheX1, CacheY1, CacheX2, CacheY2, Data, UISettings->PaintingRestriction);
+		}
 		this->Cache.Flush();
 	}
 };
diff --git a/Engine/Source/Editor/LandscapeEditor/Private/LandscapeEdModeTools.h b/Engine/Source/Editor/LandscapeEditor/Private/LandscapeEdModeTools.h
--- a/Engine/Source/Editor/LandscapeEditor/Private/LandscapeEdModeTools.h
+++ b/Engine/Source/Editor/LandscapeEditor/Private/LandscapeEdModeTools.h
@@ -504,6 +504,8 @@ public:
 	// X2/Y2 Coordinates are "inclusive" max values
 	void SetCachedData(int32 X1, int32 Y1, int32 X2, int32 Y2, TArray<AccessorType>& Data, ELandscapeLayerPaintingRestriction::Type PaintingRestriction = ELandscapeLayerPaintingRestriction::None)
 	{
+		checkSlow(Data.Num() == (1 + Y2 - Y1) * (1 + X2 - X1));
+
 		// Update cache
 		for (int32 Y = Y1; Y <= Y2; Y++)
 		{
