Fixed: Can not undo Clear Region Selection action UE-18038

diff --git a/Engine/Source/Editor/LandscapeEditor/Private/LandscapeEditorDetailCustomization_MiscTools.cpp b/Engine/Source/Editor/LandscapeEditor/Private/LandscapeEditorDetailCustomization_MiscTools.cpp
--- a/Engine/Source/Editor/LandscapeEditor/Private/LandscapeEditorDetailCustomization_MiscTools.cpp
+++ b/Engine/Source/Editor/LandscapeEditor/Private/LandscapeEditorDetailCustomization_MiscTools.cpp
@@ -210,6 +210,8 @@ FReply FLandscapeEditorDetailCustomization_MiscTools::OnClearRegionSelectionButt
 		ULandscapeInfo* LandscapeInfo = LandscapeEdMode->CurrentToolTarget.LandscapeInfo.Get();
 		if (LandscapeInfo)
 		{
+			FScopedTransaction Transaction(LOCTEXT("Region.Undo_ClearSelected", "Clearing Region Selection"));
+			LandscapeInfo->Modify();
 			LandscapeInfo->ClearSelectedRegion(false);
 		}
 	}
