Fixed: Landscape Materials Layers Show "none" but are Still Able to Paint on Landscape UE-17060

diff --git a/Engine/Source/Editor/LandscapeEditor/Private/LandscapeEditorDetailCustomization_TargetLayers.cpp b/Engine/Source/Editor/LandscapeEditor/Private/LandscapeEditorDetailCustomization_TargetLayers.cpp
--- a/Engine/Source/Editor/LandscapeEditor/Private/LandscapeEditorDetailCustomization_TargetLayers.cpp
+++ b/Engine/Source/Editor/LandscapeEditor/Private/LandscapeEditorDetailCustomization_TargetLayers.cpp
@@ -772,7 +772,8 @@ void FLandscapeEditorCustomNodeBuilder_TargetLayers::OnTargetLayerCreateClicked(
 			TArray<UObject*> Objects;
 			Objects.Add(LayerInfo);
 			GEditor->SyncBrowserToObjects(Objects);
-			//LandscapeEdMode->UpdateTargetList();
+			
+			LandscapeEdMode->TargetsListUpdated.Broadcast();
 		}
 	}
 }
