Fixed the Landscape Editor keyboard shortcuts not working

diff --git a/Engine/Source/Editor/LandscapeEditor/Private/LandscapeEdMode.cpp b/Engine/Source/Editor/LandscapeEditor/Private/LandscapeEdMode.cpp
--- a/Engine/Source/Editor/LandscapeEditor/Private/LandscapeEdMode.cpp
+++ b/Engine/Source/Editor/LandscapeEditor/Private/LandscapeEdMode.cpp
@@ -282,6 +282,12 @@ bool FEdModeLandscape::UsesToolkits() const
 	return true;
 }
 
+TSharedRef<FUICommandList> FEdModeLandscape::GetUICommandList() const
+{
+	check(Toolkit.IsValid());
+	return Toolkit->GetToolkitCommands();
+}
+
 /** FEdMode: Called when the mode is entered */
 void FEdModeLandscape::Enter()
 {
diff --git a/Engine/Source/Editor/LandscapeEditor/Private/LandscapeEditorDetails.cpp b/Engine/Source/Editor/LandscapeEditor/Private/LandscapeEditorDetails.cpp
--- a/Engine/Source/Editor/LandscapeEditor/Private/LandscapeEditorDetails.cpp
+++ b/Engine/Source/Editor/LandscapeEditor/Private/LandscapeEditorDetails.cpp
@@ -22,53 +22,8 @@ TSharedRef<IDetailCustomization> FLandscapeEditorDetails::MakeInstance()
 
 void FLandscapeEditorDetails::CustomizeDetails(IDetailLayoutBuilder& DetailBuilder)
 {
-	//FLevelEditorModule& LevelEditorModule = FModuleManager::LoadModuleChecked<FLevelEditorModule>("LevelEditor");
-	//TSharedPtr<SLevelEditor> LevelEditor = StaticCastSharedPtr<SLevelEditor>(LevelEditorModule.GetFirstLevelEditor());
-	//CommandList = LevelEditor->GetLevelEditorActions();
-
-	check(!CommandList.IsValid());
-	CommandList = MakeShareable(new FUICommandList);
-
-	auto NameToCommandMap = FLandscapeEditorCommands::Get().NameToCommandMap;
-
-#define MAP_TOOL(ToolName) CommandList->MapAction(NameToCommandMap.FindChecked("Tool_" ToolName), FUIAction(FExecuteAction::CreateSP(this, &FLandscapeEditorDetails::OnChangeTool, FName(ToolName)), FCanExecuteAction::CreateSP(this, &FLandscapeEditorDetails::IsToolEnabled, FName(ToolName)), FIsActionChecked::CreateStatic(&FLandscapeEditorDetails::IsToolActive, FName(ToolName))));
-	MAP_TOOL("NewLandscape");
-	MAP_TOOL("ResizeLandscape");
-
-	MAP_TOOL("Sculpt");
-	MAP_TOOL("Paint");
-	MAP_TOOL("Smooth");
-	MAP_TOOL("Flatten");
-	MAP_TOOL("Ramp");
-	MAP_TOOL("Erosion");
-	MAP_TOOL("HydraErosion");
-	MAP_TOOL("Noise");
-	MAP_TOOL("Retopologize");
-	MAP_TOOL("Visibility");
-
-	MAP_TOOL("Select");
-	MAP_TOOL("AddComponent");
-	MAP_TOOL("DeleteComponent");
-	MAP_TOOL("MoveToLevel");
-
-	MAP_TOOL("Mask");
-	MAP_TOOL("CopyPaste");
-
-	MAP_TOOL("Splines");
-#undef MAP_TOOL
-#define MAP_BRUSH_SET(BrushSetName) CommandList->MapAction(NameToCommandMap.FindChecked(BrushSetName), FUIAction(FExecuteAction::CreateSP(this, &FLandscapeEditorDetails::OnChangeBrushSet, FName(BrushSetName)), FCanExecuteAction::CreateSP(this, &FLandscapeEditorDetails::IsBrushSetEnabled, FName(BrushSetName)), FIsActionChecked::CreateStatic(&FLandscapeEditorDetails::IsBrushSetActive, FName(BrushSetName))));
-	MAP_BRUSH_SET("BrushSet_Circle");
-	MAP_BRUSH_SET("BrushSet_Alpha");
-	MAP_BRUSH_SET("BrushSet_Pattern");
-	MAP_BRUSH_SET("BrushSet_Component");
-	MAP_BRUSH_SET("BrushSet_Gizmo");
-#undef MAP_BRUSH_SET
-#define MAP_BRUSH(BrushName) CommandList->MapAction(NameToCommandMap.FindChecked(BrushName), FUIAction(FExecuteAction::CreateSP(this, &FLandscapeEditorDetails::OnChangeBrush, FName(BrushName)), FCanExecuteAction(), FIsActionChecked::CreateSP(this, &FLandscapeEditorDetails::IsBrushActive, FName(BrushName))));
-	MAP_BRUSH("Circle_Smooth");
-	MAP_BRUSH("Circle_Linear");
-	MAP_BRUSH("Circle_Spherical");
-	MAP_BRUSH("Circle_Tip");
-#undef MAP_BRUSH
+	FEdModeLandscape* LandscapeEdMode = GetEditorMode();
+	CommandList = LandscapeEdMode->GetUICommandList();
 
 	static const FLinearColor BorderColor = FLinearColor(0.2f, 0.2f, 0.2f, 0.2f);
 	static const FSlateBrush* BorderStyle = FEditorStyle::GetBrush("DetailsView.GroupSection");
@@ -369,29 +324,6 @@ EVisibility FLandscapeEditorDetails::GetToolSelectorVisibility() const
 	return EVisibility::Collapsed;
 }
 
-void FLandscapeEditorDetails::OnChangeTool(FName ToolName)
-{
-	FEdModeLandscape* LandscapeEdMode = GetEditorMode();
-	if (LandscapeEdMode != NULL)
-	{
-		LandscapeEdMode->SetCurrentTool(ToolName);
-	}
-}
-
-bool FLandscapeEditorDetails::IsToolEnabled(FName ToolName) const
-{
-	FEdModeLandscape* LandscapeEdMode = GetEditorMode();
-	if (LandscapeEdMode != NULL)
-	{
-		if (ToolName == "NewLandscape" || LandscapeEdMode->GetLandscapeList().Num() > 0)
-		{
-			return true;
-		}
-	}
-
-	return false;
-}
-
 FText FLandscapeEditorDetails::GetCurrentBrushName() const
 {
 	FEdModeLandscape* LandscapeEdMode = GetEditorMode();
@@ -476,26 +408,6 @@ bool FLandscapeEditorDetails::GetBrushSelectorIsVisible() const
 	return false;
 }
 
-void FLandscapeEditorDetails::OnChangeBrushSet(FName BrushSetName)
-{
-	FEdModeLandscape* LandscapeEdMode = GetEditorMode();
-	if (LandscapeEdMode != NULL)
-	{
-		LandscapeEdMode->SetCurrentBrushSet(BrushSetName);
-	}
-}
-
-bool FLandscapeEditorDetails::IsBrushSetEnabled(FName BrushSetName) const
-{
-	FEdModeLandscape* LandscapeEdMode = GetEditorMode();
-	if (LandscapeEdMode && LandscapeEdMode->CurrentTool)
-	{
-		return LandscapeEdMode->CurrentTool->ValidBrushes.Contains(BrushSetName);
-	}
-
-	return false;
-}
-
 FText FLandscapeEditorDetails::GetCurrentBrushFalloffName() const
 {
 	FEdModeLandscape* LandscapeEdMode = GetEditorMode();
@@ -561,25 +473,4 @@ bool FLandscapeEditorDetails::GetBrushFalloffSelectorIsVisible() const
 	return false;
 }
 
-void FLandscapeEditorDetails::OnChangeBrush(FName BrushName)
-{
-	FEdModeLandscape* LandscapeEdMode = GetEditorMode();
-	if (LandscapeEdMode != NULL)
-	{
-		LandscapeEdMode->SetCurrentBrush(BrushName);
-	}
-}
-
-bool FLandscapeEditorDetails::IsBrushActive(FName BrushName) const
-{
-	FEdModeLandscape* LandscapeEdMode = GetEditorMode();
-	if (LandscapeEdMode && LandscapeEdMode->CurrentBrush)
-	{
-		const FName CurrentBrushName = LandscapeEdMode->CurrentBrush->GetBrushName();
-		return CurrentBrushName == BrushName;
-	}
-
-	return false;
-}
-
 #undef LOCTEXT_NAMESPACE
diff --git a/Engine/Source/Editor/LandscapeEditor/Private/LandscapeEditorDetails.h b/Engine/Source/Editor/LandscapeEditor/Private/LandscapeEditorDetails.h
--- a/Engine/Source/Editor/LandscapeEditor/Private/LandscapeEditorDetails.h
+++ b/Engine/Source/Editor/LandscapeEditor/Private/LandscapeEditorDetails.h
@@ -35,22 +35,16 @@ protected:
 	TSharedRef<SWidget> GetToolSelector();
 	bool GetToolSelectorIsVisible() const;
 	EVisibility GetToolSelectorVisibility() const;
-	void OnChangeTool(FName ToolName);
-	bool IsToolEnabled(FName ToolName) const;
 
 	FText GetCurrentBrushName() const;
 	FSlateIcon GetCurrentBrushIcon() const;
 	TSharedRef<SWidget> GetBrushSelector();
 	bool GetBrushSelectorIsVisible() const;
-	void OnChangeBrushSet(FName BrushSetName);
-	bool IsBrushSetEnabled(FName BrushSetName) const;
 
 	FText GetCurrentBrushFalloffName() const;
 	FSlateIcon GetCurrentBrushFalloffIcon() const;
 	TSharedRef<SWidget> GetBrushFalloffSelector();
 	bool GetBrushFalloffSelectorIsVisible() const;
-	void OnChangeBrush(FName BrushName);
-	bool IsBrushActive(FName BrushName) const;
 
 	TSharedPtr<FUICommandList> CommandList;
 
diff --git a/Engine/Source/Editor/LandscapeEditor/Private/SLandscapeEditor.cpp b/Engine/Source/Editor/LandscapeEditor/Private/SLandscapeEditor.cpp
--- a/Engine/Source/Editor/LandscapeEditor/Private/SLandscapeEditor.cpp
+++ b/Engine/Source/Editor/LandscapeEditor/Private/SLandscapeEditor.cpp
@@ -60,16 +60,68 @@ void SLandscapeAssetThumbnail::SetAsset(UObject* Asset)
 
 //////////////////////////////////////////////////////////////////////////
 
-void FLandscapeToolKit::RegisterTabSpawners(const TSharedRef<class FTabManager>& TabManager)
+void FLandscapeToolKit::RegisterTabSpawners(const TSharedRef<FTabManager>& TabManager)
 {
 }
 
-void FLandscapeToolKit::UnregisterTabSpawners(const TSharedRef<class FTabManager>& TabManager)
+void FLandscapeToolKit::UnregisterTabSpawners(const TSharedRef<FTabManager>& TabManager)
 {
 }
 
-void FLandscapeToolKit::Init(const TSharedPtr< class IToolkitHost >& InitToolkitHost)
+void FLandscapeToolKit::Init(const TSharedPtr<IToolkitHost>& InitToolkitHost)
 {
+	auto NameToCommandMap = FLandscapeEditorCommands::Get().NameToCommandMap;
+
+	FEdModeLandscape* LandscapeEdMode = GetEditorMode();
+	TSharedRef<FUICommandList> CommandList = LandscapeEdMode->GetUICommandList();
+
+#define MAP_MODE(ModeName) CommandList->MapAction(NameToCommandMap.FindChecked(ModeName), FUIAction(FExecuteAction::CreateSP(this, &FLandscapeToolKit::OnChangeMode, FName(ModeName)), FCanExecuteAction::CreateSP(this, &FLandscapeToolKit::IsModeEnabled, FName(ModeName)), FIsActionChecked::CreateSP(this, &FLandscapeToolKit::IsModeActive, FName(ModeName))));
+	MAP_MODE("ToolMode_Manage");
+	MAP_MODE("ToolMode_Sculpt");
+	MAP_MODE("ToolMode_Paint");
+#undef MAP_MODE
+
+#define MAP_TOOL(ToolName) CommandList->MapAction(NameToCommandMap.FindChecked("Tool_" ToolName), FUIAction(FExecuteAction::CreateSP(this, &FLandscapeToolKit::OnChangeTool, FName(ToolName)), FCanExecuteAction::CreateSP(this, &FLandscapeToolKit::IsToolEnabled, FName(ToolName)), FIsActionChecked::CreateSP(this, &FLandscapeToolKit::IsToolActive, FName(ToolName))));
+	MAP_TOOL("NewLandscape");
+	MAP_TOOL("ResizeLandscape");
+
+	MAP_TOOL("Sculpt");
+	MAP_TOOL("Paint");
+	MAP_TOOL("Smooth");
+	MAP_TOOL("Flatten");
+	MAP_TOOL("Ramp");
+	MAP_TOOL("Erosion");
+	MAP_TOOL("HydraErosion");
+	MAP_TOOL("Noise");
+	MAP_TOOL("Retopologize");
+	MAP_TOOL("Visibility");
+
+	MAP_TOOL("Select");
+	MAP_TOOL("AddComponent");
+	MAP_TOOL("DeleteComponent");
+	MAP_TOOL("MoveToLevel");
+
+	MAP_TOOL("Mask");
+	MAP_TOOL("CopyPaste");
+
+	MAP_TOOL("Splines");
+#undef MAP_TOOL
+
+#define MAP_BRUSH_SET(BrushSetName) CommandList->MapAction(NameToCommandMap.FindChecked(BrushSetName), FUIAction(FExecuteAction::CreateSP(this, &FLandscapeToolKit::OnChangeBrushSet, FName(BrushSetName)), FCanExecuteAction::CreateSP(this, &FLandscapeToolKit::IsBrushSetEnabled, FName(BrushSetName)), FIsActionChecked::CreateSP(this, &FLandscapeToolKit::IsBrushSetActive, FName(BrushSetName))));
+	MAP_BRUSH_SET("BrushSet_Circle");
+	MAP_BRUSH_SET("BrushSet_Alpha");
+	MAP_BRUSH_SET("BrushSet_Pattern");
+	MAP_BRUSH_SET("BrushSet_Component");
+	MAP_BRUSH_SET("BrushSet_Gizmo");
+#undef MAP_BRUSH_SET
+
+#define MAP_BRUSH(BrushName) CommandList->MapAction(NameToCommandMap.FindChecked(BrushName), FUIAction(FExecuteAction::CreateSP(this, &FLandscapeToolKit::OnChangeBrush, FName(BrushName)), FCanExecuteAction(), FIsActionChecked::CreateSP(this, &FLandscapeToolKit::IsBrushActive, FName(BrushName))));
+	MAP_BRUSH("Circle_Smooth");
+	MAP_BRUSH("Circle_Linear");
+	MAP_BRUSH("Circle_Spherical");
+	MAP_BRUSH("Circle_Tip");
+#undef MAP_BRUSH
+
 	LandscapeEditorWidgets = SNew(SLandscapeEditor, SharedThis(this));
 
 	FModeToolkit::Init(InitToolkitHost);
@@ -85,7 +137,7 @@ FText FLandscapeToolKit::GetBaseToolkitName() const
 	return LOCTEXT("ToolkitName", "Landscape Editor");
 }
 
-class FEdModeLandscape* FLandscapeToolKit::GetEditorMode() const
+FEdModeLandscape* FLandscapeToolKit::GetEditorMode() const
 {
 	return (FEdModeLandscape*)GLevelEditorModeTools().GetActiveMode(FBuiltinEditorModes::EM_Landscape);
 }
@@ -105,22 +157,137 @@ void FLandscapeToolKit::NotifyBrushChanged()
 	LandscapeEditorWidgets->NotifyBrushChanged();
 }
 
+void FLandscapeToolKit::OnChangeMode(FName ModeName)
+{
+	FEdModeLandscape* LandscapeEdMode = GetEditorMode();
+	if (LandscapeEdMode)
+	{
+		LandscapeEdMode->SetCurrentToolMode(ModeName);
+	}
+}
+
+bool FLandscapeToolKit::IsModeEnabled(FName ModeName) const
+{
+	FEdModeLandscape* LandscapeEdMode = GetEditorMode();
+	if (LandscapeEdMode)
+	{
+		// Manage is the only mode enabled if we have no landscape
+		if (ModeName == "ToolMode_Manage" || LandscapeEdMode->GetLandscapeList().Num() > 0)
+		{
+			return true;
+		}
+	}
+
+	return false;
+}
+
+bool FLandscapeToolKit::IsModeActive(FName ModeName) const
+{
+	FEdModeLandscape* LandscapeEdMode = GetEditorMode();
+	if (LandscapeEdMode && LandscapeEdMode->CurrentTool)
+	{
+		return LandscapeEdMode->CurrentToolMode->ToolModeName == ModeName;
+	}
+
+	return false;
+}
+
+void FLandscapeToolKit::OnChangeTool(FName ToolName)
+{
+	FEdModeLandscape* LandscapeEdMode = GetEditorMode();
+	if (LandscapeEdMode != nullptr)
+	{
+		LandscapeEdMode->SetCurrentTool(ToolName);
+	}
+}
+
+bool FLandscapeToolKit::IsToolEnabled(FName ToolName) const
+{
+	FEdModeLandscape* LandscapeEdMode = GetEditorMode();
+	if (LandscapeEdMode != nullptr)
+	{
+		if (ToolName == "NewLandscape" || LandscapeEdMode->GetLandscapeList().Num() > 0)
+		{
+			return true;
+		}
+	}
+
+	return false;
+}
+
+bool FLandscapeToolKit::IsToolActive(FName ToolName) const
+{
+	FEdModeLandscape* LandscapeEdMode = GetEditorMode();
+	if (LandscapeEdMode != nullptr && LandscapeEdMode->CurrentTool != nullptr)
+	{
+		const FName CurrentToolName = LandscapeEdMode->CurrentTool->GetToolName();
+		return CurrentToolName == ToolName;
+	}
+
+	return false;
+}
+
+void FLandscapeToolKit::OnChangeBrushSet(FName BrushSetName)
+{
+	FEdModeLandscape* LandscapeEdMode = GetEditorMode();
+	if (LandscapeEdMode != nullptr)
+	{
+		LandscapeEdMode->SetCurrentBrushSet(BrushSetName);
+	}
+}
+
+bool FLandscapeToolKit::IsBrushSetEnabled(FName BrushSetName) const
+{
+	FEdModeLandscape* LandscapeEdMode = GetEditorMode();
+	if (LandscapeEdMode && LandscapeEdMode->CurrentTool)
+	{
+		return LandscapeEdMode->CurrentTool->ValidBrushes.Contains(BrushSetName);
+	}
+
+	return false;
+}
+
+bool FLandscapeToolKit::IsBrushSetActive(FName BrushSetName) const
+{
+	FEdModeLandscape* LandscapeEdMode = GetEditorMode();
+	if (LandscapeEdMode && LandscapeEdMode->CurrentBrushSetIndex >= 0)
+	{
+		const FName CurrentBrushSetName = LandscapeEdMode->LandscapeBrushSets[LandscapeEdMode->CurrentBrushSetIndex].BrushSetName;
+		return CurrentBrushSetName == BrushSetName;
+	}
+
+	return false;
+}
+
+void FLandscapeToolKit::OnChangeBrush(FName BrushName)
+{
+	FEdModeLandscape* LandscapeEdMode = GetEditorMode();
+	if (LandscapeEdMode != nullptr)
+	{
+		LandscapeEdMode->SetCurrentBrush(BrushName);
+	}
+}
+
+bool FLandscapeToolKit::IsBrushActive(FName BrushName) const
+{
+	FEdModeLandscape* LandscapeEdMode = GetEditorMode();
+	if (LandscapeEdMode && LandscapeEdMode->CurrentBrush)
+	{
+		const FName CurrentBrushName = LandscapeEdMode->CurrentBrush->GetBrushName();
+		return CurrentBrushName == BrushName;
+	}
+
+	return false;
+}
+
 //////////////////////////////////////////////////////////////////////////
 
 BEGIN_SLATE_FUNCTION_BUILD_OPTIMIZATION
 void SLandscapeEditor::Construct(const FArguments& InArgs, TSharedRef<FLandscapeToolKit> InParentToolkit)
 {
-	CommandList = InParentToolkit->GetToolkitCommands();
-
-	//FLevelEditorModule& LevelEditorModule = FModuleManager::LoadModuleChecked<FLevelEditorModule>("LevelEditor");
-	//TSharedPtr<SLevelEditor> LevelEditor = StaticCastSharedPtr<SLevelEditor>(LevelEditorModule.GetFirstLevelEditor());
-	//CommandList = LevelEditor->GetLevelEditorActions();
+	TSharedRef<FUICommandList> CommandList = InParentToolkit->GetToolkitCommands();
 
 	// Modes:
-	CommandList->MapAction(FLandscapeEditorCommands::Get().ManageMode, FUIAction(FExecuteAction::CreateSP(this, &SLandscapeEditor::OnChangeMode, FName("ToolMode_Manage")), FCanExecuteAction::CreateSP(this, &SLandscapeEditor::IsModeEnabled, FName(TEXT("ToolMode_Manage"))), FIsActionChecked::CreateSP(this, &SLandscapeEditor::IsModeActive, FName(TEXT("ToolMode_Manage")))));
-	CommandList->MapAction(FLandscapeEditorCommands::Get().SculptMode, FUIAction(FExecuteAction::CreateSP(this, &SLandscapeEditor::OnChangeMode, FName("ToolMode_Sculpt")), FCanExecuteAction::CreateSP(this, &SLandscapeEditor::IsModeEnabled, FName(TEXT("ToolMode_Sculpt"))), FIsActionChecked::CreateSP(this, &SLandscapeEditor::IsModeActive, FName(TEXT("ToolMode_Sculpt")))));
-	CommandList->MapAction(FLandscapeEditorCommands::Get().PaintMode,  FUIAction(FExecuteAction::CreateSP(this, &SLandscapeEditor::OnChangeMode, FName("ToolMode_Paint" )), FCanExecuteAction::CreateSP(this, &SLandscapeEditor::IsModeEnabled, FName(TEXT("ToolMode_Paint" ))), FIsActionChecked::CreateSP(this, &SLandscapeEditor::IsModeActive, FName(TEXT("ToolMode_Paint" )))));
-
 	FToolBarBuilder ModeSwitchButtons(CommandList, FMultiBoxCustomization::None);
 	{
 		ModeSwitchButtons.AddToolBarButton(FLandscapeEditorCommands::Get().ManageMode, NAME_None, LOCTEXT("Mode.Manage", "Manage"), LOCTEXT("Mode.Manage.Tooltip", "Contains tools to add a new landscape, import/export landscape, add/remove components and manage streaming"));
@@ -192,7 +359,7 @@ void SLandscapeEditor::Construct(const FArguments& InArgs, TSharedRef<FLandscape
 }
 END_SLATE_FUNCTION_BUILD_OPTIMIZATION
 
-class FEdModeLandscape* SLandscapeEditor::GetEditorMode() const
+FEdModeLandscape* SLandscapeEditor::GetEditorMode() const
 {
 	return (FEdModeLandscape*)GLevelEditorModeTools().GetActiveMode(FBuiltinEditorModes::EM_Landscape);
 }
@@ -271,7 +438,7 @@ bool SLandscapeEditor::GetIsPropertyVisible(const FPropertyAndParent& PropertyAn
 	const UProperty& Property = PropertyAndParent.Property;
 
 	FEdModeLandscape* LandscapeEdMode = GetEditorMode();
-	if (LandscapeEdMode != NULL && LandscapeEdMode->CurrentTool != NULL)
+	if (LandscapeEdMode != nullptr && LandscapeEdMode->CurrentTool != nullptr)
 	{
 		if (Property.HasMetaData("ShowForMask"))
 		{
@@ -330,41 +497,6 @@ bool SLandscapeEditor::GetIsPropertyVisible(const FPropertyAndParent& PropertyAn
 	return false;
 }
 
-void SLandscapeEditor::OnChangeMode(FName ModeName)
-{
-	FEdModeLandscape* LandscapeEdMode = GetEditorMode();
-	if (LandscapeEdMode)
-	{
-		LandscapeEdMode->SetCurrentToolMode(ModeName);
-	}
-}
-
-bool SLandscapeEditor::IsModeEnabled(FName ModeName) const
-{
-	FEdModeLandscape* LandscapeEdMode = GetEditorMode();
-	if (LandscapeEdMode)
-	{
-		// Manage is the only mode enabled if we have no landscape
-		if (ModeName == "ToolMode_Manage" || LandscapeEdMode->GetLandscapeList().Num() > 0)
-		{
-			return true;
-		}
-	}
-
-	return false;
-}
-
-bool SLandscapeEditor::IsModeActive(FName ModeName) const
-{
-	FEdModeLandscape* LandscapeEdMode = GetEditorMode();
-	if (LandscapeEdMode && LandscapeEdMode->CurrentTool)
-	{
-		return LandscapeEdMode->CurrentToolMode->ToolModeName == ModeName;
-	}
-
-	return false;
-}
-
 void SLandscapeEditor::NotifyToolChanged()
 {
 	FEdModeLandscape* LandscapeEdMode = GetEditorMode();
diff --git a/Engine/Source/Editor/LandscapeEditor/Private/SLandscapeEditor.h b/Engine/Source/Editor/LandscapeEditor/Private/SLandscapeEditor.h
--- a/Engine/Source/Editor/LandscapeEditor/Private/SLandscapeEditor.h
+++ b/Engine/Source/Editor/LandscapeEditor/Private/SLandscapeEditor.h
@@ -5,6 +5,13 @@
 #include "Editor/UnrealEd/Public/Toolkits/BaseToolkit.h"
 #include "SCompoundWidget.h"
 
+// Forward declarations
+class FAssetThumbnailPool;
+class FAssetThumbnail;
+class SLandscapeEditor;
+struct FPropertyAndParent;
+class IDetailsView;
+
 /**
  * Slate widget wrapping an FAssetThumbnail and Viewport
  */
@@ -16,7 +23,7 @@ public:
 		SLATE_ARGUMENT( FIntPoint, ThumbnailSize )
 	SLATE_END_ARGS()
 
-	void Construct(const FArguments& InArgs, UObject* Asset, TSharedRef<class FAssetThumbnailPool> ThumbnailPool);
+	void Construct(const FArguments& InArgs, UObject* Asset, TSharedRef<FAssetThumbnailPool> ThumbnailPool);
 	~SLandscapeAssetThumbnail();
 
 	void SetAsset(UObject* Asset);
@@ -24,7 +31,7 @@ public:
 private:
 	void OnMaterialCompilationFinished(UMaterialInterface* MaterialInterface);
 
-	TSharedPtr<class FAssetThumbnail> AssetThumbnail;
+	TSharedPtr<FAssetThumbnail> AssetThumbnail;
 };
 
 /**
@@ -42,15 +49,31 @@ public:
 	/** IToolkit interface */
 	virtual FName GetToolkitFName() const override;
 	virtual FText GetBaseToolkitName() const override;
-	virtual class FEdModeLandscape* GetEditorMode() const override;
-	virtual TSharedPtr<class SWidget> GetInlineContent() const override;
+	virtual FEdModeLandscape* GetEditorMode() const override;
+	virtual TSharedPtr<SWidget> GetInlineContent() const override;
 
 	void NotifyToolChanged();
 	void NotifyBrushChanged();
 
+protected:
+	void OnChangeMode(FName ModeName);
+	bool IsModeEnabled(FName ModeName) const;
+	bool IsModeActive(FName ModeName) const;
+
+	void OnChangeTool(FName ToolName);
+	bool IsToolEnabled(FName ToolName) const;
+	bool IsToolActive(FName ToolName) const;
+
+	void OnChangeBrushSet(FName BrushSetName);
+	bool IsBrushSetEnabled(FName BrushSetName) const;
+	bool IsBrushSetActive(FName BrushSetName) const;
+
+	void OnChangeBrush(FName BrushName);
+	bool IsBrushActive(FName BrushName) const;
+
 private:
 	/** Geometry tools widget */
-	TSharedPtr<class SLandscapeEditor> LandscapeEditorWidgets;
+	TSharedPtr<SLandscapeEditor> LandscapeEditorWidgets;
 };
 
 /**
@@ -74,16 +97,10 @@ protected:
 
 	bool GetLandscapeEditorIsEnabled() const;
 
-	bool GetIsPropertyVisible(const struct FPropertyAndParent& PropertyAndParent) const;
-
-	void OnChangeMode(FName ModeName);
-	bool IsModeEnabled(FName ModeName) const;
-	bool IsModeActive(FName ModeName) const;
+	bool GetIsPropertyVisible(const FPropertyAndParent& PropertyAndParent) const;
 
 protected:
-	// Command list bound to this window
-	TSharedPtr<FUICommandList> CommandList;
 	TSharedPtr<SErrorText> Error;
 
-	TSharedPtr<class IDetailsView> DetailsPanel;
+	TSharedPtr<IDetailsView> DetailsPanel;
 };
diff --git a/Engine/Source/Editor/LandscapeEditor/Public/LandscapeEdMode.h b/Engine/Source/Editor/LandscapeEditor/Public/LandscapeEdMode.h
--- a/Engine/Source/Editor/LandscapeEditor/Public/LandscapeEdMode.h
+++ b/Engine/Source/Editor/LandscapeEditor/Public/LandscapeEdMode.h
@@ -542,6 +542,8 @@ public:
 
 	virtual bool UsesToolkits() const override;
 
+	TSharedRef<FUICommandList> GetUICommandList() const;
+
 	/** FEdMode: Called when the mode is entered */
 	virtual void Enter() override;
 
diff --git a/Engine/Source/Editor/LevelEditor/Private/SLevelEditor.cpp b/Engine/Source/Editor/LevelEditor/Private/SLevelEditor.cpp
--- a/Engine/Source/Editor/LevelEditor/Private/SLevelEditor.cpp
+++ b/Engine/Source/Editor/LevelEditor/Private/SLevelEditor.cpp
@@ -1312,58 +1312,21 @@ FReply SLevelEditor::OnKeyDown( const FGeometry& MyGeometry, const FKeyEvent& In
 {
 	// Check to see if any of the actions for the level editor can be processed by the current event
 	// If we are in debug mode do not process commands
-	if( FSlateApplication::Get().IsNormalExecution() )
+	if (FSlateApplication::Get().IsNormalExecution())
 	{
-		// Figure out if any of our toolkit's tabs is the active tab.  This is important because we want
-		// the toolkit to have it's own keybinds (which may overlap the level editor's keybinds or any
-		// other toolkit).  When a toolkit tab is active, we give that toolkit a chance to process
-		// commands instead of the level editor.
-		TSharedPtr< IToolkit > ActiveToolkit;
+		for (const auto& ActiveToolkit : HostedToolkits)
 		{
-			const TSharedPtr<SDockableTab> CurrentActiveTab;// = FSlateApplication::xxxGetGlobalTabManager()->GetActiveTab();
-
-			for( auto HostedToolkitIt = HostedToolkits.CreateConstIterator(); HostedToolkitIt && !ActiveToolkit.IsValid(); ++HostedToolkitIt )
-			{
-				const auto& CurToolkit = *HostedToolkitIt;
-				if( CurToolkit.IsValid() )
-				{
-					// Iterate over this toolkits spawned tabs
-					const auto& ToolkitTabsInSpots = CurToolkit->GetToolkitTabsInSpots();
-
-					for( auto CurSpotIt( ToolkitTabsInSpots.CreateConstIterator() ); CurSpotIt && !ActiveToolkit.IsValid(); ++CurSpotIt )
-					{
-						const auto& TabsForSpot = CurSpotIt.Value();
-						for( auto CurTabIt( TabsForSpot.CreateConstIterator() ); CurTabIt; ++CurTabIt )
-						{
-							const auto& PinnedTab = CurTabIt->Pin();
-							if( PinnedTab.IsValid() )
-							{
-								if( PinnedTab == CurrentActiveTab )
-								{
-									ActiveToolkit = CurToolkit;
-								}
-							}
-						}
-					}
-				}
-			}
-		}
-
-		if( ActiveToolkit.IsValid() )
-		{
-			// A toolkit tab is active, so direct all command processing to it
-			if( ActiveToolkit->ProcessCommandBindings( InKeyEvent ) )
+			// A toolkit is active, so direct all command processing to it
+			if (ActiveToolkit->ProcessCommandBindings(InKeyEvent))
 			{
 				return FReply::Handled();
 			}
 		}
-		else
+
+		// No toolkit processed the key, so let the level editor have a chance at the keystroke
+		if (LevelEditorCommands->ProcessCommandBindings(InKeyEvent))
 		{
-			// No toolkit tab is active, so let the level editor have a chance at the keystroke
-			if( LevelEditorCommands->ProcessCommandBindings( InKeyEvent ) )
-			{
-				return FReply::Handled();
-			}
+			return FReply::Handled();
 		}
 	}
 	
@@ -1436,10 +1399,10 @@ TSharedPtr< class FAssetThumbnailPool > SLevelEditor::GetThumbnailPool() const
 	return ThumbnailPool;
 }
 
- void SLevelEditor::AppendCommands( const TSharedRef<FUICommandList>& InCommandsToAppend )
- {
-	 LevelEditorCommands->Append(InCommandsToAppend);
- }
+void SLevelEditor::AppendCommands( const TSharedRef<FUICommandList>& InCommandsToAppend )
+{
+	LevelEditorCommands->Append(InCommandsToAppend);
+}
 
 UWorld* SLevelEditor::GetWorld() const
 {
