Moved implementation details of the Landscape Editor into the Private folder

diff --git a/Engine/Source/Editor/DetailCustomizations/Private/DetailCustomizations.cpp b/Engine/Source/Editor/DetailCustomizations/Private/DetailCustomizations.cpp
--- a/Engine/Source/Editor/DetailCustomizations/Private/DetailCustomizations.cpp
+++ b/Engine/Source/Editor/DetailCustomizations/Private/DetailCustomizations.cpp
@@ -28,7 +28,6 @@
 #include "AnimGraphDefinitions.h"
 #include "SoundDefinitions.h"
 #include "TransitionPoseEvaluatorNodeDetails.h"
-#include "LandscapeSplineDetails.h"
 #include "MathStructCustomizations.h"
 #include "MathStructProxyCustomizations.h"
 #include "RangeStructCustomization.h"
@@ -259,9 +258,6 @@ void FDetailCustomizationsModule::RegisterObjectCustomizations()
 	RegisterCustomClassLayout("AnimStateTransitionNode", FOnGetDetailCustomizationInstance::CreateStatic(&FAnimTransitionNodeDetails::MakeInstance));
 	RegisterCustomClassLayout("AnimGraphNode_TransitionPoseEvaluator", FOnGetDetailCustomizationInstance::CreateStatic(&FTransitionPoseEvaluatorNodeDetails::MakeInstance));
 
-	RegisterCustomClassLayout("LandscapeSplineControlPoint", FOnGetDetailCustomizationInstance::CreateStatic(&FLandscapeSplineDetails::MakeInstance));
-	RegisterCustomClassLayout("LandscapeSplineSegment", FOnGetDetailCustomizationInstance::CreateStatic(&FLandscapeSplineDetails::MakeInstance));
-
 	RegisterCustomClassLayout("SoundWave", FOnGetDetailCustomizationInstance::CreateStatic(&FSoundWaveDetails::MakeInstance));
 	RegisterCustomClassLayout("DialogueWave", FOnGetDetailCustomizationInstance::CreateStatic(&FDialogueWaveDetails::MakeInstance));
 	RegisterCustomClassLayout("BodySetup", FOnGetDetailCustomizationInstance::CreateStatic(&FBodySetupDetails::MakeInstance));
diff --git a/Engine/Source/Editor/DetailCustomizations/Private/LandscapeSplineDetails.cpp b/Engine/Source/Editor/DetailCustomizations/Private/LandscapeSplineDetails.cpp
deleted file mode 100644
diff --git a/Engine/Source/Editor/DetailCustomizations/Private/LandscapeSplineDetails.h b/Engine/Source/Editor/DetailCustomizations/Private/LandscapeSplineDetails.h
deleted file mode 100644
diff --git a/Engine/Source/Editor/LandscapeEditor/Classes/ActorFactoryLandscape.h b/Engine/Source/Editor/LandscapeEditor/Classes/ActorFactoryLandscape.h
deleted file mode 100644
diff --git a/Engine/Source/Editor/LandscapeEditor/Classes/LandscapeEditorObject.h b/Engine/Source/Editor/LandscapeEditor/Classes/LandscapeEditorObject.h
deleted file mode 100644
diff --git a/Engine/Source/Editor/LandscapeEditor/Classes/LandscapePlaceholder.h b/Engine/Source/Editor/LandscapeEditor/Classes/LandscapePlaceholder.h
deleted file mode 100644
diff --git a/Engine/Source/Editor/LandscapeEditor/Private/ActorFactoryLandscape.cpp b/Engine/Source/Editor/LandscapeEditor/Private/ActorFactoryLandscape.cpp
--- a/Engine/Source/Editor/LandscapeEditor/Private/ActorFactoryLandscape.cpp
+++ b/Engine/Source/Editor/LandscapeEditor/Private/ActorFactoryLandscape.cpp
@@ -1,8 +1,8 @@
 // Copyright 1998-2015 Epic Games, Inc. All Rights Reserved.
 
 #include "LandscapeEditorPrivatePCH.h"
-#include "ActorFactoryLandscape.h"
-#include "LandscapePlaceholder.h"
+#include "Classes/ActorFactoryLandscape.h"
+#include "Classes/LandscapePlaceholder.h"
 
 #define LOCTEXT_NAMESPACE "Landscape"
 
diff --git a/Engine/Source/Editor/LandscapeEditor/Private/Classes/ActorFactoryLandscape.h b/Engine/Source/Editor/LandscapeEditor/Private/Classes/ActorFactoryLandscape.h
new file mode 100644
--- /dev/null
+++ b/Engine/Source/Editor/LandscapeEditor/Private/Classes/ActorFactoryLandscape.h
@@ -0,0 +1,19 @@
+// Copyright 1998-2015 Epic Games, Inc. All Rights Reserved.
+
+#pragma once
+
+#include "ActorFactoryLandscape.generated.h"
+
+UCLASS(MinimalAPI, config=Editor)
+class UActorFactoryLandscape : public UActorFactory
+{
+	GENERATED_UCLASS_BODY()
+
+	// Begin UActorFactory Interface
+	//virtual bool CanCreateActorFrom( const FAssetData& AssetData, FText& OutErrorMsg ) override;
+	virtual AActor* SpawnActor( UObject* Asset, ULevel* InLevel, const FVector& Location, const FRotator& Rotation, EObjectFlags ObjectFlags, const FName& Name ) override;
+	//virtual void PostSpawnActor( UObject* Asset, AActor* NewActor) override;
+	//virtual void PostCreateBlueprint( UObject* Asset, AActor* CDO ) override;
+	//virtual UObject* GetAssetFromActorInstance(AActor* ActorInstance) override;
+	// End UActorFactory Interface
+};
diff --git a/Engine/Source/Editor/LandscapeEditor/Private/Classes/LandscapePlaceholder.h b/Engine/Source/Editor/LandscapeEditor/Private/Classes/LandscapePlaceholder.h
new file mode 100644
--- /dev/null
+++ b/Engine/Source/Editor/LandscapeEditor/Private/Classes/LandscapePlaceholder.h
@@ -0,0 +1,16 @@
+// Copyright 1998-2015 Epic Games, Inc. All Rights Reserved.
+
+#pragma once
+
+#include "LandscapePlaceholder.generated.h"
+
+UCLASS(NotPlaceable)
+class ALandscapePlaceholder : public AActor
+{
+	GENERATED_UCLASS_BODY()
+
+	//virtual void PostEditMove(bool bFinished) override;
+	virtual bool TeleportTo(const FVector& DestLocation, const FRotator& DestRotation, bool bIsATest = false, bool bNoCheck = false) override;
+
+	void PostEditChangeProperty(FPropertyChangedEvent& PropertyChangedEvent) override;
+};
diff --git a/Engine/Source/Editor/LandscapeEditor/Private/LandscapeEdMode.cpp b/Engine/Source/Editor/LandscapeEditor/Private/LandscapeEdMode.cpp
--- a/Engine/Source/Editor/LandscapeEditor/Private/LandscapeEdMode.cpp
+++ b/Engine/Source/Editor/LandscapeEditor/Private/LandscapeEdMode.cpp
@@ -7,6 +7,7 @@
 #include "LandscapeEdMode.h"
 #include "ScopedTransaction.h"
 #include "LandscapeEdit.h"
+#include "LandscapeEditorUtils.h"
 #include "LandscapeRender.h"
 #include "LandscapeDataAccess.h"
 #include "LandscapeSplineProxies.h"
diff --git a/Engine/Source/Editor/LandscapeEditor/Private/LandscapeEdMode.h b/Engine/Source/Editor/LandscapeEditor/Private/LandscapeEdMode.h
new file mode 100644
--- /dev/null
+++ b/Engine/Source/Editor/LandscapeEditor/Private/LandscapeEdMode.h
@@ -0,0 +1,450 @@
+// Copyright 1998-2015 Epic Games, Inc. All Rights Reserved.
+
+#pragma once
+
+#include "LandscapeToolInterface.h"
+#include "LandscapeInfo.h"
+#include "LandscapeLayerInfoObject.h"
+#include "LandscapeProxy.h"
+#include "LandscapeGizmoActiveActor.h"
+#include "LevelEditorViewport.h"
+
+DECLARE_LOG_CATEGORY_EXTERN(LogLandscapeEdMode, Log, All);
+
+// Forward declarations
+class ULandscapeLayerInfoObject;
+class FLandscapeToolSplines;
+
+struct FHeightmapToolTarget;
+template<typename TargetType> class FLandscapeToolCopyPaste;
+
+struct FLandscapeToolMode
+{
+	const FName				ToolModeName;
+	int32					SupportedTargetTypes; // ELandscapeToolTargetTypeMask::Type
+
+	TArray<FName>			ValidTools;
+	FName					CurrentToolName;
+
+	FLandscapeToolMode(FName InToolModeName, int32 InSupportedTargetTypes)
+		: ToolModeName(InToolModeName)
+		, SupportedTargetTypes(InSupportedTargetTypes)
+	{
+	}
+};
+
+struct FLandscapeTargetListInfo
+{
+	FText TargetName;
+	ELandscapeToolTargetType::Type TargetType;
+	TWeakObjectPtr<ULandscapeInfo> LandscapeInfo;
+
+	//Values cloned from FLandscapeLayerStruct LayerStruct;			// ignored for heightmap
+	TWeakObjectPtr<ULandscapeLayerInfoObject> LayerInfoObj;			// ignored for heightmap
+	FName LayerName;												// ignored for heightmap
+	TWeakObjectPtr<class ALandscapeProxy> Owner;					// ignored for heightmap
+	TWeakObjectPtr<class UMaterialInstanceConstant> ThumbnailMIC;	// ignored for heightmap
+	int32 DebugColorChannel;										// ignored for heightmap
+	uint32 bValid : 1;												// ignored for heightmap
+
+	FLandscapeTargetListInfo(FText InTargetName, ELandscapeToolTargetType::Type InTargetType, const FLandscapeInfoLayerSettings& InLayerSettings)
+		: TargetName(InTargetName)
+		, TargetType(InTargetType)
+		, LandscapeInfo(InLayerSettings.Owner->GetLandscapeInfo())
+		, LayerInfoObj(InLayerSettings.LayerInfoObj)
+		, LayerName(InLayerSettings.LayerName)
+		, Owner(InLayerSettings.Owner)
+		, ThumbnailMIC(InLayerSettings.ThumbnailMIC)
+		, DebugColorChannel(InLayerSettings.DebugColorChannel)
+		, bValid(InLayerSettings.bValid)
+	{
+	}
+
+	FLandscapeTargetListInfo(FText InTargetName, ELandscapeToolTargetType::Type InTargetType, ULandscapeInfo* InLandscapeInfo)
+		: TargetName(InTargetName)
+		, TargetType(InTargetType)
+		, LandscapeInfo(InLandscapeInfo)
+		, LayerInfoObj(NULL)
+		, LayerName(NAME_None)
+		, Owner(NULL)
+		, ThumbnailMIC(NULL)
+		, bValid(true)
+	{
+	}
+
+	FLandscapeInfoLayerSettings* GetLandscapeInfoLayerSettings() const
+	{
+		if (TargetType == ELandscapeToolTargetType::Weightmap)
+		{
+			int32 Index = INDEX_NONE;
+			if (LayerInfoObj.IsValid())
+			{
+				Index = LandscapeInfo->GetLayerInfoIndex(LayerInfoObj.Get(), Owner.Get());
+			}
+			else
+			{
+				Index = LandscapeInfo->GetLayerInfoIndex(LayerName, Owner.Get());
+			}
+			if (ensure(Index != INDEX_NONE))
+			{
+				return &LandscapeInfo->Layers[Index];
+			}
+		}
+		return NULL;
+	}
+
+	FLandscapeEditorLayerSettings* GetEditorLayerSettings() const
+	{
+		if (TargetType == ELandscapeToolTargetType::Weightmap)
+		{
+			check(LayerInfoObj.IsValid());
+			ALandscapeProxy* Proxy = LandscapeInfo->GetLandscapeProxy();
+			FLandscapeEditorLayerSettings* EditorLayerSettings = Proxy->EditorLayerSettings.FindByKey(LayerInfoObj.Get());
+			if (EditorLayerSettings)
+			{
+				return EditorLayerSettings;
+			}
+			else
+			{
+				int32 Index = Proxy->EditorLayerSettings.Add(LayerInfoObj.Get());
+				return &Proxy->EditorLayerSettings[Index];
+			}
+		}
+		return NULL;
+	}
+
+	FName GetLayerName() const;
+
+	FString& ReimportFilePath() const
+	{
+		if (TargetType == ELandscapeToolTargetType::Weightmap)
+		{
+			FLandscapeEditorLayerSettings* EditorLayerSettings = GetEditorLayerSettings();
+			check(EditorLayerSettings);
+			return EditorLayerSettings->ReimportLayerFilePath;
+		}
+		else //if (TargetType == ELandscapeToolTargetType::Heightmap)
+		{
+			return LandscapeInfo->GetLandscapeProxy()->ReimportHeightmapFilePath;
+		}
+	}
+};
+
+struct FLandscapeListInfo
+{
+	FString LandscapeName;
+	ULandscapeInfo* Info;
+	int32 ComponentQuads;
+	int32 NumSubsections;
+	int32 Width;
+	int32 Height;
+
+	FLandscapeListInfo(const TCHAR* InName, ULandscapeInfo* InInfo, int32 InComponentQuads, int32 InNumSubsections, int32 InWidth, int32 InHeight)
+		: LandscapeName(InName)
+		, Info(InInfo)
+		, ComponentQuads(InComponentQuads)
+		, NumSubsections(InNumSubsections)
+		, Width(InWidth)
+		, Height(InHeight)
+	{
+	}
+};
+
+struct FGizmoHistory
+{
+	ALandscapeGizmoActor* Gizmo;
+	FString GizmoName;
+
+	FGizmoHistory(ALandscapeGizmoActor* InGizmo)
+		: Gizmo(InGizmo)
+	{
+		GizmoName = Gizmo->GetPathName();
+	}
+
+	FGizmoHistory(ALandscapeGizmoActiveActor* InGizmo)
+	{
+		// handle for ALandscapeGizmoActiveActor -> ALandscapeGizmoActor
+		// ALandscapeGizmoActor is only for history, so it has limited data
+		Gizmo = InGizmo->SpawnGizmoActor();
+		GizmoName = Gizmo->GetPathName();
+	}
+};
+
+
+namespace ELandscapeEdge
+{
+	enum Type
+	{
+		None,
+
+		// Edges
+		X_Negative,
+		X_Positive,
+		Y_Negative,
+		Y_Positive,
+
+		// Corners
+		X_Negative_Y_Negative,
+		X_Positive_Y_Negative,
+		X_Negative_Y_Positive,
+		X_Positive_Y_Positive,
+	};
+}
+
+namespace ENewLandscapePreviewMode
+{
+	enum Type
+	{
+		None,
+		NewLandscape,
+		ImportLandscape,
+	};
+}
+
+enum class ELandscapeEditingState : uint8
+{
+	Unknown,
+	Enabled,
+	BadFeatureLevel,
+	PIEWorld,
+	SIEWorld,
+	NoLandscape,
+};
+
+/**
+ * Landscape editor mode
+ */
+class FEdModeLandscape : public FEdMode
+{
+public:
+
+	class ULandscapeEditorObject* UISettings;
+
+	FLandscapeToolMode* CurrentToolMode;
+	FLandscapeTool* CurrentTool;
+	FLandscapeBrush* CurrentBrush;
+	FLandscapeToolTarget CurrentToolTarget;
+
+	// GizmoBrush for Tick
+	FLandscapeBrush* GizmoBrush;
+	// UI setting for additional UI Tools
+	int32 CurrentToolIndex;
+	// UI setting for additional UI Tools
+	int32 CurrentBrushSetIndex;
+
+	ENewLandscapePreviewMode::Type NewLandscapePreviewMode;
+	ELandscapeEdge::Type DraggingEdge;
+	float DraggingEdge_Remainder;
+
+	TWeakObjectPtr<ALandscapeGizmoActiveActor> CurrentGizmoActor;
+	//
+	FLandscapeToolCopyPaste<FHeightmapToolTarget>* CopyPasteTool;
+	void CopyDataToGizmo();
+	void PasteDataFromGizmo();
+
+	FLandscapeToolSplines* SplinesTool;
+	void ShowSplineProperties();
+	virtual void SelectAllConnectedSplineControlPoints();
+	virtual void SelectAllConnectedSplineSegments();
+	virtual void SplineMoveToCurrentLevel();
+	void SetbUseAutoRotateOnJoin(bool InbAutoRotateOnJoin);
+	bool GetbUseAutoRotateOnJoin();
+
+	void ApplyRampTool();
+	bool CanApplyRampTool();
+	void ResetRampTool();
+
+	/** Constructor */
+	FEdModeLandscape();
+
+	/** Initialization */
+	void InitializeBrushes();
+	void InitializeTool_Paint();
+	void InitializeTool_Smooth();
+	void InitializeTool_Flatten();
+	void InitializeTool_Erosion();
+	void InitializeTool_HydraErosion();
+	void InitializeTool_Noise();
+	void InitializeTool_Retopologize();
+	void InitializeTool_NewLandscape();
+	void InitializeTool_ResizeLandscape();
+	void InitializeTool_Select();
+	void InitializeTool_AddComponent();
+	void InitializeTool_DeleteComponent();
+	void InitializeTool_MoveToLevel();
+	void InitializeTool_Mask();
+	void InitializeTool_CopyPaste();
+	void InitializeTool_Visibility();
+	void InitializeTool_Splines();
+	void InitializeTool_Ramp();
+	void InitializeToolModes();
+
+	/** Destructor */
+	virtual ~FEdModeLandscape();
+
+	/** FGCObject interface */
+	virtual void AddReferencedObjects(FReferenceCollector& Collector) override;
+
+	virtual bool UsesToolkits() const override;
+
+	TSharedRef<FUICommandList> GetUICommandList() const;
+
+	/** FEdMode: Called when the mode is entered */
+	virtual void Enter() override;
+
+	/** FEdMode: Called when the mode is exited */
+	virtual void Exit() override;
+
+	/** FEdMode: Called when the mouse is moved over the viewport */
+	virtual bool MouseMove(FEditorViewportClient* ViewportClient, FViewport* Viewport, int32 x, int32 y) override;
+
+	/**
+	 * FEdMode: Called when the mouse is moved while a window input capture is in effect
+	 *
+	 * @param	InViewportClient	Level editor viewport client that captured the mouse input
+	 * @param	InViewport			Viewport that captured the mouse input
+	 * @param	InMouseX			New mouse cursor X coordinate
+	 * @param	InMouseY			New mouse cursor Y coordinate
+	 *
+	 * @return	true if input was handled
+	 */
+	virtual bool CapturedMouseMove(FEditorViewportClient* InViewportClient, FViewport* InViewport, int32 InMouseX, int32 InMouseY) override;
+
+	/** FEdMode: Called when a mouse button is pressed */
+	virtual bool StartTracking(FEditorViewportClient* InViewportClient, FViewport* InViewport) override;
+
+	/** FEdMode: Called when a mouse button is released */
+	virtual bool EndTracking(FEditorViewportClient* InViewportClient, FViewport* InViewport) override;
+
+	/** FEdMode: Allow us to disable mouse delta tracking during painting */
+	virtual bool DisallowMouseDeltaTracking() const override;
+
+	/** FEdMode: Called once per frame */
+	virtual void Tick(FEditorViewportClient* ViewportClient, float DeltaTime) override;
+
+	/** FEdMode: Called when clicking on a hit proxy */
+	virtual bool HandleClick(FEditorViewportClient* InViewportClient, HHitProxy* HitProxy, const FViewportClick& Click) override;
+
+	/** FEdMode: Called when a key is pressed */
+	virtual bool InputKey(FEditorViewportClient* InViewportClient, FViewport* InViewport, FKey InKey, EInputEvent InEvent) override;
+
+	/** FEdMode: Called when mouse drag input is applied */
+	virtual bool InputDelta(FEditorViewportClient* InViewportClient, FViewport* InViewport, FVector& InDrag, FRotator& InRot, FVector& InScale) override;
+
+	/** FEdMode: Render elements for the landscape tool */
+	virtual void Render(const FSceneView* View, FViewport* Viewport, FPrimitiveDrawInterface* PDI) override;
+
+	/** FEdMode: Render HUD elements for this tool */
+	virtual void DrawHUD(FEditorViewportClient* ViewportClient, FViewport* Viewport, const FSceneView* View, FCanvas* Canvas) override;
+
+	/** FEdMode: Handling SelectActor */
+	virtual bool Select(AActor* InActor, bool bInSelected) override;
+
+	/** FEdMode: Check to see if an actor can be selected in this mode - no side effects */
+	virtual bool IsSelectionAllowed(AActor* InActor, bool bInSelection) const override;
+
+	/** FEdMode: Called when the currently selected actor has changed */
+	virtual void ActorSelectionChangeNotify() override;
+
+	virtual void ActorMoveNotify() override;
+
+	virtual void PostUndo() override;
+
+	virtual EEditAction::Type GetActionEditDuplicate() override;
+	virtual EEditAction::Type GetActionEditDelete() override;
+	virtual EEditAction::Type GetActionEditCut() override;
+	virtual EEditAction::Type GetActionEditCopy() override;
+	virtual EEditAction::Type GetActionEditPaste() override;
+	virtual bool ProcessEditDuplicate() override;
+	virtual bool ProcessEditDelete() override;
+	virtual bool ProcessEditCut() override;
+	virtual bool ProcessEditCopy() override;
+	virtual bool ProcessEditPaste() override;
+
+	/** FEdMode: If the EdMode is handling InputDelta (ie returning true from it), this allows a mode to indicated whether or not the Widget should also move. */
+	virtual bool AllowWidgetMove() override { return true; }
+
+	/** FEdMode: Draw the transform widget while in this mode? */
+	virtual bool ShouldDrawWidget() const override;
+
+	/** FEdMode: Returns true if this mode uses the transform widget */
+	virtual bool UsesTransformWidget() const override;
+
+	virtual EAxisList::Type GetWidgetAxisToDraw(FWidget::EWidgetMode InWidgetMode) const override;
+
+	virtual FVector GetWidgetLocation() const override;
+	virtual bool GetCustomDrawingCoordinateSystem(FMatrix& InMatrix, void* InData) override;
+	virtual bool GetCustomInputCoordinateSystem(FMatrix& InMatrix, void* InData) override;
+
+	/** Forces real-time perspective viewports */
+	void ForceRealTimeViewports(const bool bEnable, const bool bStoreCurrentState);
+
+	/** Trace under the mouse cursor and return the landscape hit and the hit location (in landscape quad space) */
+	bool LandscapeMouseTrace(FEditorViewportClient* ViewportClient, float& OutHitX, float& OutHitY);
+	bool LandscapeMouseTrace(FEditorViewportClient* ViewportClient, FVector& OutHitLocation);
+
+	/** Trace under the specified coordinates and return the landscape hit and the hit location (in landscape quad space) */
+	bool LandscapeMouseTrace(FEditorViewportClient* ViewportClient, int32 MouseX, int32 MouseY, float& OutHitX, float& OutHitY);
+	bool LandscapeMouseTrace(FEditorViewportClient* ViewportClient, int32 MouseX, int32 MouseY, FVector& OutHitLocation);
+
+	/** Trace under the mouse cursor / specified screen coordinates against a world-space plane and return the hit location (in world space) */
+	bool LandscapePlaneTrace(FEditorViewportClient* ViewportClient, const FPlane& Plane, FVector& OutHitLocation);
+	bool LandscapePlaneTrace(FEditorViewportClient* ViewportClient, int32 MouseX, int32 MouseY, const FPlane& Plane, FVector& OutHitLocation);
+
+	void SetCurrentToolMode(FName ToolModeName, bool bRestoreCurrentTool = true);
+
+	/** Change current tool */
+	void SetCurrentTool(FName ToolName);
+	void SetCurrentTool(int32 ToolIdx);
+
+	void SetCurrentBrushSet(FName BrushSetName);
+	void SetCurrentBrushSet(int32 BrushSetIndex);
+
+	void SetCurrentBrush(FName BrushName);
+	void SetCurrentBrush(int32 BrushIndex);
+
+	const TArray<TSharedRef<FLandscapeTargetListInfo>>& GetTargetList();
+	const TArray<FLandscapeListInfo>& GetLandscapeList();
+
+	void AddLayerInfo(ULandscapeLayerInfoObject* LayerInfo);
+
+	int32 UpdateLandscapeList();
+	void UpdateTargetList();
+
+	DECLARE_EVENT(FEdModeLandscape, FTargetsListUpdated);
+	static FTargetsListUpdated TargetsListUpdated;
+
+	void OnWorldChange();
+
+	void OnMaterialCompilationFinished(UMaterialInterface* MaterialInterface);
+
+	void ReimportData(const FLandscapeTargetListInfo& TargetInfo);
+	void ImportData(const FLandscapeTargetListInfo& TargetInfo, const FString& Filename);
+
+	ALandscape* ChangeComponentSetting(int32 NumComponentsX, int32 NumComponentsY, int32 InNumSubsections, int32 InSubsectionSizeQuads, bool bResample);
+
+	TArray<FLandscapeToolMode> LandscapeToolModes;
+	TArray<TUniquePtr<FLandscapeTool>> LandscapeTools;
+	TArray<FLandscapeBrushSet> LandscapeBrushSets;
+
+	// For collision add visualization
+	FLandscapeAddCollision* LandscapeRenderAddCollision;
+
+	ELandscapeEditingState GetEditingState() const;
+
+	bool IsEditingEnabled() const
+	{
+		return GetEditingState() == ELandscapeEditingState::Enabled;
+	}
+
+private:
+	TArray<TSharedRef<FLandscapeTargetListInfo>> LandscapeTargetList;
+	TArray<FLandscapeListInfo> LandscapeList;
+
+	UMaterialInterface* CachedLandscapeMaterial;
+
+	bool bToolActive;
+
+	FDelegateHandle OnWorldChangeDelegateHandle;
+	FDelegateHandle OnMaterialCompilationFinishedDelegateHandle;
+};
diff --git a/Engine/Source/Editor/LandscapeEditor/Private/LandscapeEditorCommands.h b/Engine/Source/Editor/LandscapeEditor/Private/LandscapeEditorCommands.h
new file mode 100644
--- /dev/null
+++ b/Engine/Source/Editor/LandscapeEditor/Private/LandscapeEditorCommands.h
@@ -0,0 +1,88 @@
+// Copyright 1998-2015 Epic Games, Inc. All Rights Reserved.
+
+#pragma once
+
+
+/**
+ * Unreal landscape editor actions
+ */
+class FLandscapeEditorCommands : public TCommands<FLandscapeEditorCommands>
+{
+
+public:
+	FLandscapeEditorCommands() : TCommands<FLandscapeEditorCommands>
+	(
+		"LandscapeEditor", // Context name for fast lookup
+		NSLOCTEXT("Contexts", "LandscapeEditor", "Landscape Editor"), // Localized context name for displaying
+		NAME_None, //"LevelEditor" // Parent
+		FEditorStyle::GetStyleSetName() // Icon Style Set
+	)
+	{
+	}
+	
+
+	/**
+	 * Initialize commands
+	 */
+	virtual void RegisterCommands() override;
+
+public:
+	// Mode Switch
+	TSharedPtr< FUICommandInfo > ManageMode;
+	TSharedPtr< FUICommandInfo > SculptMode;
+	TSharedPtr< FUICommandInfo > PaintMode;
+
+	// Tools
+	TSharedPtr< FUICommandInfo > NewLandscape;
+	TSharedPtr< FUICommandInfo > ResizeLandscape;
+
+	TSharedPtr< FUICommandInfo > SculptTool;
+	TSharedPtr< FUICommandInfo > PaintTool;
+	TSharedPtr< FUICommandInfo > SmoothTool;
+	TSharedPtr< FUICommandInfo > FlattenTool;
+	TSharedPtr< FUICommandInfo > RampTool;
+	TSharedPtr< FUICommandInfo > ErosionTool;
+	TSharedPtr< FUICommandInfo > HydroErosionTool;
+	TSharedPtr< FUICommandInfo > NoiseTool;
+	TSharedPtr< FUICommandInfo > RetopologizeTool;
+	TSharedPtr< FUICommandInfo > VisibilityTool;
+
+	TSharedPtr< FUICommandInfo > SelectComponentTool;
+	TSharedPtr< FUICommandInfo > AddComponentTool;
+	TSharedPtr< FUICommandInfo > DeleteComponentTool;
+	TSharedPtr< FUICommandInfo > MoveToLevelTool;
+
+	TSharedPtr< FUICommandInfo > RegionSelectTool;
+	TSharedPtr< FUICommandInfo > RegionCopyPasteTool;
+
+	TSharedPtr< FUICommandInfo > SplineTool;
+
+	// Brushes
+	TSharedPtr< FUICommandInfo > CircleBrush;
+	TSharedPtr< FUICommandInfo > AlphaBrush;
+	TSharedPtr< FUICommandInfo > AlphaBrush_Pattern;
+	TSharedPtr< FUICommandInfo > ComponentBrush;
+	TSharedPtr< FUICommandInfo > GizmoBrush;
+
+	TSharedPtr< FUICommandInfo > CircleBrush_Smooth;
+	TSharedPtr< FUICommandInfo > CircleBrush_Linear;
+	TSharedPtr< FUICommandInfo > CircleBrush_Spherical;
+	TSharedPtr< FUICommandInfo > CircleBrush_Tip;
+
+	TSharedPtr< FUICommandInfo > ViewModeNormal;
+	TSharedPtr< FUICommandInfo > ViewModeLOD;
+	TSharedPtr< FUICommandInfo > ViewModeLayerDensity;
+	TSharedPtr< FUICommandInfo > ViewModeLayerDebug;
+	TSharedPtr< FUICommandInfo > ViewModeWireframeOnTop;
+
+	// Map
+	TMap<FName, TSharedPtr<FUICommandInfo>> NameToCommandMap;
+};
+
+///**
+// * Implementation of various level editor action callback functions
+// */
+//class FLevelEditorActionCallbacks
+//{
+//public:
+//};
diff --git a/Engine/Source/Editor/LandscapeEditor/Private/LandscapeEditorDetailCustomization_NewLandscape.cpp b/Engine/Source/Editor/LandscapeEditor/Private/LandscapeEditorDetailCustomization_NewLandscape.cpp
--- a/Engine/Source/Editor/LandscapeEditor/Private/LandscapeEditorDetailCustomization_NewLandscape.cpp
+++ b/Engine/Source/Editor/LandscapeEditor/Private/LandscapeEditorDetailCustomization_NewLandscape.cpp
@@ -6,6 +6,7 @@
 #include "LandscapeEditorObject.h"
 #include "LandscapeEditorDetails.h"
 #include "LandscapeEditorDetailCustomizations.h"
+#include "LandscapeEditorUtils.h"
 
 #include "DetailLayoutBuilder.h"
 #include "DetailCategoryBuilder.h"
diff --git a/Engine/Source/Editor/LandscapeEditor/Private/LandscapeEditorModule.cpp b/Engine/Source/Editor/LandscapeEditor/Private/LandscapeEditorModule.cpp
--- a/Engine/Source/Editor/LandscapeEditor/Private/LandscapeEditorModule.cpp
+++ b/Engine/Source/Editor/LandscapeEditor/Private/LandscapeEditorModule.cpp
@@ -3,11 +3,12 @@
 #include "LandscapeEditorPrivatePCH.h"
 #include "LandscapeEditorCommands.h"
 #include "LandscapeEdMode.h"
-#include "LandscapeEditorClasses.h"
+#include "Classes/ActorFactoryLandscape.h"
 
 
 #include "LandscapeEditorDetails.h"
 #include "LandscapeEditorDetailCustomizations.h"
+#include "LandscapeSplineDetails.h"
 #include "PropertyEditorModule.h"
 #include "PropertyEditorDelegates.h"
 
@@ -38,11 +39,16 @@ public:
 			300
 			);
 
+		// register customizations
 		FPropertyEditorModule& PropertyModule = FModuleManager::GetModuleChecked<FPropertyEditorModule>("PropertyEditor");
-		PropertyModule.RegisterCustomClassLayout("LandscapeEditorObject", FOnGetDetailCustomizationInstance::CreateStatic(&FLandscapeEditorDetails::MakeInstance));
-		PropertyModule.RegisterCustomPropertyTypeLayout("GizmoImportLayer", FOnGetPropertyTypeCustomizationInstance::CreateStatic(&FLandscapeEditorStructCustomization_FGizmoImportLayer::MakeInstance));
+		PropertyModule.RegisterCustomClassLayout("LandscapeEditorObject",       FOnGetDetailCustomizationInstance::CreateStatic(&FLandscapeEditorDetails::MakeInstance));
+		PropertyModule.RegisterCustomPropertyTypeLayout("GizmoImportLayer",     FOnGetPropertyTypeCustomizationInstance::CreateStatic(&FLandscapeEditorStructCustomization_FGizmoImportLayer::MakeInstance));
 		PropertyModule.RegisterCustomPropertyTypeLayout("LandscapeImportLayer", FOnGetPropertyTypeCustomizationInstance::CreateStatic(&FLandscapeEditorStructCustomization_FLandscapeImportLayer::MakeInstance));
 
+		PropertyModule.RegisterCustomClassLayout("LandscapeSplineControlPoint", FOnGetDetailCustomizationInstance::CreateStatic(&FLandscapeSplineDetails::MakeInstance));
+		PropertyModule.RegisterCustomClassLayout("LandscapeSplineSegment",      FOnGetDetailCustomizationInstance::CreateStatic(&FLandscapeSplineDetails::MakeInstance));
+
+		// add menu extension
 		TSharedRef<FUICommandList> CommandList = MakeShareable(new FUICommandList);
 		const FLandscapeEditorCommands& LandscapeActions = FLandscapeEditorCommands::Get();
 		CommandList->MapAction(LandscapeActions.ViewModeNormal,       FExecuteAction::CreateStatic(&ChangeLandscapeViewMode, ELandscapeViewMode::Normal),       FCanExecuteAction(), FIsActionChecked::CreateStatic(&IsLandscapeViewModeSelected, ELandscapeViewMode::Normal));
@@ -51,12 +57,12 @@ public:
 		CommandList->MapAction(LandscapeActions.ViewModeLayerDebug,   FExecuteAction::CreateStatic(&ChangeLandscapeViewMode, ELandscapeViewMode::DebugLayer),   FCanExecuteAction(), FIsActionChecked::CreateStatic(&IsLandscapeViewModeSelected, ELandscapeViewMode::DebugLayer));
 		CommandList->MapAction(LandscapeActions.ViewModeWireframeOnTop,FExecuteAction::CreateStatic(&ChangeLandscapeViewMode, ELandscapeViewMode::WireframeOnTop), FCanExecuteAction(), FIsActionChecked::CreateStatic(&IsLandscapeViewModeSelected, ELandscapeViewMode::WireframeOnTop));
 
-		TSharedRef<FExtender> ViewportMenuExtender = MakeShareable(new FExtender);
+		ViewportMenuExtender = MakeShareable(new FExtender);
 		ViewportMenuExtender->AddMenuExtension("LevelViewportLandscape", EExtensionHook::First, CommandList, FMenuExtensionDelegate::CreateStatic(&ConstructLandscapeViewportMenu));
 		FLevelEditorModule& LevelEditorModule = FModuleManager::LoadModuleChecked<FLevelEditorModule>("LevelEditor");
 		LevelEditorModule.GetMenuExtensibilityManager()->AddExtender(ViewportMenuExtender);
 
-		// Actor Factories
+		// add actor factories
 		UActorFactoryLandscape* LandscapeActorFactory = NewObject<UActorFactoryLandscape>();
 		LandscapeActorFactory->NewActorClass = ALandscape::StaticClass();
 		GEditor->ActorFactories.Add(LandscapeActorFactory);
@@ -75,6 +81,24 @@ public:
 
 		// unregister the editor mode
 		FEditorModeRegistry::Get().UnregisterMode(FBuiltinEditorModes::EM_Landscape);
+
+		// unregister customizations
+		FPropertyEditorModule& PropertyModule = FModuleManager::GetModuleChecked<FPropertyEditorModule>("PropertyEditor");
+		PropertyModule.UnregisterCustomClassLayout("LandscapeEditorObject");
+		PropertyModule.UnregisterCustomPropertyTypeLayout("GizmoImportLayer");
+		PropertyModule.UnregisterCustomPropertyTypeLayout("LandscapeImportLayer");
+
+		PropertyModule.UnregisterCustomClassLayout("LandscapeSplineControlPoint");
+		PropertyModule.UnregisterCustomClassLayout("LandscapeSplineSegment");
+
+		// remove menu extension
+		FLevelEditorModule& LevelEditorModule = FModuleManager::LoadModuleChecked<FLevelEditorModule>("LevelEditor");
+		LevelEditorModule.GetMenuExtensibilityManager()->RemoveExtender(ViewportMenuExtender);
+		ViewportMenuExtender = nullptr;
+
+		// remove actor factories
+		// TODO - this crashes on shutdown
+		// GEditor->ActorFactories.RemoveAll([](const UActorFactory* ActorFactory) { return ActorFactory->IsA<UActorFactoryLandscape>(); });
 	}
 
 	static void ConstructLandscapeViewportMenu(FMenuBuilder& MenuBuilder)
@@ -91,7 +115,7 @@ public:
 					InMenuBuilder.AddMenuEntry(LandscapeActions.ViewModeLOD, NAME_None, LOCTEXT("LandscapeViewModeLOD", "LOD"));
 					InMenuBuilder.AddMenuEntry(LandscapeActions.ViewModeLayerDensity, NAME_None, LOCTEXT("LandscapeViewModeLayerDensity", "Layer Density"));
 					InMenuBuilder.AddMenuEntry(LandscapeActions.ViewModeLayerDebug, NAME_None, LOCTEXT("LandscapeViewModeLayerDebug", "Layer Debug"));
-					InMenuBuilder.AddMenuEntry(LandscapeActions.ViewModeWireframeOnTop, NAME_None, LOCTEXT("LandscapeViewModeWireframeOnTop", "Wireframe on Top"));				
+					InMenuBuilder.AddMenuEntry(LandscapeActions.ViewModeWireframeOnTop, NAME_None, LOCTEXT("LandscapeViewModeWireframeOnTop", "Wireframe on Top"));
 				}
 				InMenuBuilder.EndSection();
 			}
@@ -108,6 +132,9 @@ public:
 	{
 		return GLandscapeViewMode == ViewMode;
 	}
+
+protected:
+	TSharedPtr<FExtender> ViewportMenuExtender;
 };
 
 IMPLEMENT_MODULE( FLandscapeEditorModule, LandscapeEditor );
diff --git a/Engine/Source/Editor/LandscapeEditor/Private/LandscapeSplineDetails.cpp b/Engine/Source/Editor/LandscapeEditor/Private/LandscapeSplineDetails.cpp
new file mode 100644
--- /dev/null
+++ b/Engine/Source/Editor/LandscapeEditor/Private/LandscapeSplineDetails.cpp
@@ -0,0 +1,112 @@
+// Copyright 1998-2015 Epic Games, Inc. All Rights Reserved.
+
+#include "LandscapeEditorPrivatePCH.h"
+#include "LandscapeSplineDetails.h"
+#include "LandscapeEdMode.h"
+#include "DetailCategoryBuilder.h"
+#include "DetailLayoutBuilder.h"
+#include "DetailWidgetRow.h"
+
+#define LOCTEXT_NAMESPACE "LandscapeSplineDetails"
+
+
+TSharedRef<IDetailCustomization> FLandscapeSplineDetails::MakeInstance()
+{
+	return MakeShareable(new FLandscapeSplineDetails);
+}
+
+void FLandscapeSplineDetails::CustomizeDetails(IDetailLayoutBuilder& DetailBuilder)
+{
+	IDetailCategoryBuilder& LandscapeSplineCategory = DetailBuilder.EditCategory("LandscapeSpline", FText::GetEmpty(), ECategoryPriority::Transform);
+
+	LandscapeSplineCategory.AddCustomRow(FText::GetEmpty())
+	[
+		SNew(SHorizontalBox)
+		+ SHorizontalBox::Slot()
+		.Padding(0, 0, 2, 0)
+		.VAlign(VAlign_Center)
+		.FillWidth(1)
+		[
+			SNew(STextBlock)
+			.Text(LOCTEXT("SelectAll", "Select all connected:"))
+		]
+		+ SHorizontalBox::Slot()
+		.FillWidth(1)
+		[
+			SNew(SButton)
+			.Text(LOCTEXT("ControlPoints", "Control Points"))
+			.HAlign(HAlign_Center)
+			.OnClicked(this, &FLandscapeSplineDetails::OnSelectConnectedControlPointsButtonClicked)
+		]
+		+ SHorizontalBox::Slot()
+		.FillWidth(1)
+		[
+			SNew(SButton)
+			.Text(LOCTEXT("Segments", "Segments"))
+			.HAlign(HAlign_Center)
+			.OnClicked(this, &FLandscapeSplineDetails::OnSelectConnectedSegmentsButtonClicked)
+		]
+	];
+
+	LandscapeSplineCategory.AddCustomRow(FText::GetEmpty())
+	[
+		SNew(SHorizontalBox)
+		+ SHorizontalBox::Slot()
+		.Padding(0, 0, 2, 0)
+		.VAlign(VAlign_Center)
+		.FillWidth(1)
+		[
+			SNew(SButton)
+			.Text(LOCTEXT("Move Selected ControlPnts+Segs to Current level", "Move to current level"))
+			.HAlign(HAlign_Center)
+			.OnClicked(this, &FLandscapeSplineDetails::OnMoveToCurrentLevelButtonClicked)
+			.IsEnabled(this, &FLandscapeSplineDetails::IsMoveToCurrentLevelButtonEnabled)
+		]
+	];
+}
+
+class FEdModeLandscape* FLandscapeSplineDetails::GetEditorMode() const
+{
+	return (FEdModeLandscape*)GLevelEditorModeTools().GetActiveMode(FBuiltinEditorModes::EM_Landscape);
+}
+
+FReply FLandscapeSplineDetails::OnSelectConnectedControlPointsButtonClicked()
+{
+	FEdModeLandscape* LandscapeEdMode = GetEditorMode();
+	if (LandscapeEdMode && LandscapeEdMode->CurrentToolTarget.LandscapeInfo.IsValid())
+	{
+		LandscapeEdMode->SelectAllConnectedSplineControlPoints();
+	}
+
+	return FReply::Handled();
+}
+
+FReply FLandscapeSplineDetails::OnSelectConnectedSegmentsButtonClicked()
+{
+	FEdModeLandscape* LandscapeEdMode = GetEditorMode();
+	if (LandscapeEdMode && LandscapeEdMode->CurrentToolTarget.LandscapeInfo.IsValid())
+	{
+		LandscapeEdMode->SelectAllConnectedSplineSegments();
+	}
+
+	return FReply::Handled();
+}
+
+FReply FLandscapeSplineDetails::OnMoveToCurrentLevelButtonClicked()
+{
+	FEdModeLandscape* LandscapeEdMode = GetEditorMode();
+	if (LandscapeEdMode && LandscapeEdMode->CurrentToolTarget.LandscapeInfo.IsValid() && LandscapeEdMode->CurrentToolTarget.LandscapeInfo->GetCurrentLevelLandscapeProxy(true))
+	{
+		LandscapeEdMode->SplineMoveToCurrentLevel();
+	}
+
+	return FReply::Handled();
+}
+
+bool FLandscapeSplineDetails::IsMoveToCurrentLevelButtonEnabled() const
+{
+	FEdModeLandscape* LandscapeEdMode = GetEditorMode();
+	return (LandscapeEdMode && LandscapeEdMode->CurrentToolTarget.LandscapeInfo.IsValid() && LandscapeEdMode->CurrentToolTarget.LandscapeInfo->GetCurrentLevelLandscapeProxy(true));
+}
+
+#undef LOCTEXT_NAMESPACE
diff --git a/Engine/Source/Editor/LandscapeEditor/Private/LandscapeSplineDetails.h b/Engine/Source/Editor/LandscapeEditor/Private/LandscapeSplineDetails.h
new file mode 100644
--- /dev/null
+++ b/Engine/Source/Editor/LandscapeEditor/Private/LandscapeSplineDetails.h
@@ -0,0 +1,26 @@
+// Copyright 1998-2015 Epic Games, Inc. All Rights Reserved.
+
+#pragma once
+
+// Forward declarations
+class FEdModeLandscape;
+
+#include "IDetailCustomization.h"
+
+
+class FLandscapeSplineDetails : public IDetailCustomization
+{
+public:
+	/** Makes a new instance of this detail layout class for a specific detail view requesting it */
+	static TSharedRef<IDetailCustomization> MakeInstance();
+
+	/** IDetailCustomization interface */
+	virtual void CustomizeDetails(IDetailLayoutBuilder& DetailBuilder) override;
+
+private:
+	FEdModeLandscape* GetEditorMode() const;
+	FReply OnSelectConnectedControlPointsButtonClicked();
+	FReply OnSelectConnectedSegmentsButtonClicked();
+	FReply OnMoveToCurrentLevelButtonClicked();
+	bool IsMoveToCurrentLevelButtonEnabled() const;
+};
diff --git a/Engine/Source/Editor/LandscapeEditor/Private/Tests/LandscapeAutomationTests.cpp b/Engine/Source/Editor/LandscapeEditor/Private/Tests/LandscapeAutomationTests.cpp
--- a/Engine/Source/Editor/LandscapeEditor/Private/Tests/LandscapeAutomationTests.cpp
+++ b/Engine/Source/Editor/LandscapeEditor/Private/Tests/LandscapeAutomationTests.cpp
@@ -5,7 +5,7 @@
 #include "AutomationCommon.h"
 #include "AutomationEditorCommon.h"
 
-#include "LandscapeEdMode.h"
+#include "../LandscapeEdMode.h"
 #include "../LandscapeEditorDetailCustomization_NewLandscape.h"
 #include "ScopedTransaction.h"
 #include "Landscape.h"
diff --git a/Engine/Source/Editor/LandscapeEditor/Public/LandscapeEdMode.h b/Engine/Source/Editor/LandscapeEditor/Public/LandscapeEdMode.h
deleted file mode 100644
diff --git a/Engine/Source/Editor/LandscapeEditor/Public/LandscapeEditorCommands.h b/Engine/Source/Editor/LandscapeEditor/Public/LandscapeEditorCommands.h
deleted file mode 100644
diff --git a/Engine/Source/Editor/LandscapeEditor/Public/LandscapeEditorObject.h b/Engine/Source/Editor/LandscapeEditor/Public/LandscapeEditorObject.h
new file mode 100644
--- /dev/null
+++ b/Engine/Source/Editor/LandscapeEditor/Public/LandscapeEditorObject.h
@@ -0,0 +1,610 @@
+// Copyright 1998-2015 Epic Games, Inc. All Rights Reserved.
+
+#pragma once
+
+#include "LandscapeProxy.h" // for ELandscapeLayerPaintingRestriction
+#include "LandscapeEditorObject.generated.h"
+
+UENUM()
+namespace ELandscapeToolFlattenMode
+{
+	enum Type
+	{
+		// Invalid = -1,
+
+		/** Flatten may both raise and lower values */
+		Both = 0,
+
+		/** Flatten may only raise values, values above the clicked point will be left unchanged */
+		Raise = 1,
+
+		/** Flatten may only lower values, values below the clicked point will be left unchanged */
+		Lower = 2,
+	};
+}
+
+UENUM()
+namespace ELandscapeToolErosionMode
+{
+	enum Type
+	{
+		// Invalid = -1,
+
+		/** Apply all erosion effects, both raising and lowering the heightmap */
+		Both = 0,
+
+		/** Only applies erosion effects that result in raising the heightmap */
+		Raise = 1,
+
+		/** Only applies erosion effects that result in lowering the heightmap */
+		Lower = 2,
+	};
+}
+
+UENUM()
+namespace ELandscapeToolHydroErosionMode
+{
+	enum Type
+	{
+		// Invalid = -1,
+
+		/** Rains in some places and not others, randomly */
+		Both = 0,
+
+		/** Rain is applied to the entire area */
+		Positive = 1,
+	};
+}
+
+// Temp
+#if !CPP
+UENUM()
+namespace ELandscapeToolNoiseMode
+{
+	enum Type
+	{
+		// Invalid = -1,
+
+		/** Noise will both raise and lower the heightmap */
+		Both = 0,
+
+		/** Noise will only raise the heightmap */
+		Raise = 1,
+
+		/** Noise will only lower the heightmap */
+		Lower = 2,
+	};
+}
+#endif
+
+UENUM()
+namespace ELandscapeToolPasteMode
+{
+	enum Type
+	{
+		// Invalid = -1,
+
+		/** Paste may both raise and lower values */
+		Both = 0,
+
+		/** Paste may only raise values, places where the pasted data would be below the heightmap are left unchanged. Good for copy/pasting mountains */
+		Raise = 1,
+
+		/** Paste may only lower values, places where the pasted data would be above the heightmap are left unchanged. Good for copy/pasting valleys or pits */
+		Lower = 2,
+	};
+}
+
+UENUM()
+namespace ELandscapeConvertMode
+{
+	enum Type
+	{
+		// Invalid = -1,
+
+		/** Given the new component size, the edges of the landscape will be expanded as necessary until its overall size is a whole number of landscape components. */
+		Expand = 0,
+
+		/** Given the new component size, the edges of the landscape will be trimmed until its overall size is a whole number of landscape components. */
+		Clip = 1,
+
+		/** The landscape will have the same overall size in the world, and have the same number of components. Existing landscape geometry and layer data will be resampled to match the new resolution. */
+		Resample = 2,
+	};
+}
+
+UENUM()
+namespace EColorChannel
+{
+	enum Type
+	{
+		Red,
+		Green,
+		Blue,
+		Alpha,
+	};
+}
+
+USTRUCT()
+struct FGizmoImportLayer
+{
+	GENERATED_USTRUCT_BODY()
+
+	UPROPERTY(Category="Import", EditAnywhere)
+	FString LayerFilename;
+
+	UPROPERTY(Category="Import", EditAnywhere)
+	FString LayerName;
+
+	UPROPERTY(Category="Import", EditAnywhere)
+	bool bNoImport;
+
+	FGizmoImportLayer()
+		: LayerFilename("")
+		, LayerName("")
+		, bNoImport(false)
+	{
+	}
+};
+
+UENUM()
+namespace ELandscapeImportHeightmapError
+{
+	enum Type
+	{
+		None,
+		FileNotFound,
+		InvalidSize,
+		CorruptFile,
+		ColorPng,
+		LowBitDepth,
+	};
+}
+
+UENUM()
+enum class ELandscapeImportLayerError : uint8
+{
+	None,
+	MissingLayerInfo,
+	FileNotFound,
+	FileSizeMismatch,
+	CorruptFile,
+	ColorPng,
+};
+
+USTRUCT()
+struct FLandscapeImportLayer : public FLandscapeImportLayerInfo
+{
+	GENERATED_USTRUCT_BODY()
+
+	UPROPERTY(Category="Import", VisibleAnywhere)
+	ELandscapeImportLayerError ImportError;
+
+	FLandscapeImportLayer()
+		: FLandscapeImportLayerInfo()
+		, ImportError(ELandscapeImportLayerError::None)
+	{
+	}
+};
+
+USTRUCT()
+struct FLandscapePatternBrushWorldSpaceSettings
+{
+	GENERATED_USTRUCT_BODY()
+
+	UPROPERTY(Category="World-Space", EditAnywhere)
+	FVector2D Origin;
+
+	UPROPERTY(Category = "World-Space", EditAnywhere, Meta = (ClampMin = "-360", ClampMax = "360", UIMin = "-180", UIMax = "180"))
+	float Rotation;
+
+	// if true, the texture used for the pattern is centered on the PatternOrigin.
+	// if false, the corner of the texture is placed at the PatternOrigin
+	UPROPERTY(Category = "World-Space", EditAnywhere)
+	bool bCenterTextureOnOrigin;
+
+	UPROPERTY(Category = "World-Space", EditAnywhere)
+	float RepeatSize;
+
+	FLandscapePatternBrushWorldSpaceSettings() = default;
+};
+
+UCLASS()
+class ULandscapeEditorObject : public UObject
+{
+	GENERATED_UCLASS_BODY()
+
+	FEdModeLandscape* ParentMode;
+
+
+	// Common Tool Settings:
+
+	// Strength of the tool. If you're using a pen/tablet with pressure-sensing, the pressure used affects the strength of the tool.
+	UPROPERTY(Category="Tool Settings", EditAnywhere, NonTransactional, meta=(ShowForTools="Paint,Sculpt,Smooth,Flatten,Erosion,HydraErosion,Noise,Mask,CopyPaste", ClampMin="0", ClampMax="10", UIMin="0", UIMax="1"))
+	float ToolStrength;
+
+	// Enable to make tools blend towards a target value
+	UPROPERTY(NonTransactional)
+	bool bUseWeightTargetValue;
+
+	// Enable to make tools blend towards a target value
+	UPROPERTY(Category="Tool Settings", EditAnywhere, NonTransactional, meta=(DisplayName="Use Target Value", EditCondition="bUseWeightTargetValue", ShowForTools="Paint,Sculpt,Noise", ClampMin="0", ClampMax="10", UIMin="0", UIMax="1"))
+	float WeightTargetValue;
+
+	// I have no idea what this is for but it's used by the noise and erosion tools, and isn't exposed to the UI
+	UPROPERTY(NonTransactional)
+	float MaximumValueRadius;
+
+	// Flatten Tool:
+
+	// Whether to flatten by lowering, raising, or both
+	UPROPERTY(Category="Tool Settings", EditAnywhere, NonTransactional, meta=(ShowForTools="Flatten"))
+	TEnumAsByte<ELandscapeToolFlattenMode::Type> FlattenMode;
+
+	// Flattens to the angle of the clicked point, instead of horizontal
+	UPROPERTY(Category="Tool Settings", EditAnywhere, NonTransactional, meta=(ShowForTools="Flatten", ShowForTargetTypes="Heightmap"))
+	bool bUseSlopeFlatten;
+
+	// Constantly picks new values to flatten towards when dragging around, instead of only using the first clicked point
+	UPROPERTY(Category="Tool Settings", EditAnywhere, NonTransactional, meta=(ShowForTools="Flatten", ShowForTargetTypes="Heightmap"))
+	bool bPickValuePerApply;
+
+	// Enable to flatten towards a target height
+	UPROPERTY(NonTransactional)
+	bool bUseFlattenTarget;
+
+	// Target height to flatten towards (in Unreal Units)
+	UPROPERTY(Category="Tool Settings", EditAnywhere, NonTransactional, meta=(ShowForTools="Flatten", ShowForTargetTypes="Heightmap", EditCondition="bUseFlattenTarget", UIMin="-32768", UIMax="32768"))
+	float FlattenTarget;
+
+	// Whether to show the preview grid for the flatten target height
+	UPROPERTY(Category = "Tool Settings", EditAnywhere, NonTransactional, AdvancedDisplay, meta = (DisplayName = "Show Preview Grid", ShowForTools = "Flatten", ShowForTargetTypes = "Heightmap", EditCondition = "bUseFlattenTarget", HideEditConditionToggle, UIMin = "-32768", UIMax = "32768"))
+	bool bShowFlattenTargetPreview;
+
+	// Ramp Tool:
+
+	// Width of ramp
+	UPROPERTY(Category="Tool Settings", EditAnywhere, NonTransactional, meta=(ShowForTools="Ramp", ClampMin="1", UIMin="1", UIMax="8192", SliderExponent=3))
+	float RampWidth;
+
+	// Falloff on side of ramp
+	UPROPERTY(Category="Tool Settings", EditAnywhere, NonTransactional, meta=(DisplayName="Side Falloff", ShowForTools="Ramp", ClampMin="0", ClampMax="1", UIMin="0", UIMax="1"))
+	float RampSideFalloff;
+
+	// Smooth Tool:
+
+	// The radius smoothing is performed over
+	// Higher values smooth out bigger details, lower values only smooth out smaller details
+	UPROPERTY(Category="Tool Settings", EditAnywhere, NonTransactional, meta=(DisplayName="Filter Kernel Radius", ShowForTools="Smooth", ClampMin="1", ClampMax="31", UIMin="0", UIMax="7"))
+	int32 SmoothFilterKernelSize;
+
+	// If checked, performs a detail preserving smooth using the specified detail smoothing value
+	UPROPERTY(NonTransactional)
+	bool bDetailSmooth;
+
+	// Larger detail smoothing values remove more details, while smaller values preserve more details
+	UPROPERTY(Category="Tool Settings", EditAnywhere, NonTransactional, meta=(DisplayName="Detail Smooth", EditCondition="bDetailSmooth", ShowForTools="Smooth", ClampMin="0", ClampMax="0.99"))
+	float DetailScale;
+
+	// Erosion Tool:
+
+	// The minimum height difference necessary for the erosion effects to be applied. Smaller values will result in more erosion being applied
+	UPROPERTY(Category="Tool Settings", EditAnywhere, NonTransactional, meta=(DisplayName="Threshold", ShowForTools="Erosion", ClampMin="0", ClampMax="256", UIMin="0", UIMax="128"))
+	int32 ErodeThresh;
+
+	// The thickness of the surface for the layer weight erosion effect
+	UPROPERTY(Category="Tool Settings", EditAnywhere, NonTransactional, meta=(DisplayName="Surface Thickness", ShowForTools="Erosion", ClampMin="128", ClampMax="1024", UIMin="128", UIMax="512"))
+	int32 ErodeSurfaceThickness;
+
+	// Number of erosion iterations, more means more erosion but is slower
+	UPROPERTY(Category="Tool Settings", EditAnywhere, NonTransactional, meta=(DisplayName="Iterations", ShowForTools="Erosion", ClampMin="1", ClampMax="300", UIMin="1", UIMax="150"))
+	int32 ErodeIterationNum;
+
+	// Whether to erode by lowering, raising, or both
+	UPROPERTY(Category="Tool Settings", EditAnywhere, NonTransactional, meta=(DisplayName="Noise Mode", ShowForTools="Erosion"))
+	TEnumAsByte<ELandscapeToolErosionMode::Type> ErosionNoiseMode;
+
+	// The size of the perlin noise filter used
+	UPROPERTY(Category="Tool Settings", EditAnywhere, NonTransactional, meta=(DisplayName="Noise Scale", ShowForTools="Erosion", ClampMin="1", ClampMax="512", UIMin="1.1", UIMax="256"))
+	float ErosionNoiseScale;
+
+	// Hydraulic Erosion Tool:
+
+	// The amount of rain to apply to the surface. Larger values will result in more erosion
+	UPROPERTY(Category="Tool Settings", EditAnywhere, NonTransactional, meta=(ShowForTools="HydraErosion", ClampMin="1", ClampMax="512", UIMin="1", UIMax="256"))
+	int32 RainAmount;
+
+	// The amount of sediment that the water can carry. Larger values will result in more erosion
+	UPROPERTY(Category="Tool Settings", EditAnywhere, NonTransactional, meta=(DisplayName="Sediment Cap.", ShowForTools="HydraErosion", ClampMin="0.1", ClampMax="1.0"))
+	float SedimentCapacity;
+
+	// Number of erosion iterations, more means more erosion but is slower
+	UPROPERTY(Category="Tool Settings", EditAnywhere, NonTransactional, meta=(DisplayName="Iterations", ShowForTools="HydraErosion", ClampMin="1", ClampMax="300", UIMin="1", UIMax="150"))
+	int32 HErodeIterationNum;
+
+	// Initial Rain Distribution
+	UPROPERTY(Category="Tool Settings", EditAnywhere, NonTransactional, meta=(DisplayName="Initial Rain Distribution", ShowForTools="HydraErosion"))
+	TEnumAsByte<ELandscapeToolHydroErosionMode::Type> RainDistMode;
+
+	// The size of the noise filter for applying initial rain to the surface
+	UPROPERTY(Category="Tool Settings", EditAnywhere, NonTransactional, meta=(ShowForTools="HydraErosion", ClampMin="1", ClampMax="512", UIMin="1.1", UIMax="256"))
+	float RainDistScale;
+
+	// If checked, performs a detail-preserving smooth to the erosion effect using the specified detail smoothing value
+	UPROPERTY(NonTransactional)
+	bool bHErosionDetailSmooth;
+
+	// Larger detail smoothing values remove more details, while smaller values preserve more details
+	UPROPERTY(Category="Tool Settings", EditAnywhere, NonTransactional, meta=(DisplayName="Detail Smooth", EditCondition="bHErosionDetailSmooth", ShowForTools="HydraErosion", ClampMin="0", ClampMax="0.99"))
+	float HErosionDetailScale;
+
+	// Noise Tool:
+
+	// Whether to apply noise that raises, lowers, or both
+	UPROPERTY(Category="Tool Settings", EditAnywhere, NonTransactional, meta=(DisplayName="Noise Mode", ShowForTools="Noise"))
+	TEnumAsByte<ELandscapeToolNoiseMode::Type> NoiseMode;
+
+	// The size of the perlin noise filter used
+	UPROPERTY(Category="Tool Settings", EditAnywhere, NonTransactional, meta=(DisplayName="Noise Scale", ShowForTools="Noise", ClampMin="1", ClampMax="512", UIMin="1.1", UIMax="256"))
+	float NoiseScale;
+
+	// Mask Tool:
+
+	// Uses selected region as a mask for other tools
+	UPROPERTY(Category="Tool Settings", EditAnywhere, NonTransactional, meta=(DisplayName="Use Region as Mask", ShowForTools="Mask", ShowForMask))
+	bool bUseSelectedRegion;
+
+	// If enabled, protects the selected region from changes
+	// If disabled, only allows changes in the selected region
+	UPROPERTY(Category="Tool Settings", EditAnywhere, NonTransactional, meta=(DisplayName="Negative Mask", ShowForTools="Mask", ShowForMask))
+	bool bUseNegativeMask;
+
+	// Copy/Paste Tool:
+
+	// Whether to paste will only raise, only lower, or both
+	UPROPERTY(Category="Tool Settings", EditAnywhere, NonTransactional, meta=(ShowForTools="CopyPaste"))
+	TEnumAsByte<ELandscapeToolPasteMode::Type> PasteMode;
+
+	// If set, copies/pastes all layers, otherwise only copy/pastes the layer selected in the targets panel
+	UPROPERTY(Category="Tool Settings", EditAnywhere, NonTransactional, meta=(DisplayName="Gizmo copy/paste all layers", ShowForTools="CopyPaste"))
+	bool bApplyToAllTargets;
+
+	// Makes sure the gizmo is snapped perfectly to the landscape so that the sample points line up, which makes copy/paste less blurry. Irrelevant if gizmo is scaled
+	UPROPERTY(Category="Tool Settings", EditAnywhere, NonTransactional, meta=(DisplayName="Snap Gizmo to Landscape grid", ShowForTools="CopyPaste"))
+	bool bSnapGizmo;
+
+	// Smooths the edges of the gizmo data into the landscape. Without this, the edges of the pasted data will be sharp
+	UPROPERTY(Category="Tool Settings", EditAnywhere, NonTransactional, meta=(DisplayName="Use Smooth Gizmo Brush", ShowForTools="CopyPaste"))
+	bool bSmoothGizmoBrush;
+
+	UPROPERTY(Category="Tool Settings", EditAnywhere, NonTransactional, AdvancedDisplay, meta=(DisplayName="Heightmap", ShowForTools="CopyPaste"))
+	FString GizmoHeightmapFilenameString;
+
+	UPROPERTY(Category="Tool Settings", EditAnywhere, NonTransactional, AdvancedDisplay, meta=(DisplayName="Heightmap Size", ShowForTools="CopyPaste"))
+	FIntPoint GizmoImportSize;
+
+	UPROPERTY(Category="Tool Settings", EditAnywhere, NonTransactional, AdvancedDisplay, meta=(DisplayName="Layers", ShowForTools="CopyPaste"))
+	TArray<FGizmoImportLayer> GizmoImportLayers;
+
+	TArray<FGizmoHistory> GizmoHistories;
+
+
+	// Resize Landscape Tool
+
+	// Number of quads per landscape component section
+	UPROPERTY(Category="Change Component Size", EditAnywhere, NonTransactional, meta=(DisplayName="Section Size", ShowForTools="ResizeLandscape"))
+	int32 ResizeLandscape_QuadsPerSection;
+
+	// Number of sections per landscape component
+	UPROPERTY(Category="Change Component Size", EditAnywhere, NonTransactional, meta=(DisplayName="Sections Per Component", ShowForTools="ResizeLandscape"))
+	int32 ResizeLandscape_SectionsPerComponent;
+
+	// Number of components in resulting landscape
+	UPROPERTY(Category="Change Component Size", EditAnywhere, NonTransactional, meta=(DisplayName="Number of Components", ShowForTools="ResizeLandscape"))
+	FIntPoint ResizeLandscape_ComponentCount;
+
+	// Determines how the new component size will be applied to the existing landscape geometry.
+	UPROPERTY(Category="Change Component Size", EditAnywhere, NonTransactional, meta=(DisplayName="Resize Mode", ShowForTools="ResizeLandscape"))
+	TEnumAsByte<ELandscapeConvertMode::Type> ResizeLandscape_ConvertMode;
+
+	int32 ResizeLandscape_Original_QuadsPerSection;
+	int32 ResizeLandscape_Original_SectionsPerComponent;
+	FIntPoint ResizeLandscape_Original_ComponentCount;
+
+	// New Landscape "Tool"
+
+	// Material initially applied to the landscape. Setting a material here exposes properties for setting up layer info based on the landscape blend nodes in the material.
+	UPROPERTY(Category="New Landscape", EditAnywhere, meta=(DisplayName="Material", ShowForTools="NewLandscape"))
+	TWeakObjectPtr<UMaterialInterface> NewLandscape_Material;
+
+	// The number of quads in a single landscape section. One section is the unit of LOD transition for landscape rendering.
+	UPROPERTY(Category="New Landscape", EditAnywhere, meta=(DisplayName="Section Size", ShowForTools="NewLandscape"))
+	int32 NewLandscape_QuadsPerSection;
+
+	// The number of sections in a single landscape component. This along with the section size determines the size of each landscape component. A component is the base unit of rendering and culling.
+	UPROPERTY(Category="New Landscape", EditAnywhere, meta=(DisplayName="Sections Per Component", ShowForTools="NewLandscape"))
+	int32 NewLandscape_SectionsPerComponent;
+
+	// The number of components in the X and Y direction, determining the overall size of the landscape.
+	UPROPERTY(Category="New Landscape", EditAnywhere, meta=(DisplayName="Number of Components", ShowForTools="NewLandscape"))
+	FIntPoint NewLandscape_ComponentCount;
+
+	// The location of the new landscape
+	UPROPERTY(Category="New Landscape", EditAnywhere, meta=(DisplayName="Location", ShowForTools="NewLandscape"))
+	FVector NewLandscape_Location;
+
+	// The rotation of the new landscape
+	UPROPERTY(Category="New Landscape", EditAnywhere, meta=(DisplayName="Rotation", ShowForTools="NewLandscape"))
+	FRotator NewLandscape_Rotation;
+
+	// The scale of the new landscape. This is the distance between each vertex on the landscape, defaulting to 100 units.
+	UPROPERTY(Category="New Landscape", EditAnywhere, meta=(DisplayName="Scale", ShowForTools="NewLandscape"))
+	FVector NewLandscape_Scale;
+
+	UPROPERTY(Category="New Landscape", VisibleAnywhere, NonTransactional, meta=(ShowForTools="NewLandscape"))
+	TEnumAsByte<ELandscapeImportHeightmapError::Type> ImportLandscape_HeightmapError;
+
+	// Specify a height map file in 16-bit RAW or PNG format
+	UPROPERTY(Category="New Landscape", EditAnywhere, NonTransactional, meta=(DisplayName="Heightmap File", ShowForTools="NewLandscape"))
+	FString ImportLandscape_HeightmapFilename;
+	UPROPERTY(NonTransactional)
+	int32 ImportLandscape_Width;
+	UPROPERTY(NonTransactional)
+	int32 ImportLandscape_Height;
+	UPROPERTY(NonTransactional)
+	TArray<uint16> ImportLandscape_Data;
+
+	// The landscape layers that will be created. Only layer names referenced in the material assigned above are shown here. Modify the material to add more layers.
+	UPROPERTY(Category="New Landscape", EditAnywhere, NonTransactional, EditFixedSize, meta=(DisplayName="Layers", ShowForTools="NewLandscape"))
+	TArray<FLandscapeImportLayer> ImportLandscape_Layers;
+
+	// Common Brush Settings:
+
+	// The radius of the brush, in unreal units
+	UPROPERTY(Category="Brush Settings", EditAnywhere, NonTransactional, meta=(DisplayName="Brush Size", ShowForBrushes="BrushSet_Circle,BrushSet_Alpha,BrushSet_Pattern", ClampMin="1", ClampMax="65536", UIMin="1", UIMax="8192", SliderExponent="3"))
+	float BrushRadius;
+
+	// The falloff at the edge of the brush, as a fraction of the brush's size. 0 = no falloff, 1 = all falloff
+	UPROPERTY(Category="Brush Settings", EditAnywhere, NonTransactional, meta=(ShowForBrushes="BrushSet_Circle,BrushSet_Gizmo,BrushSet_Pattern", ClampMin="0", ClampMax="1"))
+	float BrushFalloff;
+
+	// Selects the Clay Brush painting mode
+	UPROPERTY(Category="Brush Settings", EditAnywhere, NonTransactional, meta=(ShowForTools="Sculpt", ShowForBrushes="BrushSet_Circle,BrushSet_Alpha,BrushSet_Pattern"))
+	bool bUseClayBrush;
+
+	// Alpha/Pattern Brush:
+
+	// Scale of the brush texture. A scale of 1.000 maps the brush texture to the landscape at a 1 pixel = 1 vertex size
+	UPROPERTY(Category="Brush Settings", EditAnywhere, NonTransactional, meta=(DisplayName="Texture Scale", ShowForBrushes="BrushSet_Pattern", ClampMin="0.005", ClampMax="5", SliderExponent="3"))
+	float AlphaBrushScale;
+
+	// Rotate brush to follow mouse
+	UPROPERTY(Category = "Brush Settings", EditAnywhere, NonTransactional, meta = (DisplayName = "Auto-Rotate", ShowForBrushes = "BrushSet_Alpha"))
+	bool bAlphaBrushAutoRotate;
+
+	// Rotates the brush mask texture
+	UPROPERTY(Category="Brush Settings", EditAnywhere, NonTransactional, meta=(DisplayName="Texture Rotation", ShowForBrushes="BrushSet_Alpha,BrushSet_Pattern", ClampMin="-360", ClampMax="360", UIMin="-180", UIMax="180"))
+	float AlphaBrushRotation;
+
+	// Horizontally offsets the brush mask texture
+	UPROPERTY(Category="Brush Settings", EditAnywhere, NonTransactional, meta=(DisplayName="Texture Pan U", ShowForBrushes="BrushSet_Pattern", ClampMin="0", ClampMax="1"))
+	float AlphaBrushPanU;
+
+	// Vertically offsets the brush mask texture
+	UPROPERTY(Category="Brush Settings", EditAnywhere, NonTransactional, meta=(DisplayName="Texture Pan V", ShowForBrushes="BrushSet_Pattern", ClampMin="0", ClampMax="1"))
+	float AlphaBrushPanV;
+
+	UPROPERTY(Category = "Brush Settings", EditAnywhere, NonTransactional, meta = (DisplayName = "Use World-Space", ShowForBrushes = "BrushSet_Pattern"))
+	bool bUseWorldSpacePatternBrush;
+
+	UPROPERTY(Category = "Brush Settings", EditAnywhere, NonTransactional, meta = (DisplayName = "World-Space Settings", ShowForBrushes = "BrushSet_Pattern", EditCondition = "bUseWorldSpacePatternBrush"))
+	FLandscapePatternBrushWorldSpaceSettings WorldSpacePatternBrushSettings;
+
+	// Mask texture to use
+	UPROPERTY(Category="Brush Settings", EditAnywhere, NonTransactional, meta=(DisplayName="Texture", ShowForBrushes="BrushSet_Alpha,BrushSet_Pattern"))
+	UTexture2D* AlphaTexture;
+
+	// Channel of Mask Texture to use
+	UPROPERTY(Category="Brush Settings", EditAnywhere, NonTransactional, meta=(DisplayName="Texture Channel", ShowForBrushes="BrushSet_Alpha,BrushSet_Pattern"))
+	TEnumAsByte<EColorChannel::Type> AlphaTextureChannel;
+
+	UPROPERTY(NonTransactional)
+	int32 AlphaTextureSizeX;
+	UPROPERTY(NonTransactional)
+	int32 AlphaTextureSizeY;
+	UPROPERTY(NonTransactional)
+	TArray<uint8> AlphaTextureData;
+
+	// Component Brush:
+
+	// Number of components X/Y to affect at once. 1 means 1x1, 2 means 2x2, etc
+	UPROPERTY(Category="Brush Settings", EditAnywhere, NonTransactional, meta=(DisplayName="Brush Size", ShowForBrushes="BrushSet_Component", ClampMin="1", ClampMax="128", UIMin="1", UIMax="64", SliderExponent="3"))
+	int32 BrushComponentSize;
+
+
+	// Target Layer Settings:
+
+	// Limits painting to only the components that already have the selected layer
+	UPROPERTY(Category="Target Layers", EditAnywhere, NonTransactional, meta=(ShowForTargetTypes="Weightmap,Visibility"))
+	TEnumAsByte<ELandscapeLayerPaintingRestriction::Type> PaintingRestriction;
+
+#if WITH_EDITOR
+	// Begin UObject Interface
+	void PostEditChangeProperty(FPropertyChangedEvent& PropertyChangedEvent) override;
+	// End UObject Interface
+#endif // WITH_EDITOR
+
+	void Load();
+	void Save();
+
+	// Region
+	void SetbUseSelectedRegion(bool InbUseSelectedRegion);
+	void SetbUseNegativeMask(bool InbUseNegativeMask);
+
+	// Copy/Paste
+	void SetPasteMode(ELandscapeToolPasteMode::Type InPasteMode);
+	void GuessGizmoImportSize();
+
+	// Alpha/Pattern Brush
+	bool SetAlphaTexture(UTexture2D* InTexture, EColorChannel::Type InTextureChannel);
+
+	// New Landscape
+	FString LastImportPath;
+
+	const TArray<uint16>& GetImportLandscapeData();
+	void ClearImportLandscapeData() { ImportLandscape_Data.Empty(); }
+
+	void RefreshImportLayersList();
+	
+	int32 ClampLandscapeSize(int32 InComponenetsCount) const
+	{
+		// Max size is either whole components below 8192 verts, or 32 components 
+		return FMath::Clamp(InComponenetsCount, 1, FMath::Min(32, FMath::FloorToInt(8191 / (NewLandscape_SectionsPerComponent * NewLandscape_QuadsPerSection))));
+	}
+	
+	int32 CalcComponentsCount(int32 InResolution) const
+	{
+		return ClampLandscapeSize(InResolution / (NewLandscape_SectionsPerComponent * NewLandscape_QuadsPerSection));
+	}
+
+	void NewLandscape_ClampSize()
+	{
+		NewLandscape_ComponentCount.X = ClampLandscapeSize(NewLandscape_ComponentCount.X);
+		NewLandscape_ComponentCount.Y = ClampLandscapeSize(NewLandscape_ComponentCount.Y);
+	}
+
+	void UpdateComponentCount()
+	{
+		// ignore invalid cases
+		if (ResizeLandscape_QuadsPerSection == 0 || ResizeLandscape_SectionsPerComponent == 0 ||ResizeLandscape_ComponentCount.X == 0 || ResizeLandscape_ComponentCount.Y == 0) 
+		{
+			return;
+		}
+		const int32 ComponentSizeQuads = ResizeLandscape_QuadsPerSection * ResizeLandscape_SectionsPerComponent;
+		const int32 Original_ComponentSizeQuads = ResizeLandscape_Original_QuadsPerSection * ResizeLandscape_Original_SectionsPerComponent;
+		const FIntPoint OriginalResolution = ResizeLandscape_Original_ComponentCount * Original_ComponentSizeQuads;
+		switch (ResizeLandscape_ConvertMode)
+		{
+		case ELandscapeConvertMode::Expand:
+			ResizeLandscape_ComponentCount.X = FMath::DivideAndRoundUp(OriginalResolution.X, ComponentSizeQuads);
+			ResizeLandscape_ComponentCount.Y = FMath::DivideAndRoundUp(OriginalResolution.Y, ComponentSizeQuads);
+			break;
+		case ELandscapeConvertMode::Clip:
+			ResizeLandscape_ComponentCount.X = FMath::Max(1, OriginalResolution.X / ComponentSizeQuads);
+			ResizeLandscape_ComponentCount.Y = FMath::Max(1, OriginalResolution.Y / ComponentSizeQuads);
+			break;
+		case ELandscapeConvertMode::Resample:
+			ResizeLandscape_ComponentCount = ResizeLandscape_Original_ComponentCount;
+			break;
+		default:
+			check(0);
+		}
+	}
+
+	void SetbSnapGizmo(bool InbSnapGizmo);
+
+	void SetParent(FEdModeLandscape* LandscapeParent)
+	{
+		ParentMode = LandscapeParent;
+	}
+};
diff --git a/Engine/Source/Editor/LandscapeEditor/Public/LandscapeEditorUtils.h b/Engine/Source/Editor/LandscapeEditor/Public/LandscapeEditorUtils.h
new file mode 100644
--- /dev/null
+++ b/Engine/Source/Editor/LandscapeEditor/Public/LandscapeEditorUtils.h
@@ -0,0 +1,100 @@
+// Copyright 1998-2015 Epic Games, Inc. All Rights Reserved.
+
+#pragma once
+
+// Forward declarations
+class ULandscapeLayerInfoObject;
+
+namespace LandscapeEditorUtils
+{
+	template<typename T>
+	void ExpandData(T* OutData, const T* InData,
+		int32 OldMinX, int32 OldMinY, int32 OldMaxX, int32 OldMaxY,
+		int32 NewMinX, int32 NewMinY, int32 NewMaxX, int32 NewMaxY)
+	{
+		const int32 OldWidth = OldMaxX - OldMinX + 1;
+		const int32 OldHeight = OldMaxY - OldMinY + 1;
+		const int32 NewWidth = NewMaxX - NewMinX + 1;
+		const int32 NewHeight = NewMaxY - NewMinY + 1;
+		const int32 OffsetX = NewMinX - OldMinX;
+		const int32 OffsetY = NewMinY - OldMinY;
+
+		for (int32 Y = 0; Y < NewHeight; ++Y)
+		{
+			const int32 OldY = FMath::Clamp<int32>(Y + OffsetY, 0, OldHeight - 1);
+
+			// Pad anything to the left
+			const T PadLeft = InData[OldY * OldWidth + 0];
+			for (int32 X = 0; X < -OffsetX; ++X)
+			{
+				OutData[Y * NewWidth + X] = PadLeft;
+			}
+
+			// Copy one row of the old data
+			{
+				const int32 X = FMath::Max(0, -OffsetX);
+				const int32 OldX = FMath::Clamp<int32>(X + OffsetX, 0, OldWidth - 1);
+				FMemory::Memcpy(&OutData[Y * NewWidth + X], &InData[OldY * OldWidth + OldX], FMath::Min<int32>(OldWidth, NewWidth) * sizeof(T));
+			}
+
+			const T PadRight = InData[OldY * OldWidth + OldWidth - 1];
+			for (int32 X = -OffsetX + OldWidth; X < NewWidth; ++X)
+			{
+				OutData[Y * NewWidth + X] = PadRight;
+			}
+		}
+	}
+
+	template<typename T>
+	TArray<T> ExpandData(const TArray<T>& Data,
+		int32 OldMinX, int32 OldMinY, int32 OldMaxX, int32 OldMaxY,
+		int32 NewMinX, int32 NewMinY, int32 NewMaxX, int32 NewMaxY)
+	{
+		const int32 NewWidth = NewMaxX - NewMinX + 1;
+		const int32 NewHeight = NewMaxY - NewMinY + 1;
+
+		TArray<T> Result;
+		Result.Empty(NewWidth * NewHeight);
+		Result.AddUninitialized(NewWidth * NewHeight);
+
+		ExpandData(Result.GetData(), Data.GetData(),
+			OldMinX, OldMinY, OldMaxX, OldMaxY,
+			NewMinX, NewMinY, NewMaxX, NewMaxY);
+
+		return Result;
+	}
+
+	template<typename T>
+	TArray<T> ResampleData(const TArray<T>& Data, int32 OldWidth, int32 OldHeight, int32 NewWidth, int32 NewHeight)
+	{
+		TArray<T> Result;
+		Result.Empty(NewWidth * NewHeight);
+		Result.AddUninitialized(NewWidth * NewHeight);
+
+		const float XScale = (float)(OldWidth - 1) / (NewWidth - 1);
+		const float YScale = (float)(OldHeight - 1) / (NewHeight - 1);
+		for (int32 Y = 0; Y < NewHeight; ++Y)
+		{
+			for (int32 X = 0; X < NewWidth; ++X)
+			{
+				const float OldY = Y * YScale;
+				const float OldX = X * XScale;
+				const int32 X0 = FMath::FloorToInt(OldX);
+				const int32 X1 = FMath::Min(FMath::FloorToInt(OldX) + 1, OldWidth - 1);
+				const int32 Y0 = FMath::FloorToInt(OldY);
+				const int32 Y1 = FMath::Min(FMath::FloorToInt(OldY) + 1, OldHeight - 1);
+				const T& Original00 = Data[Y0 * OldWidth + X0];
+				const T& Original10 = Data[Y0 * OldWidth + X1];
+				const T& Original01 = Data[Y1 * OldWidth + X0];
+				const T& Original11 = Data[Y1 * OldWidth + X1];
+				Result[Y * NewWidth + X] = FMath::BiLerp(Original00, Original10, Original01, Original11, FMath::Fractional(OldX), FMath::Fractional(OldY));
+			}
+		}
+
+		return Result;
+	}
+
+	bool LANDSCAPEEDITOR_API SetHeightmapData(ALandscapeProxy* Landscape, const TArray<uint16>& Data);
+
+	bool LANDSCAPEEDITOR_API SetWeightmapData(ALandscapeProxy* Landscape, ULandscapeLayerInfoObject* LayerObject, const TArray<uint8>& Data);
+}
diff --git a/Engine/Source/Editor/LandscapeEditor/Public/LandscapeToolInterface.h b/Engine/Source/Editor/LandscapeEditor/Public/LandscapeToolInterface.h
new file mode 100644
--- /dev/null
+++ b/Engine/Source/Editor/LandscapeEditor/Public/LandscapeToolInterface.h
@@ -0,0 +1,262 @@
+// Copyright 1998-2015 Epic Games, Inc. All Rights Reserved.
+
+#pragma once
+
+// Forward declarations
+class ULandscapeInfo;
+class ULandscapeLayerInfoObject;
+
+// FLandscapeToolMousePosition - Struct to store mouse positions since the last time we applied the brush
+struct FLandscapeToolMousePosition
+{
+	// Stored in heightmap space.
+	FVector2D Position;
+	bool bShiftDown;
+
+	FLandscapeToolMousePosition(FVector2D InPosition, bool InbShiftDown)
+		: Position(InPosition)
+		, bShiftDown(InbShiftDown)
+	{
+	}
+};
+
+enum class ELandscapeBrushType
+{
+	Normal = 0,
+	Alpha,
+	Component,
+	Gizmo,
+	Splines
+};
+
+class FLandscapeBrushData
+{
+protected:
+	FIntRect Bounds;
+	TArray<float> BrushAlpha;
+
+public:
+	FLandscapeBrushData()
+		: Bounds()
+	{
+	}
+
+	FLandscapeBrushData(FIntRect InBounds)
+		: Bounds(InBounds)
+	{
+		BrushAlpha.SetNumZeroed(Bounds.Area());
+	}
+
+	FIntRect GetBounds() const
+	{
+		return Bounds;
+	}
+
+	// For compatibility with older landscape code that uses inclusive bounds in 4 int32s
+	void GetInclusiveBounds(int32& X1, int32& Y1, int32& X2, int32& Y2) const
+	{
+		X1 = Bounds.Min.X;
+		Y1 = Bounds.Min.Y;
+		X2 = Bounds.Max.X - 1;
+		Y2 = Bounds.Max.Y - 1;
+	}
+
+	float* GetDataPtr(FIntPoint Position)
+	{
+		return BrushAlpha.GetData() + (Position.Y - Bounds.Min.Y) * Bounds.Width() + (Position.X - Bounds.Min.X);
+	}
+	const float* GetDataPtr(FIntPoint Position) const
+	{
+		return BrushAlpha.GetData() + (Position.Y - Bounds.Min.Y) * Bounds.Width() + (Position.X - Bounds.Min.X);
+	}
+
+	FORCEINLINE_EXPLICIT_OPERATOR_BOOL() const
+	{
+		return BrushAlpha.Num() != 0;
+	}
+
+	FORCEINLINE bool operator!() const
+	{
+		return !(bool)*this;
+	}
+};
+
+class FLandscapeBrush : public FGCObject
+{
+public:
+	virtual void MouseMove(float LandscapeX, float LandscapeY) = 0;
+	virtual FLandscapeBrushData ApplyBrush(const TArray<FLandscapeToolMousePosition>& MousePositions) = 0;
+	virtual TOptional<bool> InputKey(FEditorViewportClient* InViewportClient, FViewport* InViewport, FKey InKey, EInputEvent InEvent) { return TOptional<bool>(); }
+	virtual void Tick(FEditorViewportClient* ViewportClient, float DeltaTime) {};
+	virtual void BeginStroke(float LandscapeX, float LandscapeY, class FLandscapeTool* CurrentTool);
+	virtual void EndStroke();
+	virtual void EnterBrush() {}
+	virtual void LeaveBrush() {}
+	virtual ~FLandscapeBrush() {}
+	virtual UMaterialInterface* GetBrushMaterial() { return NULL; }
+	virtual const TCHAR* GetBrushName() = 0;
+	virtual FText GetDisplayName() = 0;
+	virtual ELandscapeBrushType GetBrushType() { return ELandscapeBrushType::Normal; }
+
+	// FGCObject interface
+	virtual void AddReferencedObjects(FReferenceCollector& Collector) override {}
+};
+
+struct FLandscapeBrushSet
+{
+	FLandscapeBrushSet(const TCHAR* InBrushSetName)
+		: BrushSetName(InBrushSetName)
+		, PreviousBrushIndex(0)
+	{
+	}
+
+	const FName BrushSetName;
+	TArray<FLandscapeBrush*> Brushes;
+	int32 PreviousBrushIndex;
+
+	virtual ~FLandscapeBrushSet()
+	{
+		for (int32 BrushIdx = 0; BrushIdx < Brushes.Num(); BrushIdx++)
+		{
+			delete Brushes[BrushIdx];
+		}
+	}
+};
+
+namespace ELandscapeToolTargetType
+{
+	enum Type : int8
+	{
+		Heightmap  = 0,
+		Weightmap  = 1,
+		Visibility = 2,
+
+		Invalid    = -1, // only valid for LandscapeEdMode->CurrentToolTarget.TargetType
+	};
+}
+
+namespace ELandscapeToolTargetTypeMask
+{
+	enum Type : uint8
+	{
+		Heightmap  = 1 << ELandscapeToolTargetType::Heightmap,
+		Weightmap  = 1 << ELandscapeToolTargetType::Weightmap,
+		Visibility = 1 << ELandscapeToolTargetType::Visibility,
+
+		NA = 0,
+		All = 0xFF,
+	};
+
+	inline ELandscapeToolTargetTypeMask::Type FromType(ELandscapeToolTargetType::Type TargetType)
+	{
+		if (TargetType == ELandscapeToolTargetType::Invalid)
+		{
+			return ELandscapeToolTargetTypeMask::NA;
+		}
+		return (ELandscapeToolTargetTypeMask::Type)(1 << TargetType);
+	}
+}
+
+namespace ELandscapeToolNoiseMode
+{
+	enum Type
+	{
+		Invalid = -1,
+		Both = 0,
+		Add = 1,
+		Sub = 2,
+	};
+
+	inline float Conversion(Type Mode, float NoiseAmount, float OriginalValue)
+	{
+		switch (Mode)
+		{
+		case Add: // always +
+			OriginalValue += NoiseAmount;
+			break;
+		case Sub: // always -
+			OriginalValue -= NoiseAmount;
+			break;
+		case Both:
+			break;
+		}
+		return OriginalValue;
+	}
+}
+
+struct FLandscapeToolTarget
+{
+	TWeakObjectPtr<ULandscapeInfo> LandscapeInfo;
+	ELandscapeToolTargetType::Type TargetType;
+	TWeakObjectPtr<ULandscapeLayerInfoObject> LayerInfo;
+	FName LayerName;
+
+	FLandscapeToolTarget()
+		: LandscapeInfo(NULL)
+		, TargetType(ELandscapeToolTargetType::Heightmap)
+		, LayerInfo(NULL)
+		, LayerName(NAME_None)
+	{
+	}
+};
+
+enum class ELandscapeToolType
+{
+	Normal = 0,
+	Mask,
+};
+
+/**
+ * FLandscapeTool
+ */
+class FLandscapeTool : public FGCObject
+{
+public:
+	virtual void EnterTool() {}
+	virtual void ExitTool() {}
+	virtual bool BeginTool(FEditorViewportClient* ViewportClient, const FLandscapeToolTarget& Target, const FVector& InHitLocation) = 0;
+	virtual void EndTool(FEditorViewportClient* ViewportClient) = 0;
+	virtual void Tick(FEditorViewportClient* ViewportClient, float DeltaTime) {};
+	virtual bool MouseMove(FEditorViewportClient* ViewportClient, FViewport* Viewport, int32 x, int32 y) = 0;
+	virtual bool HandleClick(HHitProxy* HitProxy, const FViewportClick& Click) { return false; }
+	virtual bool InputKey(FEditorViewportClient* InViewportClient, FViewport* InViewport, FKey InKey, EInputEvent InEvent) { return false; }
+	virtual bool InputDelta(FEditorViewportClient* InViewportClient, FViewport* InViewport, FVector& InDrag, FRotator& InRot, FVector& InScale) { return false; }
+	FLandscapeTool() : PreviousBrushIndex(-1) {}
+	virtual ~FLandscapeTool() {}
+	virtual const TCHAR* GetToolName() = 0;
+	virtual FText GetDisplayName() = 0;
+	virtual void SetEditRenderType();
+	virtual void Render(const FSceneView* View, FViewport* Viewport, FPrimitiveDrawInterface* PDI) {}
+	virtual bool SupportsMask() { return true; }
+	virtual bool SupportsComponentSelection() { return false; }
+	virtual bool OverrideSelection() const { return false; }
+	virtual bool IsSelectionAllowed(AActor* InActor, bool bInSelection) const { return false; }
+	virtual bool UsesTransformWidget() const { return false; }
+	virtual EAxisList::Type GetWidgetAxisToDraw(FWidget::EWidgetMode InWidgetMode) const { return EAxisList::All; }
+	virtual FVector GetWidgetLocation() const { return FVector::ZeroVector; }
+	virtual FMatrix GetWidgetRotation() const { return FMatrix::Identity; }
+	virtual bool DisallowMouseDeltaTracking() const { return false; }
+
+	virtual EEditAction::Type GetActionEditDuplicate() { return EEditAction::Skip; }
+	virtual EEditAction::Type GetActionEditDelete() { return EEditAction::Skip; }
+	virtual EEditAction::Type GetActionEditCut() { return EEditAction::Skip; }
+	virtual EEditAction::Type GetActionEditCopy() { return EEditAction::Skip; }
+	virtual EEditAction::Type GetActionEditPaste() { return EEditAction::Skip; }
+	virtual bool ProcessEditDuplicate() { return false; }
+	virtual bool ProcessEditDelete() { return false; }
+	virtual bool ProcessEditCut() { return false; }
+	virtual bool ProcessEditCopy() { return false; }
+	virtual bool ProcessEditPaste() { return false; }
+
+	// Functions which doesn't need Viewport data...
+	virtual void Process(int32 Index, int32 Arg) {}
+	virtual ELandscapeToolType GetToolType() { return ELandscapeToolType::Normal; }
+	virtual ELandscapeToolTargetTypeMask::Type GetSupportedTargetTypes() { return ELandscapeToolTargetTypeMask::NA; };
+
+	// FGCObject interface
+	virtual void AddReferencedObjects(FReferenceCollector& Collector) override {}
+
+public:
+	int32					PreviousBrushIndex;
+	TArray<FName>			ValidBrushes;
+};
diff --git a/Engine/Source/Editor/WorldBrowser/Private/Tiles/WorldTileCollectionModel.cpp b/Engine/Source/Editor/WorldBrowser/Private/Tiles/WorldTileCollectionModel.cpp
--- a/Engine/Source/Editor/WorldBrowser/Private/Tiles/WorldTileCollectionModel.cpp
+++ b/Engine/Source/Editor/WorldBrowser/Private/Tiles/WorldTileCollectionModel.cpp
@@ -10,7 +10,7 @@
 #include "MeshUtilities.h"
 #include "MaterialUtilities.h"
 #include "RawMesh.h"
-#include "LandscapeEdMode.h"
+#include "LandscapeEditorUtils.h"
 #include "ImageWrapper.h"
 
 #include "WorldTileDetails.h"
