Fixed external reference errors caused by moving components with an XY-offset map texture to another level   UE-18720

diff --git a/Engine/Source/Editor/LandscapeEditor/Private/LandscapeEdModeComponentTools.cpp b/Engine/Source/Editor/LandscapeEditor/Private/LandscapeEdModeComponentTools.cpp
--- a/Engine/Source/Editor/LandscapeEditor/Private/LandscapeEdModeComponentTools.cpp
+++ b/Engine/Source/Editor/LandscapeEditor/Private/LandscapeEdModeComponentTools.cpp
@@ -412,17 +412,14 @@ public:
 
 				GWarn->BeginSlowTask(LOCTEXT("BeginMovingLandscapeComponentsToCurrentLevelTask", "Moving Landscape components to current level"), true);
 
-				TSet<ALandscapeProxy*> SelectProxies;
-				TSet<UTexture2D*> OldTextureSet;
-				TSet<ULandscapeComponent*> TargetSelectedComponents;
-				TArray<ULandscapeHeightfieldCollisionComponent*> TargetSelectedCollisionComponents;
-				TSet<ULandscapeComponent*> HeightmapUpdateComponents;
-
 				int32 Progress = 0;
 				LandscapeInfo->SortSelectedComponents();
-				int32 ComponentSizeVerts = Landscape->NumSubsections * (Landscape->SubsectionSizeQuads + 1);
-				int32 NeedHeightmapSize = 1 << FMath::CeilLogTwo(ComponentSizeVerts);
+				const int32 ComponentSizeVerts = Landscape->NumSubsections * (Landscape->SubsectionSizeQuads + 1);
+				const int32 NeedHeightmapSize = 1 << FMath::CeilLogTwo(ComponentSizeVerts);
 
+				TSet<ALandscapeProxy*> SelectProxies;
+				TSet<ULandscapeComponent*> TargetSelectedComponents;
+				TArray<ULandscapeHeightfieldCollisionComponent*> TargetSelectedCollisionComponents;
 				for (ULandscapeComponent* Component : SelectedComponents)
 				{
 					SelectProxies.Add(Component->GetLandscapeProxy());
@@ -442,52 +439,46 @@ public:
 				int32 TotalProgress = TargetSelectedComponents.Num() * TargetSelectedCollisionComponents.Num();
 
 				// Check which ones are need for height map change
+				TSet<UTexture2D*> OldHeightmapTextures;
 				for (ULandscapeComponent* Component : TargetSelectedComponents)
 				{
 					Component->Modify();
-					OldTextureSet.Add(Component->HeightmapTexture);
+					OldHeightmapTextures.Add(Component->HeightmapTexture);
 				}
 
 				// Need to split all the component which share Heightmap with selected components
-				// Search neighbor only
+				TMap<ULandscapeComponent*, bool> HeightmapUpdateComponents;
+				HeightmapUpdateComponents.Reserve(TargetSelectedComponents.Num() * 4); // worst case
 				for (ULandscapeComponent* Component : TargetSelectedComponents)
 				{
-					int32 SearchX = Component->HeightmapTexture->Source.GetSizeX() / NeedHeightmapSize;
-					int32 SearchY = Component->HeightmapTexture->Source.GetSizeY() / NeedHeightmapSize;
-					FIntPoint ComponentBase = Component->GetSectionBase() / Component->ComponentSizeQuads;
+					// Search neighbor only
+					const int32 SearchX = Component->HeightmapTexture->Source.GetSizeX() / NeedHeightmapSize - 1;
+					const int32 SearchY = Component->HeightmapTexture->Source.GetSizeY() / NeedHeightmapSize - 1;
+					const FIntPoint ComponentBase = Component->GetSectionBase() / Component->ComponentSizeQuads;
 
-					for (int32 Y = 0; Y < SearchY; ++Y)
+					for (int32 Y = -SearchY; Y <= SearchY; ++Y)
 					{
-						for (int32 X = 0; X < SearchX; ++X)
+						for (int32 X = -SearchX; X <= SearchX; ++X)
 						{
-							// Search for four directions...
-							for (int32 Dir = 0; Dir < 4; ++Dir)
+							ULandscapeComponent* const Neighbor = LandscapeInfo->XYtoComponentMap.FindRef(ComponentBase + FIntPoint(X, Y));
+							if (Neighbor && Neighbor->HeightmapTexture == Component->HeightmapTexture && !HeightmapUpdateComponents.Contains(Neighbor))
 							{
-								int32 XDir = (Dir >> 1) ? 1 : -1;
-								int32 YDir = (Dir % 2) ? 1 : -1;
-								ULandscapeComponent* Neighbor = LandscapeInfo->XYtoComponentMap.FindRef(ComponentBase + FIntPoint(XDir*X, YDir*Y));
-								if (Neighbor && Neighbor->HeightmapTexture == Component->HeightmapTexture && !HeightmapUpdateComponents.Contains(Neighbor))
-								{
-									Neighbor->Modify();
-									if (!TargetSelectedComponents.Contains(Neighbor))
-									{
-										Neighbor->HeightmapScaleBias.X = -1.0f; // just mark this component is for original level, not current level
-									}
-									HeightmapUpdateComponents.Add(Neighbor);
-								}
+								Neighbor->Modify();
+								bool bNeedsMoveToCurrentLevel = TargetSelectedComponents.Contains(Neighbor);
+								HeightmapUpdateComponents.Add(Neighbor, bNeedsMoveToCurrentLevel);
 							}
 						}
 					}
 				}
 
 				// Changing Heightmap format for selected components
-				for (ULandscapeComponent* Component : HeightmapUpdateComponents)
+				for (const auto& HeightmapUpdateComponentPair : HeightmapUpdateComponents)
 				{
-					ALandscape::SplitHeightmap(Component, (Component->HeightmapScaleBias.X > 0.0f));
+					ALandscape::SplitHeightmap(HeightmapUpdateComponentPair.Key, HeightmapUpdateComponentPair.Value);
 				}
 
 				// Delete if it is no referenced textures...
-				for (UTexture2D* Texture : OldTextureSet)
+				for (UTexture2D* Texture : OldHeightmapTextures)
 				{
 					Texture->SetFlags(RF_Transactional);
 					Texture->Modify();
@@ -523,6 +514,16 @@ public:
 				LandscapeProxy->Modify();
 				LandscapeProxy->MarkPackageDirty();
 
+				// Handle XY-offset textures (these don't need splitting, as they aren't currently shared between components like heightmaps/weightmaps can be)
+				for (ULandscapeComponent* Component : TargetSelectedComponents)
+				{
+					if (Component->XYOffsetmapTexture)
+					{
+						Component->XYOffsetmapTexture->Modify();
+						Component->XYOffsetmapTexture->Rename(nullptr, LandscapeProxy->GetOutermost());
+					}
+				}
+
 				// Change Weight maps...
 				{
 					FLandscapeEditDataInterface LandscapeEdit(LandscapeInfo);
