Fixed crash when using the landscape resize tool

diff --git a/Engine/Source/Editor/LandscapeEditor/Private/LandscapeEdMode.cpp b/Engine/Source/Editor/LandscapeEditor/Private/LandscapeEdMode.cpp
--- a/Engine/Source/Editor/LandscapeEditor/Private/LandscapeEdMode.cpp
+++ b/Engine/Source/Editor/LandscapeEditor/Private/LandscapeEdMode.cpp
@@ -1714,6 +1714,7 @@ int32 FEdModeLandscape::UpdateLandscapeList()
 			}
 			CurrentToolTarget.LandscapeInfo = LandscapeList[0].Info;
 			CurrentIndex = 0;
+			UpdateTargetList();
 			if (CurrentTool != NULL)
 			{
 				CurrentTool->EnterTool();
