Added landscape mirroring tool

diff --git a/Engine/Source/Editor/EditorStyle/Private/SlateEditorStyle.cpp b/Engine/Source/Editor/EditorStyle/Private/SlateEditorStyle.cpp
--- a/Engine/Source/Editor/EditorStyle/Private/SlateEditorStyle.cpp
+++ b/Engine/Source/Editor/EditorStyle/Private/SlateEditorStyle.cpp
@@ -5640,55 +5640,58 @@ void FSlateEditorStyle::FStyle::SetupLandscapeEditorStyle()
 		Set("LandscapeEditor.NoiseTool",        new IMAGE_BRUSH("Icons/icon_Landscape_Tool_Noise_40x",            Icon40x40));
 		Set("LandscapeEditor.RetopologizeTool", new IMAGE_BRUSH("Icons/icon_Landscape_Tool_Retopologize_40x",     Icon40x40));
 		Set("LandscapeEditor.VisibilityTool",   new IMAGE_BRUSH("Icons/icon_Landscape_Tool_Visibility_40x",       Icon40x40));
-		Set("LandscapeEditor.SculptTool.Small", new IMAGE_BRUSH("Icons/icon_Landscape_Tool_Sculpt_20x", Icon20x20));
-		Set("LandscapeEditor.PaintTool.Small", new IMAGE_BRUSH("Icons/icon_Landscape_Tool_Paint_20x", Icon20x20));
-		Set("LandscapeEditor.SmoothTool.Small", new IMAGE_BRUSH("Icons/icon_Landscape_Tool_Smooth_20x", Icon20x20));
-		Set("LandscapeEditor.FlattenTool.Small", new IMAGE_BRUSH("Icons/icon_Landscape_Tool_Flatten_20x", Icon20x20));
-		Set("LandscapeEditor.RampTool.Small", new IMAGE_BRUSH("Icons/icon_Landscape_Tool_Ramp_20x", Icon20x20));
-		Set("LandscapeEditor.ErosionTool.Small", new IMAGE_BRUSH("Icons/icon_Landscape_Tool_Erosion_20x", Icon20x20));
+		Set("LandscapeEditor.SculptTool.Small",       new IMAGE_BRUSH("Icons/icon_Landscape_Tool_Sculpt_20x",           Icon20x20));
+		Set("LandscapeEditor.PaintTool.Small",        new IMAGE_BRUSH("Icons/icon_Landscape_Tool_Paint_20x",            Icon20x20));
+		Set("LandscapeEditor.SmoothTool.Small",       new IMAGE_BRUSH("Icons/icon_Landscape_Tool_Smooth_20x",           Icon20x20));
+		Set("LandscapeEditor.FlattenTool.Small",      new IMAGE_BRUSH("Icons/icon_Landscape_Tool_Flatten_20x",          Icon20x20));
+		Set("LandscapeEditor.RampTool.Small",         new IMAGE_BRUSH("Icons/icon_Landscape_Tool_Ramp_20x",             Icon20x20));
+		Set("LandscapeEditor.ErosionTool.Small",      new IMAGE_BRUSH("Icons/icon_Landscape_Tool_Erosion_20x",          Icon20x20));
 		Set("LandscapeEditor.HydroErosionTool.Small", new IMAGE_BRUSH("Icons/icon_Landscape_Tool_HydraulicErosion_20x", Icon20x20));
-		Set("LandscapeEditor.NoiseTool.Small", new IMAGE_BRUSH("Icons/icon_Landscape_Tool_Noise_20x", Icon20x20));
-		Set("LandscapeEditor.RetopologizeTool.Small", new IMAGE_BRUSH("Icons/icon_Landscape_Tool_Retopologize_20x", Icon20x20));
-		Set("LandscapeEditor.VisibilityTool.Small", new IMAGE_BRUSH("Icons/icon_Landscape_Tool_Visibility_20x", Icon20x20));
+		Set("LandscapeEditor.NoiseTool.Small",        new IMAGE_BRUSH("Icons/icon_Landscape_Tool_Noise_20x",            Icon20x20));
+		Set("LandscapeEditor.RetopologizeTool.Small", new IMAGE_BRUSH("Icons/icon_Landscape_Tool_Retopologize_20x",     Icon20x20));
+		Set("LandscapeEditor.VisibilityTool.Small",   new IMAGE_BRUSH("Icons/icon_Landscape_Tool_Visibility_20x",       Icon20x20));
 
-		Set("LandscapeEditor.SelectComponentTool", new IMAGE_BRUSH("Icons/icon_Landscape_Tool_Selection_40x", Icon40x40));
-		Set("LandscapeEditor.AddComponentTool", new IMAGE_BRUSH("Icons/icon_Landscape_Tool_AddComponent_40x", Icon40x40));
+		Set("LandscapeEditor.SelectComponentTool", new IMAGE_BRUSH("Icons/icon_Landscape_Tool_Selection_40x",       Icon40x40));
+		Set("LandscapeEditor.AddComponentTool",    new IMAGE_BRUSH("Icons/icon_Landscape_Tool_AddComponent_40x",    Icon40x40));
 		Set("LandscapeEditor.DeleteComponentTool", new IMAGE_BRUSH("Icons/icon_Landscape_Tool_DeleteComponent_40x", Icon40x40));
-		Set("LandscapeEditor.MoveToLevelTool", new IMAGE_BRUSH("Icons/icon_Landscape_Tool_MoveToLevel_40x", Icon40x40));
-		Set("LandscapeEditor.SelectComponentTool.Small", new IMAGE_BRUSH("Icons/icon_Landscape_Tool_Selection_20x", Icon20x20));
-		Set("LandscapeEditor.AddComponentTool.Small", new IMAGE_BRUSH("Icons/icon_Landscape_Tool_AddComponent_20x", Icon20x20));
+		Set("LandscapeEditor.MoveToLevelTool",     new IMAGE_BRUSH("Icons/icon_Landscape_Tool_MoveToLevel_40x",     Icon40x40));
+		Set("LandscapeEditor.SelectComponentTool.Small", new IMAGE_BRUSH("Icons/icon_Landscape_Tool_Selection_20x",       Icon20x20));
+		Set("LandscapeEditor.AddComponentTool.Small",    new IMAGE_BRUSH("Icons/icon_Landscape_Tool_AddComponent_20x",    Icon20x20));
 		Set("LandscapeEditor.DeleteComponentTool.Small", new IMAGE_BRUSH("Icons/icon_Landscape_Tool_DeleteComponent_20x", Icon20x20));
-		Set("LandscapeEditor.MoveToLevelTool.Small", new IMAGE_BRUSH("Icons/icon_Landscape_Tool_MoveToLevel_20x", Icon20x20));
+		Set("LandscapeEditor.MoveToLevelTool.Small",     new IMAGE_BRUSH("Icons/icon_Landscape_Tool_MoveToLevel_20x",     Icon20x20));
 
-		Set("LandscapeEditor.RegionSelectTool", new IMAGE_BRUSH("Icons/icon_Landscape_Tool_Mask_40x", Icon40x40));
+		Set("LandscapeEditor.RegionSelectTool",    new IMAGE_BRUSH("Icons/icon_Landscape_Tool_Mask_40x", Icon40x40));
 		Set("LandscapeEditor.RegionCopyPasteTool", new IMAGE_BRUSH("Icons/icon_Landscape_Tool_Copy_40x", Icon40x40));
-		Set("LandscapeEditor.RegionSelectTool.Small", new IMAGE_BRUSH("Icons/icon_Landscape_Tool_Mask_20x", Icon20x20));
+		Set("LandscapeEditor.RegionSelectTool.Small",    new IMAGE_BRUSH("Icons/icon_Landscape_Tool_Mask_20x", Icon20x20));
 		Set("LandscapeEditor.RegionCopyPasteTool.Small", new IMAGE_BRUSH("Icons/icon_Landscape_Tool_Copy_20x", Icon20x20));
 
-		Set("LandscapeEditor.SplineTool", new IMAGE_BRUSH("Icons/icon_Landscape_Tool_Splines_40x", Icon40x40));
+		Set("LandscapeEditor.MirrorTool",       new IMAGE_BRUSH("Icons/icon_Landscape_Tool_Mirror_40x", Icon40x40));
+		Set("LandscapeEditor.MirrorTool.Small", new IMAGE_BRUSH("Icons/icon_Landscape_Tool_Mirror_20x", Icon20x20));
+
+		Set("LandscapeEditor.SplineTool",       new IMAGE_BRUSH("Icons/icon_Landscape_Tool_Splines_40x", Icon40x40));
 		Set("LandscapeEditor.SplineTool.Small", new IMAGE_BRUSH("Icons/icon_Landscape_Tool_Splines_20x", Icon20x20));
 
 		// Brush Sets
-		Set("LandscapeEditor.CircleBrush", new IMAGE_BRUSH("Icons/icon_Landscape_Brush_Circle_smooth_40x", Icon40x40));
-		Set("LandscapeEditor.AlphaBrush", new IMAGE_BRUSH("Icons/icon_Landscape_Brush_Alpha_40x", Icon40x40));
-		Set("LandscapeEditor.AlphaBrush_Pattern", new IMAGE_BRUSH("Icons/icon_Landscape_Brush_Pattern_40x", Icon40x40));
-		Set("LandscapeEditor.ComponentBrush", new IMAGE_BRUSH("Icons/icon_Landscape_Brush_Component_40x", Icon40x40));
-		Set("LandscapeEditor.GizmoBrush", new IMAGE_BRUSH("Icons/icon_Landscape_Brush_Gizmo_40x", Icon40x40));
-		Set("LandscapeEditor.CircleBrush.Small", new IMAGE_BRUSH("Icons/icon_Landscape_Brush_Circle_smooth_20x", Icon20x20));
-		Set("LandscapeEditor.AlphaBrush.Small", new IMAGE_BRUSH("Icons/icon_Landscape_Brush_Alpha_20x", Icon20x20));
-		Set("LandscapeEditor.AlphaBrush_Pattern.Small", new IMAGE_BRUSH("Icons/icon_Landscape_Brush_Pattern_20x", Icon20x20));
-		Set("LandscapeEditor.ComponentBrush.Small", new IMAGE_BRUSH("Icons/icon_Landscape_Brush_Component_20x", Icon20x20));
-		Set("LandscapeEditor.GizmoBrush.Small", new IMAGE_BRUSH("Icons/icon_Landscape_Brush_Gizmo_20x", Icon20x20));
+		Set("LandscapeEditor.CircleBrush",        new IMAGE_BRUSH("Icons/icon_Landscape_Brush_Circle_smooth_40x", Icon40x40));
+		Set("LandscapeEditor.AlphaBrush",         new IMAGE_BRUSH("Icons/icon_Landscape_Brush_Alpha_40x",     Icon40x40));
+		Set("LandscapeEditor.AlphaBrush_Pattern", new IMAGE_BRUSH("Icons/icon_Landscape_Brush_Pattern_40x",   Icon40x40));
+		Set("LandscapeEditor.ComponentBrush",     new IMAGE_BRUSH("Icons/icon_Landscape_Brush_Component_40x", Icon40x40));
+		Set("LandscapeEditor.GizmoBrush",         new IMAGE_BRUSH("Icons/icon_Landscape_Brush_Gizmo_40x",     Icon40x40));
+		Set("LandscapeEditor.CircleBrush.Small",        new IMAGE_BRUSH("Icons/icon_Landscape_Brush_Circle_smooth_20x", Icon20x20));
+		Set("LandscapeEditor.AlphaBrush.Small",         new IMAGE_BRUSH("Icons/icon_Landscape_Brush_Alpha_20x",     Icon20x20));
+		Set("LandscapeEditor.AlphaBrush_Pattern.Small", new IMAGE_BRUSH("Icons/icon_Landscape_Brush_Pattern_20x",   Icon20x20));
+		Set("LandscapeEditor.ComponentBrush.Small",     new IMAGE_BRUSH("Icons/icon_Landscape_Brush_Component_20x", Icon20x20));
+		Set("LandscapeEditor.GizmoBrush.Small",         new IMAGE_BRUSH("Icons/icon_Landscape_Brush_Gizmo_20x",     Icon20x20));
 
 		// Brushes
-		Set("LandscapeEditor.CircleBrush_Smooth", new IMAGE_BRUSH("Icons/icon_Landscape_Brush_Circle_smooth_40x", Icon40x40));
-		Set("LandscapeEditor.CircleBrush_Linear", new IMAGE_BRUSH("Icons/icon_Landscape_Brush_Circle_linear_40x", Icon40x40));
+		Set("LandscapeEditor.CircleBrush_Smooth",    new IMAGE_BRUSH("Icons/icon_Landscape_Brush_Circle_smooth_40x",    Icon40x40));
+		Set("LandscapeEditor.CircleBrush_Linear",    new IMAGE_BRUSH("Icons/icon_Landscape_Brush_Circle_linear_40x",    Icon40x40));
 		Set("LandscapeEditor.CircleBrush_Spherical", new IMAGE_BRUSH("Icons/icon_Landscape_Brush_Circle_spherical_40x", Icon40x40));
-		Set("LandscapeEditor.CircleBrush_Tip", new IMAGE_BRUSH("Icons/icon_Landscape_Brush_Circle_tip_40x", Icon40x40));
-		Set("LandscapeEditor.CircleBrush_Smooth.Small", new IMAGE_BRUSH("Icons/icon_Landscape_Brush_Circle_smooth_20x", Icon20x20));
-		Set("LandscapeEditor.CircleBrush_Linear.Small", new IMAGE_BRUSH("Icons/icon_Landscape_Brush_Circle_linear_20x", Icon20x20));
+		Set("LandscapeEditor.CircleBrush_Tip",       new IMAGE_BRUSH("Icons/icon_Landscape_Brush_Circle_tip_40x",       Icon40x40));
+		Set("LandscapeEditor.CircleBrush_Smooth.Small",    new IMAGE_BRUSH("Icons/icon_Landscape_Brush_Circle_smooth_20x",    Icon20x20));
+		Set("LandscapeEditor.CircleBrush_Linear.Small",    new IMAGE_BRUSH("Icons/icon_Landscape_Brush_Circle_linear_20x",    Icon20x20));
 		Set("LandscapeEditor.CircleBrush_Spherical.Small", new IMAGE_BRUSH("Icons/icon_Landscape_Brush_Circle_spherical_20x", Icon20x20));
-		Set("LandscapeEditor.CircleBrush_Tip.Small", new IMAGE_BRUSH("Icons/icon_Landscape_Brush_Circle_tip_20x", Icon20x20));
+		Set("LandscapeEditor.CircleBrush_Tip.Small",       new IMAGE_BRUSH("Icons/icon_Landscape_Brush_Circle_tip_20x",       Icon20x20));
 
 		Set("LandscapeEditor.Brushes.Alpha.UseRChannel", new IMAGE_BRUSH("Icons/icon_Landscape_Brush_Alpha_UseRChannel_20x", Icon20x20));
 		Set("LandscapeEditor.Brushes.Alpha.UseGChannel", new IMAGE_BRUSH("Icons/icon_Landscape_Brush_Alpha_UseGChannel_20x", Icon20x20));
@@ -5696,18 +5699,18 @@ void FSlateEditorStyle::FStyle::SetupLandscapeEditorStyle()
 		Set("LandscapeEditor.Brushes.Alpha.UseAChannel", new IMAGE_BRUSH("Icons/icon_Landscape_Brush_Alpha_UseAChannel_20x", Icon20x20));
 
 		// Target List
-		Set("LandscapeEditor.TargetList.RowBackground", new FSlateNoResource());
+		Set("LandscapeEditor.TargetList.RowBackground",        new FSlateNoResource());
 		Set("LandscapeEditor.TargetList.RowBackgroundHovered", new BOX_BRUSH("Common/RoundedSelection_16x", 4.0f / 16.0f, FLinearColor(1.0f, 1.0f, 1.0f, 0.1f)));
-		Set("LandscapeEditor.TargetList.RowSelected", new BOX_BRUSH("Common/RoundedSelection_16x", 4.0f / 16.0f, SelectionColor_Pressed));
-		Set("LandscapeEditor.TargetList.RowSelectedHovered", new BOX_BRUSH("Common/RoundedSelection_16x", 4.0f / 16.0f, SelectionColor));
+		Set("LandscapeEditor.TargetList.RowSelected",          new BOX_BRUSH("Common/RoundedSelection_16x", 4.0f / 16.0f, SelectionColor_Pressed));
+		Set("LandscapeEditor.TargetList.RowSelectedHovered",   new BOX_BRUSH("Common/RoundedSelection_16x", 4.0f / 16.0f, SelectionColor));
 
-		Set("LandscapeEditor.Target_Heightmap", new IMAGE_BRUSH("Icons/icon_Landscape_Target_Heightmap_48x", Icon48x48));
+		Set("LandscapeEditor.Target_Heightmap",  new IMAGE_BRUSH("Icons/icon_Landscape_Target_Heightmap_48x",  Icon48x48));
 		Set("LandscapeEditor.Target_Visibility", new IMAGE_BRUSH("Icons/icon_Landscape_Target_Visibility_48x", Icon48x48));
-		Set("LandscapeEditor.Target_Invalid", new IMAGE_BRUSH("Icons/icon_Landscape_Target_Invalid_48x", Icon48x48));
+		Set("LandscapeEditor.Target_Invalid",    new IMAGE_BRUSH("Icons/icon_Landscape_Target_Invalid_48x",    Icon48x48));
 
-		Set("LandscapeEditor.Target_Create", new IMAGE_BRUSH("Icons/icon_Landscape_Target_Create_12x", Icon12x12));
-		Set("LandscapeEditor.Target_MakePublic", new IMAGE_BRUSH("Icons/assign_right_12x", Icon12x12));
-		Set("LandscapeEditor.Target_Delete", new IMAGE_BRUSH("Icons/Cross_12x", Icon12x12));
+		Set("LandscapeEditor.Target_Create",     new IMAGE_BRUSH("Icons/icon_Landscape_Target_Create_12x", Icon12x12));
+		Set("LandscapeEditor.Target_MakePublic", new IMAGE_BRUSH("Icons/assign_right_12x",                 Icon12x12));
+		Set("LandscapeEditor.Target_Delete",     new IMAGE_BRUSH("Icons/Cross_12x",                        Icon12x12));
 	}
 #endif
 }
diff --git a/Engine/Source/Editor/LandscapeEditor/Private/LandscapeEdMode.cpp b/Engine/Source/Editor/LandscapeEditor/Private/LandscapeEdMode.cpp
--- a/Engine/Source/Editor/LandscapeEditor/Private/LandscapeEdMode.cpp
+++ b/Engine/Source/Editor/LandscapeEditor/Private/LandscapeEdMode.cpp
@@ -178,7 +178,6 @@ FEdModeLandscape::FEdModeLandscape()
 	InitializeTool_Paint();
 	InitializeTool_Smooth();
 	InitializeTool_Flatten();
-	InitializeTool_Ramp();
 	InitializeTool_Erosion();
 	InitializeTool_HydraErosion();
 	InitializeTool_Noise();
@@ -193,6 +192,8 @@ FEdModeLandscape::FEdModeLandscape()
 	InitializeTool_CopyPaste();
 	InitializeTool_Visibility();
 	InitializeTool_Splines();
+	InitializeTool_Ramp();
+	InitializeTool_Mirror();
 
 	CurrentTool = nullptr;
 	CurrentToolIndex = INDEX_NONE;
@@ -269,6 +270,7 @@ void FEdModeLandscape::InitializeToolModes()
 	ToolMode_Sculpt->ValidTools.Add(TEXT("Visibility"));
 	ToolMode_Sculpt->ValidTools.Add(TEXT("Mask"));
 	ToolMode_Sculpt->ValidTools.Add(TEXT("CopyPaste"));
+	ToolMode_Sculpt->ValidTools.Add(TEXT("Mirror"));
 
 	FLandscapeToolMode* ToolMode_Paint = new(LandscapeToolModes)FLandscapeToolMode(TEXT("ToolMode_Paint"), ELandscapeToolTargetTypeMask::Weightmap);
 	ToolMode_Paint->ValidTools.Add(TEXT("Paint"));
diff --git a/Engine/Source/Editor/LandscapeEditor/Private/LandscapeEdMode.h b/Engine/Source/Editor/LandscapeEditor/Private/LandscapeEdMode.h
--- a/Engine/Source/Editor/LandscapeEditor/Private/LandscapeEdMode.h
+++ b/Engine/Source/Editor/LandscapeEditor/Private/LandscapeEdMode.h
@@ -237,11 +237,12 @@ public:
 	float DraggingEdge_Remainder;
 
 	TWeakObjectPtr<ALandscapeGizmoActiveActor> CurrentGizmoActor;
-	//
+	// UI callbacks for copy/paste tool
 	FLandscapeToolCopyPaste<FHeightmapToolTarget>* CopyPasteTool;
 	void CopyDataToGizmo();
 	void PasteDataFromGizmo();
 
+	// UI callbacks for splines tool
 	FLandscapeToolSplines* SplinesTool;
 	void ShowSplineProperties();
 	virtual void SelectAllConnectedSplineControlPoints();
@@ -250,10 +251,15 @@ public:
 	void SetbUseAutoRotateOnJoin(bool InbAutoRotateOnJoin);
 	bool GetbUseAutoRotateOnJoin();
 
+	// UI callbacks for ramp tool
 	void ApplyRampTool();
 	bool CanApplyRampTool();
 	void ResetRampTool();
 
+	// UI callbacks for mirror tool
+	void ApplyMirrorTool();
+	void CenterMirrorTool();
+
 	/** Constructor */
 	FEdModeLandscape();
 
@@ -277,6 +283,7 @@ public:
 	void InitializeTool_Visibility();
 	void InitializeTool_Splines();
 	void InitializeTool_Ramp();
+	void InitializeTool_Mirror();
 	void InitializeToolModes();
 
 	/** Destructor */
diff --git a/Engine/Source/Editor/LandscapeEditor/Private/LandscapeEdModeMirrorTool.cpp b/Engine/Source/Editor/LandscapeEditor/Private/LandscapeEdModeMirrorTool.cpp
new file mode 100644
--- /dev/null
+++ b/Engine/Source/Editor/LandscapeEditor/Private/LandscapeEdModeMirrorTool.cpp
@@ -0,0 +1,509 @@
+// Copyright 1998-2015 Epic Games, Inc. All Rights Reserved.
+
+#include "LandscapeEditorPrivatePCH.h"
+#include "LandscapeEdMode.h"
+#include "ScopedTransaction.h"
+#include "LandscapeEdit.h"
+#include "LandscapeRender.h"
+#include "LandscapeHeightfieldCollisionComponent.h"
+#include "Algo/Reverse.h"
+//#include "LandscapeDataAccess.h"
+//#include "AI/Navigation/NavigationSystem.h"
+
+#define LOCTEXT_NAMESPACE "Landscape"
+
+class FLandscapeToolMirror : public FLandscapeTool
+{
+protected:
+	FEdModeLandscape* EdMode;
+	UMaterialInstanceDynamic* MirrorPlaneMaterial;
+
+	ECoordSystem SavedCoordSystem;
+
+public:
+	FLandscapeToolMirror(FEdModeLandscape* InEdMode)
+		: EdMode(InEdMode)
+	{
+		UMaterialInterface* BaseMirrorPlaneMaterial = LoadObject<UMaterialInterface>(nullptr, TEXT("/Engine/EditorLandscapeResources/MirrorPlaneMaterial.MirrorPlaneMaterial"));
+		MirrorPlaneMaterial = UMaterialInstanceDynamic::Create(BaseMirrorPlaneMaterial, GetTransientPackage());
+		MirrorPlaneMaterial->SetScalarParameterValue(FName("LineThickness"), 2.0f);
+	}
+
+	virtual void AddReferencedObjects(FReferenceCollector& Collector) override
+	{
+		Collector.AddReferencedObject(MirrorPlaneMaterial);
+	}
+
+	virtual const TCHAR* GetToolName() override { return TEXT("Mirror"); }
+	virtual FText GetDisplayName() override { return FText(); /*NSLOCTEXT("UnrealEd", "LandscapeTool_Mirror", "Mirror Landscape");*/ };
+
+	virtual void SetEditRenderType() override { GLandscapeEditRenderMode = ELandscapeEditRenderMode::None | (GLandscapeEditRenderMode & ELandscapeEditRenderMode::BitMaskForMask); }
+	virtual bool SupportsMask() override { return false; }
+
+	virtual ELandscapeToolTargetTypeMask::Type GetSupportedTargetTypes() override
+	{
+		return ELandscapeToolTargetTypeMask::Heightmap;
+	}
+
+	virtual void EnterTool() override
+	{
+		if (EdMode->UISettings->MirrorPoint == FVector2D::ZeroVector)
+		{
+			CenterMirrorPoint();
+		}
+		GLevelEditorModeTools().SetWidgetMode(FWidget::WM_Translate);
+		SavedCoordSystem = GLevelEditorModeTools().GetCoordSystem();
+		GLevelEditorModeTools().SetCoordSystem(COORD_Local);
+	}
+
+	virtual void ExitTool() override
+	{
+		GLevelEditorModeTools().SetCoordSystem(SavedCoordSystem);
+	}
+
+	virtual bool BeginTool(FEditorViewportClient* ViewportClient, const FLandscapeToolTarget& Target, const FVector& InHitLocation) override
+	{
+		return true;
+	}
+
+	virtual void EndTool(FEditorViewportClient* ViewportClient) override
+	{
+	}
+
+	virtual bool MouseMove(FEditorViewportClient* ViewportClient, FViewport* Viewport, int32 x, int32 y) override
+	{
+		return false;
+	}
+
+	virtual bool InputKey(FEditorViewportClient* InViewportClient, FViewport* InViewport, FKey InKey, EInputEvent InEvent) override
+	{
+		if (InKey == EKeys::Enter && InEvent == IE_Pressed)
+		{
+			ApplyMirror();
+		}
+
+		return false;
+	}
+
+	virtual bool InputDelta(FEditorViewportClient* InViewportClient, FViewport* InViewport, FVector& InDrag, FRotator& InRot, FVector& InScale) override
+	{
+		if (InViewportClient->GetCurrentWidgetAxis() != EAxisList::None)
+		{
+			const ALandscapeProxy* LandscapeProxy = EdMode->CurrentToolTarget.LandscapeInfo->GetLandscapeProxy();
+			const FTransform LandscapeToWorld = LandscapeProxy->LandscapeActorToWorld();
+
+			EdMode->UISettings->MirrorPoint += FVector2D(LandscapeToWorld.InverseTransformVector(InDrag));
+
+			return true;
+		}
+
+		return false;
+	}
+
+	virtual void Render(const FSceneView* View, FViewport* Viewport, FPrimitiveDrawInterface* PDI) override
+	{
+		const ULandscapeInfo* const LandscapeInfo = EdMode->CurrentToolTarget.LandscapeInfo.Get();
+		const ALandscapeProxy* const LandscapeProxy = LandscapeInfo->GetLandscapeProxy();
+		const FTransform LandscapeToWorld = LandscapeProxy->LandscapeActorToWorld();
+
+		int32 MinX, MinY, MaxX, MaxY;
+		if (LandscapeInfo->GetLandscapeExtent(MinX, MinY, MaxX, MaxY))
+		{
+			FVector MirrorPoint3D = FVector((MaxX + MinX) / 2.0f, (MaxY + MinY) / 2.0f, 0);
+			FVector MirrorPlaneScale = FVector(0, 1, 100);
+
+			if (EdMode->UISettings->MirrorOp == ELandscapeMirrorOperation::MinusXToPlusX ||
+				EdMode->UISettings->MirrorOp == ELandscapeMirrorOperation::PlusXToMinusX)
+			{
+				MirrorPoint3D.X = EdMode->UISettings->MirrorPoint.X;
+				MirrorPlaneScale.Y = (MaxY - MinY) / 2.0f;
+			}
+			else
+			{
+				MirrorPoint3D.Y = EdMode->UISettings->MirrorPoint.Y;
+				MirrorPlaneScale.Y = (MaxX - MinX) / 2.0f;
+			}
+			MirrorPoint3D = LandscapeToWorld.TransformPosition(MirrorPoint3D);
+
+			FMatrix Matrix;
+			if (EdMode->UISettings->MirrorOp == ELandscapeMirrorOperation::MinusYToPlusY ||
+				EdMode->UISettings->MirrorOp == ELandscapeMirrorOperation::PlusYToMinusY)
+			{
+				Matrix = FScaleRotationTranslationMatrix(MirrorPlaneScale, FRotator(0, 90, 0), FVector::ZeroVector);
+			}
+			else
+			{
+				Matrix = FScaleMatrix(MirrorPlaneScale);
+			}
+
+			Matrix *= LandscapeToWorld.ToMatrixWithScale();
+			Matrix.SetOrigin(MirrorPoint3D);
+
+			// Convert plane from horizontal to vertical
+			Matrix = FMatrix(FVector(0, 1, 0), FVector(0, 0, 1), FVector(1, 0, 0), FVector(0, 0, 0)) * Matrix;
+
+			const FBox Box = FBox(FVector(-1, -1, 0), FVector(+1, +1, 0));
+			DrawWireBox(PDI, Matrix, Box, FLinearColor::Green, SDPG_World);
+
+			const float LandscapeScaleRatio = LandscapeToWorld.GetScale3D().Z / LandscapeToWorld.GetScale3D().X;
+			FVector2D UVScale = FVector2D(FMath::RoundToFloat(MirrorPlaneScale.Y / 10), FMath::RoundToFloat(MirrorPlaneScale.Z * LandscapeScaleRatio / 10 / 2) * 2);
+			MirrorPlaneMaterial->SetVectorParameterValue(FName("GridSize"), FVector(UVScale, 0));
+			DrawPlane10x10(PDI, Matrix, 1, FVector2D(0, 0), FVector2D(1, 1), MirrorPlaneMaterial->GetRenderProxy(false), SDPG_World);
+		}
+	}
+
+	virtual bool OverrideSelection() const override
+	{
+		return true;
+	}
+
+	virtual bool IsSelectionAllowed(AActor* InActor, bool bInSelection) const override
+	{
+		// Only filter selection not deselection
+		if (bInSelection)
+		{
+			return false;
+		}
+
+		return true;
+	}
+
+	virtual bool UsesTransformWidget() const override
+	{
+		return true;
+	}
+
+	virtual EAxisList::Type GetWidgetAxisToDraw(FWidget::EWidgetMode CheckMode) const override
+	{
+		if (CheckMode == FWidget::WM_Translate)
+		{
+			switch (EdMode->UISettings->MirrorOp)
+			{
+			case ELandscapeMirrorOperation::MinusXToPlusX:
+			case ELandscapeMirrorOperation::PlusXToMinusX:
+				return EAxisList::X;
+			case ELandscapeMirrorOperation::MinusYToPlusY:
+			case ELandscapeMirrorOperation::PlusYToMinusY:
+				return EAxisList::Y;
+			default:
+				check(0);
+				return EAxisList::None;
+			}
+		}
+		else
+		{
+			return EAxisList::None;
+		}
+
+		return EAxisList::None;
+	}
+
+	virtual FVector GetWidgetLocation() const override
+	{
+		const ULandscapeInfo* const LandscapeInfo = EdMode->CurrentToolTarget.LandscapeInfo.Get();
+		const ALandscapeProxy* const LandscapeProxy = LandscapeInfo->GetLandscapeProxy();
+		const FTransform LandscapeToWorld = LandscapeProxy->LandscapeActorToWorld();
+
+		int32 MinX, MinY, MaxX, MaxY;
+		if (!LandscapeInfo->GetLandscapeExtent(MinX, MinY, MaxX, MaxY))
+		{
+			MinX = MinY = 0;
+			MaxX = MaxY = 0;
+		}
+
+		FVector MirrorPoint3D = FVector((MaxX + MinX) / 2.0f, (MaxY + MinY) / 2.0f, 0);
+		if (EdMode->UISettings->MirrorOp == ELandscapeMirrorOperation::MinusXToPlusX ||
+			EdMode->UISettings->MirrorOp == ELandscapeMirrorOperation::PlusXToMinusX)
+		{
+			MirrorPoint3D.X = EdMode->UISettings->MirrorPoint.X;
+		}
+		else
+		{
+			MirrorPoint3D.Y = EdMode->UISettings->MirrorPoint.Y;
+		}
+		MirrorPoint3D = LandscapeToWorld.TransformPosition(MirrorPoint3D);
+		return MirrorPoint3D;
+	}
+
+	virtual FMatrix GetWidgetRotation() const override
+	{
+		const ALandscapeProxy* const LandscapeProxy = EdMode->CurrentToolTarget.LandscapeInfo->GetLandscapeProxy();
+		const FTransform LandscapeToWorld = LandscapeProxy->LandscapeActorToWorld();
+
+		FMatrix Result = FQuatRotationTranslationMatrix(LandscapeToWorld.GetRotation(), FVector::ZeroVector);
+		if (EdMode->UISettings->MirrorOp == ELandscapeMirrorOperation::PlusXToMinusX ||
+			EdMode->UISettings->MirrorOp == ELandscapeMirrorOperation::PlusYToMinusY)
+		{
+			Result = FRotationMatrix(FRotator(0, 180, 0)) * Result;
+		}
+
+		return Result;
+	}
+
+protected:
+	template<typename T>
+	void ApplyMirrorInternal(T* MirrorData, int32 SizeX, int32 SizeY, int32 MirrorSize)
+	{
+		switch (EdMode->UISettings->MirrorOp)
+		{
+		case ELandscapeMirrorOperation::MinusXToPlusX:
+		{
+			for (int32 Y = 0; Y < SizeY; ++Y) // copy extra column to calc normals for mirror column
+			{
+				Algo::Reverse(&MirrorData[Y * (2 + MirrorSize) + 1], MirrorSize + 1);
+				MirrorData[Y * (2 + MirrorSize) + 0] = MirrorData[Y * (2 + MirrorSize) + 2];
+			}
+		}
+		break;
+		case ELandscapeMirrorOperation::PlusXToMinusX:
+		{
+			for (int32 Y = 0; Y < SizeY; ++Y) // copy extra column to calc normals for mirror column
+			{
+				Algo::Reverse(&MirrorData[Y * (2 + MirrorSize) + 0], MirrorSize + 1);
+				MirrorData[Y * (2 + MirrorSize) + MirrorSize + 1] = MirrorData[Y * (2 + MirrorSize) + MirrorSize - 1];
+			}
+		}
+		break;
+		case ELandscapeMirrorOperation::MinusYToPlusY:
+		{
+			FMemory::Memcpy(&MirrorData[(MirrorSize + 1) * SizeX], &MirrorData[(MirrorSize - 1) * SizeX], SizeX * sizeof(T)); // copy extra row to calc normals for mirror row
+		}
+			break;
+		case ELandscapeMirrorOperation::PlusYToMinusY:
+		{
+			FMemory::Memcpy(&MirrorData[0 * SizeX], &MirrorData[2 * SizeX], SizeX * sizeof(T)); // copy extra row to calc normals for mirror row
+		}
+			break;
+		default:
+			check(0);
+			return;
+		}
+	}
+
+public:
+	virtual void ApplyMirror()
+	{
+		FScopedTransaction Transaction(LOCTEXT("Mirror_Apply", "Landscape Editing: Mirror Landscape"));
+
+		const ULandscapeInfo* const LandscapeInfo = EdMode->CurrentToolTarget.LandscapeInfo.Get();
+		const ALandscapeProxy* const LandscapeProxy = LandscapeInfo->GetLandscapeProxy();
+		const FTransform LandscapeToWorld = LandscapeProxy->LandscapeActorToWorld();
+		const FVector2D MirrorPoint = EdMode->UISettings->MirrorPoint;
+
+		FLandscapeEditDataInterface LandscapeEdit(EdMode->CurrentToolTarget.LandscapeInfo.Get());
+
+		int32 MinX, MinY, MaxX, MaxY;
+		if (!LandscapeInfo->GetLandscapeExtent(MinX, MinY, MaxX, MaxY))
+		{
+			return;
+		}
+
+		const int32 SizeX = (1 + MaxX - MinX);
+		const int32 SizeY = (1 + MaxY - MinY);
+
+		int32 SourceMinX, SourceMinY;
+		int32 SourceMaxX, SourceMaxY;
+		int32 DestMinX, DestMinY;
+		int32 DestMaxX, DestMaxY;
+		int32 MirrorSize;
+		int32 DataSize;
+		int32 ReadOffset;
+		int32 ReadStride;
+		int32 WriteOffset;
+		int32 WriteStride;
+
+		switch (EdMode->UISettings->MirrorOp)
+		{
+		case ELandscapeMirrorOperation::MinusXToPlusX:
+		{
+			const int32 MirrorX = FMath::RoundToInt(MirrorPoint.X);
+			if (MirrorX <= MinX || MirrorX >= MaxX)
+			{
+				return;
+			}
+			MirrorSize = FMath::Max(MaxX - MirrorX, MirrorX - MinX); // not including the mirror column itself
+			DataSize = SizeY * (2 + MirrorSize); // +2 for mirror column and extra to calc normals for mirror column
+			SourceMinX = MirrorX - MirrorSize;
+			SourceMaxX = MirrorX;
+			SourceMinY = MinY;
+			SourceMaxY = MaxY;
+			DestMinX = MirrorX - 1; // extra column to calc normals for mirror column
+			DestMaxX = MirrorX + MirrorSize;
+			DestMinY = MinY;
+			DestMaxY = MaxY;
+			ReadOffset = 1;
+			ReadStride = (2 + MirrorSize);
+			WriteOffset = 0;
+			WriteStride = (2 + MirrorSize);
+		}
+		break;
+		case ELandscapeMirrorOperation::PlusXToMinusX:
+		{
+			const int32 MirrorX = FMath::RoundToInt(MirrorPoint.X);
+			if (MirrorX <= MinX || MirrorX >= MaxX)
+			{
+				return;
+			}
+			MirrorSize = FMath::Max(MaxX - MirrorX, MirrorX - MinX); // not including the mirror column itself
+			DataSize = SizeY * (2 + MirrorSize); // +2 for mirror column and extra to calc normals for mirror column
+			SourceMinX = MirrorX;
+			SourceMaxX = MirrorX + MirrorSize;
+			SourceMinY = MinY;
+			SourceMaxY = MaxY;
+			DestMinX = MirrorX - MirrorSize;
+			DestMaxX = MirrorX + 1; // extra column to calc normals for mirror column
+			DestMinY = MinY;
+			DestMaxY = MaxY;
+			ReadOffset = 0;
+			ReadStride = (2 + MirrorSize);
+			WriteOffset = 0;
+			WriteStride = (2 + MirrorSize);
+		}
+		break;
+		case ELandscapeMirrorOperation::MinusYToPlusY:
+		{
+			const int32 MirrorY = FMath::RoundToInt(MirrorPoint.Y);
+			if (MirrorY <= MinY || MirrorY >= MaxY)
+			{
+				return;
+			}
+			MirrorSize = FMath::Max(MaxY - MirrorY, MirrorY - MinY); // not including the mirror row itself
+			DataSize = (2 + MirrorSize) * SizeX; // +2 for mirror row and extra to calc normals for mirror row
+			SourceMinX = MinX;
+			SourceMaxX = MaxX;
+			SourceMinY = MirrorY - MirrorSize;
+			SourceMaxY = MirrorY;
+			DestMinX = MinX;
+			DestMaxX = MaxX;
+			DestMinY = MirrorY - 1; // extra column to calc normals for mirror column
+			DestMaxY = MirrorY + MirrorSize;
+			ReadOffset = 0 * SizeX;
+			ReadStride = SizeX;
+			WriteOffset = DataSize - SizeX; // last line first and negative stride to flip :)
+			WriteStride = -SizeX;
+		}
+			break;
+		case ELandscapeMirrorOperation::PlusYToMinusY:
+		{
+			const int32 MirrorY = FMath::RoundToInt(MirrorPoint.Y);
+			if (MirrorY <= MinY || MirrorY >= MaxY)
+			{
+				return;
+			}
+			MirrorSize = FMath::Max(MaxY - MirrorY, MirrorY - MinY); // not including the mirror row itself
+			DataSize = (2 + MirrorSize) * SizeX; // +2 for mirror row and extra to calc normals for mirror row
+			SourceMinX = MinX;
+			SourceMaxX = MaxX;
+			SourceMinY = MirrorY;
+			SourceMaxY = MirrorY + MirrorSize;
+			DestMinX = MinX;
+			DestMaxX = MaxX;
+			DestMinY = MirrorY - MirrorSize;
+			DestMaxY = MirrorY + 1; // extra column to calc normals for mirror column
+			ReadOffset = 1 * SizeX;
+			ReadStride = SizeX;
+			WriteOffset = DataSize - SizeX; // last line first and negative stride to flip :)
+			WriteStride = -SizeX;
+		}
+			break;
+		default:
+			check(0);
+			return;
+		}
+
+		TArray<uint16> MirrorHeightData;
+		MirrorHeightData.AddZeroed(DataSize);
+		int32 TempMinX = SourceMinX; // GetHeightData overwrites its input min/max x/y
+		int32 TempMaxX = SourceMaxX;
+		int32 TempMinY = SourceMinY;
+		int32 TempMaxY = SourceMaxY;
+		LandscapeEdit.GetHeightData(TempMinX, TempMinY, TempMaxX, TempMaxY, &MirrorHeightData[ReadOffset], ReadStride);
+		ApplyMirrorInternal(MirrorHeightData.GetData(), SizeX, SizeY, MirrorSize);
+		LandscapeEdit.SetHeightData(DestMinX, DestMinY, DestMaxX, DestMaxY, &MirrorHeightData[WriteOffset], WriteStride, true);
+
+		TArray<uint8> MirrorWeightData;
+		MirrorWeightData.AddZeroed(DataSize);
+		for (const auto& LayerSettings : LandscapeInfo->Layers)
+		{
+			ULandscapeLayerInfoObject* LayerInfo = LayerSettings.LayerInfoObj;
+			if (LayerInfo)
+			{
+				TempMinX = SourceMinX; // GetWeightData overwrites its input min/max x/y
+				TempMaxX = SourceMaxX;
+				TempMinY = SourceMinY;
+				TempMaxY = SourceMaxY;
+				LandscapeEdit.GetWeightData(LayerInfo, TempMinX, TempMinY, TempMaxX, TempMaxY, &MirrorWeightData[ReadOffset], ReadStride);
+				ApplyMirrorInternal(MirrorWeightData.GetData(), SizeX, SizeY, MirrorSize);
+				LandscapeEdit.SetAlphaData(LayerInfo, DestMinX, DestMinY, DestMaxX, DestMaxY, &MirrorWeightData[WriteOffset], WriteStride, ELandscapeLayerPaintingRestriction::None);
+			}
+		}
+
+		LandscapeEdit.Flush();
+
+		TSet<ULandscapeComponent*> Components;
+		if (LandscapeEdit.GetComponentsInRegion(DestMinX, DestMinY, DestMaxX, DestMaxY, &Components))
+		{
+			for (ULandscapeComponent* Component : Components)
+			{
+				// Recreate collision for modified components and update the navmesh
+				ULandscapeHeightfieldCollisionComponent* CollisionComponent = Component->CollisionComponent.Get();
+				if (CollisionComponent)
+				{
+					CollisionComponent->RecreateCollision();
+					UNavigationSystem* NavSys = UNavigationSystem::GetCurrent(Component);
+					if (NavSys)
+					{
+						NavSys->UpdateNavOctree(CollisionComponent);
+					}
+				}
+			}
+		}
+	}
+
+	void CenterMirrorPoint()
+	{
+		ULandscapeInfo* const LandscapeInfo = EdMode->CurrentToolTarget.LandscapeInfo.Get();
+		ALandscapeProxy* const LandscapeProxy = LandscapeInfo->GetLandscapeProxy();
+		int32 MinX, MinY, MaxX, MaxY;
+		if (LandscapeInfo->GetLandscapeExtent(MinX, MinY, MaxX, MaxY))
+		{
+			EdMode->UISettings->MirrorPoint = FVector2D((float)(MinX + MaxX) / 2.0f, (float)(MinY + MaxY) / 2.0f);
+		}
+		else
+		{
+			EdMode->UISettings->MirrorPoint = FVector2D(0, 0);
+		}
+	}
+};
+
+void FEdModeLandscape::ApplyMirrorTool()
+{
+	if (CurrentTool->GetToolName() == FName("Mirror"))
+	{
+		FLandscapeToolMirror* MirrorTool = (FLandscapeToolMirror*)CurrentTool;
+		MirrorTool->ApplyMirror();
+		GEditor->RedrawLevelEditingViewports();
+	}
+}
+
+void FEdModeLandscape::CenterMirrorTool()
+{
+	if (CurrentTool->GetToolName() == FName("Mirror"))
+	{
+		FLandscapeToolMirror* MirrorTool = (FLandscapeToolMirror*)CurrentTool;
+		MirrorTool->CenterMirrorPoint();
+		GEditor->RedrawLevelEditingViewports();
+	}
+}
+
+//
+// Toolset initialization
+//
+void FEdModeLandscape::InitializeTool_Mirror()
+{
+	auto Tool_Mirror = MakeUnique<FLandscapeToolMirror>(this);
+	Tool_Mirror->ValidBrushes.Add("BrushSet_Dummy");
+	LandscapeTools.Add(MoveTemp(Tool_Mirror));
+}
+
+#undef LOCTEXT_NAMESPACE
diff --git a/Engine/Source/Editor/LandscapeEditor/Private/LandscapeEditorCommands.cpp b/Engine/Source/Editor/LandscapeEditor/Private/LandscapeEditorCommands.cpp
--- a/Engine/Source/Editor/LandscapeEditor/Private/LandscapeEditorCommands.cpp
+++ b/Engine/Source/Editor/LandscapeEditor/Private/LandscapeEditorCommands.cpp
@@ -55,6 +55,8 @@ void FLandscapeEditorCommands::RegisterCommands()
 	NameToCommandMap.Add("Tool_Mask", RegionSelectTool);
 	UI_COMMAND(RegionCopyPasteTool, "Tool - Copy/Paste", "", EUserInterfaceActionType::RadioButton, FInputChord());
 	NameToCommandMap.Add("Tool_CopyPaste", RegionCopyPasteTool);
+	UI_COMMAND(MirrorTool, "Tool - Mirror", "", EUserInterfaceActionType::RadioButton, FInputChord());
+	NameToCommandMap.Add("Tool_Mirror", MirrorTool);
 
 	UI_COMMAND(SplineTool, "Tool - Edit Splines", "", EUserInterfaceActionType::RadioButton, FInputChord());
 	NameToCommandMap.Add("Tool_Splines", SplineTool);
diff --git a/Engine/Source/Editor/LandscapeEditor/Private/LandscapeEditorCommands.h b/Engine/Source/Editor/LandscapeEditor/Private/LandscapeEditorCommands.h
--- a/Engine/Source/Editor/LandscapeEditor/Private/LandscapeEditorCommands.h
+++ b/Engine/Source/Editor/LandscapeEditor/Private/LandscapeEditorCommands.h
@@ -28,52 +28,53 @@ public:
 
 public:
 	// Mode Switch
-	TSharedPtr< FUICommandInfo > ManageMode;
-	TSharedPtr< FUICommandInfo > SculptMode;
-	TSharedPtr< FUICommandInfo > PaintMode;
+	TSharedPtr<FUICommandInfo> ManageMode;
+	TSharedPtr<FUICommandInfo> SculptMode;
+	TSharedPtr<FUICommandInfo> PaintMode;
 
 	// Tools
-	TSharedPtr< FUICommandInfo > NewLandscape;
-	TSharedPtr< FUICommandInfo > ResizeLandscape;
-
-	TSharedPtr< FUICommandInfo > SculptTool;
-	TSharedPtr< FUICommandInfo > PaintTool;
-	TSharedPtr< FUICommandInfo > SmoothTool;
-	TSharedPtr< FUICommandInfo > FlattenTool;
-	TSharedPtr< FUICommandInfo > RampTool;
-	TSharedPtr< FUICommandInfo > ErosionTool;
-	TSharedPtr< FUICommandInfo > HydroErosionTool;
-	TSharedPtr< FUICommandInfo > NoiseTool;
-	TSharedPtr< FUICommandInfo > RetopologizeTool;
-	TSharedPtr< FUICommandInfo > VisibilityTool;
-
-	TSharedPtr< FUICommandInfo > SelectComponentTool;
-	TSharedPtr< FUICommandInfo > AddComponentTool;
-	TSharedPtr< FUICommandInfo > DeleteComponentTool;
-	TSharedPtr< FUICommandInfo > MoveToLevelTool;
-
-	TSharedPtr< FUICommandInfo > RegionSelectTool;
-	TSharedPtr< FUICommandInfo > RegionCopyPasteTool;
-
-	TSharedPtr< FUICommandInfo > SplineTool;
+	TSharedPtr<FUICommandInfo> NewLandscape;
+	TSharedPtr<FUICommandInfo> ResizeLandscape;
+
+	TSharedPtr<FUICommandInfo> SculptTool;
+	TSharedPtr<FUICommandInfo> PaintTool;
+	TSharedPtr<FUICommandInfo> SmoothTool;
+	TSharedPtr<FUICommandInfo> FlattenTool;
+	TSharedPtr<FUICommandInfo> RampTool;
+	TSharedPtr<FUICommandInfo> ErosionTool;
+	TSharedPtr<FUICommandInfo> HydroErosionTool;
+	TSharedPtr<FUICommandInfo> NoiseTool;
+	TSharedPtr<FUICommandInfo> RetopologizeTool;
+	TSharedPtr<FUICommandInfo> VisibilityTool;
+
+	TSharedPtr<FUICommandInfo> SelectComponentTool;
+	TSharedPtr<FUICommandInfo> AddComponentTool;
+	TSharedPtr<FUICommandInfo> DeleteComponentTool;
+	TSharedPtr<FUICommandInfo> MoveToLevelTool;
+
+	TSharedPtr<FUICommandInfo> RegionSelectTool;
+	TSharedPtr<FUICommandInfo> RegionCopyPasteTool;
+	TSharedPtr<FUICommandInfo> MirrorTool;
+
+	TSharedPtr<FUICommandInfo> SplineTool;
 
 	// Brushes
-	TSharedPtr< FUICommandInfo > CircleBrush;
-	TSharedPtr< FUICommandInfo > AlphaBrush;
-	TSharedPtr< FUICommandInfo > AlphaBrush_Pattern;
-	TSharedPtr< FUICommandInfo > ComponentBrush;
-	TSharedPtr< FUICommandInfo > GizmoBrush;
-
-	TSharedPtr< FUICommandInfo > CircleBrush_Smooth;
-	TSharedPtr< FUICommandInfo > CircleBrush_Linear;
-	TSharedPtr< FUICommandInfo > CircleBrush_Spherical;
-	TSharedPtr< FUICommandInfo > CircleBrush_Tip;
-
-	TSharedPtr< FUICommandInfo > ViewModeNormal;
-	TSharedPtr< FUICommandInfo > ViewModeLOD;
-	TSharedPtr< FUICommandInfo > ViewModeLayerDensity;
-	TSharedPtr< FUICommandInfo > ViewModeLayerDebug;
-	TSharedPtr< FUICommandInfo > ViewModeWireframeOnTop;
+	TSharedPtr<FUICommandInfo> CircleBrush;
+	TSharedPtr<FUICommandInfo> AlphaBrush;
+	TSharedPtr<FUICommandInfo> AlphaBrush_Pattern;
+	TSharedPtr<FUICommandInfo> ComponentBrush;
+	TSharedPtr<FUICommandInfo> GizmoBrush;
+
+	TSharedPtr<FUICommandInfo> CircleBrush_Smooth;
+	TSharedPtr<FUICommandInfo> CircleBrush_Linear;
+	TSharedPtr<FUICommandInfo> CircleBrush_Spherical;
+	TSharedPtr<FUICommandInfo> CircleBrush_Tip;
+
+	TSharedPtr<FUICommandInfo> ViewModeNormal;
+	TSharedPtr<FUICommandInfo> ViewModeLOD;
+	TSharedPtr<FUICommandInfo> ViewModeLayerDensity;
+	TSharedPtr<FUICommandInfo> ViewModeLayerDebug;
+	TSharedPtr<FUICommandInfo> ViewModeWireframeOnTop;
 
 	// Map
 	TMap<FName, TSharedPtr<FUICommandInfo>> NameToCommandMap;
diff --git a/Engine/Source/Editor/LandscapeEditor/Private/LandscapeEditorDetailCustomization_MiscTools.cpp b/Engine/Source/Editor/LandscapeEditor/Private/LandscapeEditorDetailCustomization_MiscTools.cpp
--- a/Engine/Source/Editor/LandscapeEditor/Private/LandscapeEditorDetailCustomization_MiscTools.cpp
+++ b/Engine/Source/Editor/LandscapeEditor/Private/LandscapeEditorDetailCustomization_MiscTools.cpp
@@ -143,6 +143,37 @@ void FLandscapeEditorDetailCustomization_MiscTools::CustomizeDetails(IDetailLayo
 			]
 		];
 	}
+
+	if (IsToolActive("Mirror"))
+	{
+		ToolsCategory.AddProperty(GET_MEMBER_NAME_CHECKED(ULandscapeEditorObject, MirrorPoint));
+		ToolsCategory.AddProperty(GET_MEMBER_NAME_CHECKED(ULandscapeEditorObject, MirrorOp));
+		ToolsCategory.AddCustomRow(LOCTEXT("ApplyMirrorLabel", "Apply Mirror"))
+		[
+			SNew(SBox)
+			.Padding(FMargin(0, 0, 12, 0)) // Line up with the other properties due to having no reset to default button
+			[
+				SNew(SHorizontalBox)
+				+ SHorizontalBox::Slot()
+				.Padding(0, 0, 3, 0)
+				[
+					SNew(SButton)
+					.Text(LOCTEXT("Mirror.Reset", "Recenter"))
+					.HAlign(HAlign_Center)
+					.OnClicked_Lambda(&FLandscapeEditorDetailCustomization_MiscTools::OnResetMirrorPointButtonClicked)
+				]
+				+ SHorizontalBox::Slot()
+				.Padding(3, 0, 0, 0)
+				[
+					SNew(SButton)
+					//.IsEnabled_Static(&FLandscapeEditorDetailCustomization_MiscTools::GetApplyMirrorButtonIsEnabled)
+					.Text(LOCTEXT("Mirror.Apply", "Apply"))
+					.HAlign(HAlign_Center)
+					.OnClicked_Static(&FLandscapeEditorDetailCustomization_MiscTools::OnApplyMirrorButtonClicked)
+				]
+			]
+		];
+	}
 }
 END_SLATE_FUNCTION_BUILD_OPTIMIZATION
 
@@ -293,4 +324,26 @@ FReply FLandscapeEditorDetailCustomization_MiscTools::OnResetRampButtonClicked()
 	return FReply::Handled();
 }
 
+FReply FLandscapeEditorDetailCustomization_MiscTools::OnApplyMirrorButtonClicked()
+{
+	FEdModeLandscape* LandscapeEdMode = GetEditorMode();
+	if (LandscapeEdMode != NULL && IsToolActive(FName("Mirror")))
+	{
+		LandscapeEdMode->ApplyMirrorTool();
+	}
+
+	return FReply::Handled();
+}
+
+FReply FLandscapeEditorDetailCustomization_MiscTools::OnResetMirrorPointButtonClicked()
+{
+	FEdModeLandscape* LandscapeEdMode = GetEditorMode();
+	if (LandscapeEdMode != NULL && IsToolActive(FName("Mirror")))
+	{
+		LandscapeEdMode->CenterMirrorTool();
+	}
+
+	return FReply::Handled();
+}
+
 #undef LOCTEXT_NAMESPACE
diff --git a/Engine/Source/Editor/LandscapeEditor/Private/LandscapeEditorDetailCustomization_MiscTools.h b/Engine/Source/Editor/LandscapeEditor/Private/LandscapeEditorDetailCustomization_MiscTools.h
--- a/Engine/Source/Editor/LandscapeEditor/Private/LandscapeEditorDetailCustomization_MiscTools.h
+++ b/Engine/Source/Editor/LandscapeEditor/Private/LandscapeEditorDetailCustomization_MiscTools.h
@@ -37,4 +37,8 @@ protected:
 	static FReply OnApplyRampButtonClicked();
 	static bool GetApplyRampButtonIsEnabled();
 	static FReply OnResetRampButtonClicked();
+
+	// Mirror Tool
+	static FReply OnApplyMirrorButtonClicked();
+	static FReply OnResetMirrorPointButtonClicked();
 };
diff --git a/Engine/Source/Editor/LandscapeEditor/Private/LandscapeEditorDetails.cpp b/Engine/Source/Editor/LandscapeEditor/Private/LandscapeEditorDetails.cpp
--- a/Engine/Source/Editor/LandscapeEditor/Private/LandscapeEditorDetails.cpp
+++ b/Engine/Source/Editor/LandscapeEditor/Private/LandscapeEditorDetails.cpp
@@ -129,6 +129,7 @@ FText FLandscapeEditorDetails::GetLocalizedName(FString Name)
 
 		LOCTEXT("ToolSet_Mask", "Selection");
 		LOCTEXT("ToolSet_CopyPaste", "Copy/Paste");
+		LOCTEXT("ToolSet_Mirror", "Mirror");
 
 		LOCTEXT("ToolSet_Splines", "Edit Splines");
 
@@ -282,6 +283,7 @@ TSharedRef<SWidget> FLandscapeEditorDetails::GetToolSelector()
 			MenuBuilder.BeginSection(NAME_None, LOCTEXT("RegionToolsTitle", "Region Tools"));
 			MenuBuilder.AddToolButton(NameToCommandMap.FindChecked("Tool_Mask"), NAME_None, LOCTEXT("Tool.RegionSelect", "Selection"), LOCTEXT("Tool.RegionSelect.Tooltip", "Select a region of landscape to use as a mask for other tools"));
 			MenuBuilder.AddToolButton(NameToCommandMap.FindChecked("Tool_CopyPaste"), NAME_None, LOCTEXT("Tool.RegionCopyPaste", "Copy/Paste"), LOCTEXT("Tool.RegionCopyPaste.Tooltip", "Copy/Paste areas of the landscape, or import/export a copied area of landscape from disk"));
+			MenuBuilder.AddToolButton(NameToCommandMap.FindChecked("Tool_Mirror"), NAME_None, LOCTEXT("Tool.Mirror", "Mirror"), LOCTEXT("Tool.Mirror.Tooltip", "Copies one side of a landscape to the other, to easily create a mirrored landscape."));
 			MenuBuilder.EndSection();
 		}
 
diff --git a/Engine/Source/Editor/LandscapeEditor/Private/LandscapeEditorObject.cpp b/Engine/Source/Editor/LandscapeEditor/Private/LandscapeEditorObject.cpp
--- a/Engine/Source/Editor/LandscapeEditor/Private/LandscapeEditorObject.cpp
+++ b/Engine/Source/Editor/LandscapeEditor/Private/LandscapeEditorObject.cpp
@@ -57,6 +57,9 @@ ULandscapeEditorObject::ULandscapeEditorObject(const FObjectInitializer& ObjectI
 	, bSnapGizmo(false)
 	, bSmoothGizmoBrush(true)
 
+	, MirrorPoint(FVector::ZeroVector)
+	, MirrorOp(ELandscapeMirrorOperation::MinusXToPlusX)
+
 	, ResizeLandscape_QuadsPerSection(0)
 	, ResizeLandscape_SectionsPerComponent(0)
 	, ResizeLandscape_ComponentCount(0, 0)
@@ -245,6 +248,10 @@ void ULandscapeEditorObject::Load()
 	//PasteMode = (ELandscapeToolPasteMode::Type)InPasteMode;
 	SetPasteMode((ELandscapeToolPasteMode::Type)InPasteMode);
 
+	int32 InMirrorOp = (int32)ELandscapeMirrorOperation::MinusXToPlusX;
+	GConfig->GetInt(TEXT("LandscapeEdit"), TEXT("MirrorOp"), InMirrorOp, GEditorPerProjectIni);
+	MirrorOp = (ELandscapeMirrorOperation)InMirrorOp;
+
 	int32 InConvertMode = ResizeLandscape_ConvertMode;
 	GConfig->GetInt(TEXT("LandscapeEdit"), TEXT("ConvertMode"), InConvertMode, GEditorPerProjectIni);
 	ResizeLandscape_ConvertMode = (ELandscapeConvertMode::Type)InConvertMode;
@@ -334,6 +341,9 @@ void ULandscapeEditorObject::Save()
 
 	GConfig->SetBool(TEXT("LandscapeEdit"), TEXT("bSmoothGizmoBrush"), bSmoothGizmoBrush, GEditorPerProjectIni);
 	GConfig->SetInt(TEXT("LandscapeEdit"), TEXT("PasteMode"), (int32)PasteMode, GEditorPerProjectIni);
+
+	GConfig->SetInt(TEXT("LandscapeEdit"), TEXT("MirrorOp"), (int32)MirrorOp, GEditorPerProjectIni);
+
 	GConfig->SetInt(TEXT("LandscapeEdit"), TEXT("ConvertMode"), (int32)ResizeLandscape_ConvertMode, GEditorPerProjectIni);
 	//GConfig->SetBool(TEXT("LandscapeEdit"), TEXT("bUseSelectedRegion"), bUseSelectedRegion, GEditorPerProjectIni);
 	//GConfig->SetBool(TEXT("LandscapeEdit"), TEXT("bUseNegativeMask"), bUseNegativeMask, GEditorPerProjectIni);
diff --git a/Engine/Source/Editor/LandscapeEditor/Private/SLandscapeEditor.cpp b/Engine/Source/Editor/LandscapeEditor/Private/SLandscapeEditor.cpp
--- a/Engine/Source/Editor/LandscapeEditor/Private/SLandscapeEditor.cpp
+++ b/Engine/Source/Editor/LandscapeEditor/Private/SLandscapeEditor.cpp
@@ -103,6 +103,7 @@ void FLandscapeToolKit::Init(const TSharedPtr<IToolkitHost>& InitToolkitHost)
 
 	MAP_TOOL("Mask");
 	MAP_TOOL("CopyPaste");
+	MAP_TOOL("Mirror");
 
 	MAP_TOOL("Splines");
 #undef MAP_TOOL
diff --git a/Engine/Source/Editor/LandscapeEditor/Public/LandscapeEditorObject.h b/Engine/Source/Editor/LandscapeEditor/Public/LandscapeEditorObject.h
--- a/Engine/Source/Editor/LandscapeEditor/Public/LandscapeEditorObject.h
+++ b/Engine/Source/Editor/LandscapeEditor/Public/LandscapeEditorObject.h
@@ -125,6 +125,15 @@ namespace EColorChannel
 	};
 }
 
+UENUM()
+enum class ELandscapeMirrorOperation : uint8
+{
+	MinusXToPlusX UMETA(DisplayName="-X to +X"),
+	PlusXToMinusX UMETA(DisplayName="+X to -X"),
+	MinusYToPlusY UMETA(DisplayName="-Y to +Y"),
+	PlusYToMinusY UMETA(DisplayName="+Y to -Y"),
+};
+
 USTRUCT()
 struct FGizmoImportLayer
 {
@@ -388,6 +397,15 @@ class ULandscapeEditorObject : public UObject
 
 	TArray<FGizmoHistory> GizmoHistories;
 
+	// Mirror Tool
+
+	// Location of the mirror plane, defaults to the center of the landscape. Doesn't normally need to be changed!
+	UPROPERTY(Category="Tool Settings", EditAnywhere, Transient, meta=(DisplayName="Mirror Point", ShowForTools="Mirror"))
+	FVector2D MirrorPoint;
+
+	// Type of mirroring operation to perform e.g. "Minus X To Plus X" copies and flips the -X half of the landscape onto the +X half
+	UPROPERTY(Category="Tool Settings", EditAnywhere, NonTransactional, meta=(DisplayName="Operation", ShowForTools="Mirror"))
+	ELandscapeMirrorOperation MirrorOp;
 
 	// Resize Landscape Tool
 
diff --git a/Engine/Source/Runtime/Core/Public/Containers/Algo/Reverse.h b/Engine/Source/Runtime/Core/Public/Containers/Algo/Reverse.h
new file mode 100644
--- /dev/null
+++ b/Engine/Source/Runtime/Core/Public/Containers/Algo/Reverse.h
@@ -0,0 +1,52 @@
+// Copyright 1998-2015 Epic Games, Inc. All Rights Reserved.
+
+#pragma once
+
+namespace AlgoImpl
+{
+	template <typename T>
+	FORCEINLINE void Reverse(T* Array, int32 ArraySize)
+	{
+		for (int32 i = 0, i2 = ArraySize - 1; i < ArraySize / 2 /*rounding down*/; ++i, --i2)
+		{
+			Swap(Array[i], Array[i2]);
+		}
+	}
+}
+
+namespace Algo
+{
+	/**
+	 * Reverses a range
+	 *
+	 * @param  Array  The array to reverse.
+	 */
+	template <typename T, int32 ArraySize>
+	FORCEINLINE void Reverse(T (&Array)[ArraySize])
+	{
+		return AlgoImpl::Reverse((T*)Array, ArraySize);
+	}
+
+	/**
+	 * Reverses a range
+	 *
+	 * @param  Array      A pointer to the array to reverse
+	 * @param  ArraySize  The number of elements in the array.
+	 */
+	template <typename T>
+	FORCEINLINE void Reverse(T* Array, int32 ArraySize)
+	{
+		return AlgoImpl::Reverse(Array, ArraySize);
+	}
+
+	/**
+	 * Reverses a range
+	 *
+	 * @param  Container  The container to reverse
+	 */
+	template <typename ContainerType>
+	FORCEINLINE void Reverse(ContainerType& Container)
+	{
+		return AlgoImpl::Reverse(Container.GetData(), Container.Num());
+	}
+}
diff --git a/Engine/Source/Runtime/Landscape/Private/LandscapeEditInterface.cpp b/Engine/Source/Runtime/Landscape/Private/LandscapeEditInterface.cpp
--- a/Engine/Source/Runtime/Landscape/Private/LandscapeEditInterface.cpp
+++ b/Engine/Source/Runtime/Landscape/Private/LandscapeEditInterface.cpp
@@ -201,9 +201,12 @@ bool FLandscapeEditDataInterface::GetComponentsInRegion(int32 X1, int32 Y1, int3
 
 void FLandscapeEditDataInterface::SetHeightData(int32 X1, int32 Y1, int32 X2, int32 Y2, const uint16* Data, int32 Stride, bool CalcNormals, const uint16* NormalData /*= NULL*/, bool CreateComponents /*= false*/)
 {
-	if( Stride==0 )
+	const int32 NumVertsX = 1 + X2 - X1;
+	const int32 NumVertsY = 1 + Y2 - Y1;
+
+	if (Stride == 0)
 	{
-		Stride = (1+X2-X1);
+		Stride = NumVertsX;
 	}
 
 	check(ComponentSizeQuads > 0);
@@ -211,14 +214,12 @@ void FLandscapeEditDataInterface::SetHeightData(int32 X1, int32 Y1, int32 X2, in
 	int32 ComponentIndexX1, ComponentIndexY1, ComponentIndexX2, ComponentIndexY2;
 	ALandscape::CalcComponentIndicesOverlap(X1, Y1, X2, Y2, ComponentSizeQuads, ComponentIndexX1, ComponentIndexY1, ComponentIndexX2, ComponentIndexY2);
 
-	FVector* VertexNormals = NULL;
-	if( CalcNormals )
+	FVector* VertexNormals = nullptr;
+	if (CalcNormals)
 	{
 		// Calculate the normals for each of the two triangles per quad.
 		// Note that the normals at the edges are not correct because they include normals
 		// from triangles outside the current area. They are not updated
-		int32 NumVertsX = 1+X2-X1;
-		int32 NumVertsY = 1+Y2-Y1;
 		VertexNormals = new FVector[NumVertsX*NumVertsY];
 		FMemory::Memzero(VertexNormals, NumVertsX*NumVertsY*sizeof(FVector));
 
@@ -227,14 +228,14 @@ void FLandscapeEditDataInterface::SetHeightData(int32 X1, int32 Y1, int32 X2, in
 		FMemory::Memzero(XYOffsets, NumVertsX*NumVertsY*sizeof(FVector2D));
 		GetXYOffsetDataFast(X1, Y1, X2, Y2, XYOffsets, 0);
 
-		for( int32 Y=0;Y<NumVertsY-1;Y++ )
+		for (int32 Y = 0; Y < NumVertsY - 1; Y++)
 		{
-			for( int32 X=0;X<NumVertsX-1;X++ )
+			for (int32 X = 0; X < NumVertsX - 1; X++)
 			{
-				FVector Vert00 = FVector(XYOffsets[(X+0) + NumVertsX*(Y+0)].X,		XYOffsets[(X+0) + NumVertsX*(Y+0)].Y, ((float)Data[(X+0) + Stride*(Y+0)] - 32768.0f)*LANDSCAPE_ZSCALE) * DrawScale;
-				FVector Vert01 = FVector(XYOffsets[(X+0) + NumVertsX*(Y+0)].X,		XYOffsets[(X+0) + NumVertsX*(Y+0)].Y+1.0f, ((float)Data[(X+0) + Stride*(Y+1)] - 32768.0f)*LANDSCAPE_ZSCALE) * DrawScale;
-				FVector Vert10 = FVector(XYOffsets[(X+0) + NumVertsX*(Y+0)].X+1.0f,	XYOffsets[(X+0) + NumVertsX*(Y+0)].Y, ((float)Data[(X+1) + Stride*(Y+0)] - 32768.0f)*LANDSCAPE_ZSCALE) * DrawScale;
-				FVector Vert11 = FVector(XYOffsets[(X+0) + NumVertsX*(Y+0)].X+1.0f,	XYOffsets[(X+0) + NumVertsX*(Y+0)].Y+1.0f, ((float)Data[(X+1) + Stride*(Y+1)] - 32768.0f)*LANDSCAPE_ZSCALE) * DrawScale;
+				FVector Vert00 = FVector(XYOffsets[(X+0) + NumVertsX*(Y+0)].X,      XYOffsets[(X+0) + NumVertsX*(Y+0)].Y,      ((float)Data[(X+0) + Stride*(Y+0)] - 32768.0f) * LANDSCAPE_ZSCALE) * DrawScale;
+				FVector Vert01 = FVector(XYOffsets[(X+0) + NumVertsX*(Y+0)].X,      XYOffsets[(X+0) + NumVertsX*(Y+0)].Y+1.0f, ((float)Data[(X+0) + Stride*(Y+1)] - 32768.0f) * LANDSCAPE_ZSCALE) * DrawScale;
+				FVector Vert10 = FVector(XYOffsets[(X+0) + NumVertsX*(Y+0)].X+1.0f, XYOffsets[(X+0) + NumVertsX*(Y+0)].Y,      ((float)Data[(X+1) + Stride*(Y+0)] - 32768.0f) * LANDSCAPE_ZSCALE) * DrawScale;
+				FVector Vert11 = FVector(XYOffsets[(X+0) + NumVertsX*(Y+0)].X+1.0f, XYOffsets[(X+0) + NumVertsX*(Y+0)].Y+1.0f, ((float)Data[(X+1) + Stride*(Y+1)] - 32768.0f) * LANDSCAPE_ZSCALE) * DrawScale;
 
 				FVector FaceNormal1 = ((Vert00-Vert10) ^ (Vert10-Vert11)).GetSafeNormal();
 				FVector FaceNormal2 = ((Vert11-Vert01) ^ (Vert01-Vert00)).GetSafeNormal(); 
@@ -250,17 +251,17 @@ void FLandscapeEditDataInterface::SetHeightData(int32 X1, int32 Y1, int32 X2, in
 		delete[] XYOffsets;
 	}
 
-	for( int32 ComponentIndexY=ComponentIndexY1;ComponentIndexY<=ComponentIndexY2;ComponentIndexY++ )
+	for (int32 ComponentIndexY = ComponentIndexY1; ComponentIndexY <= ComponentIndexY2; ComponentIndexY++)
 	{
-		for( int32 ComponentIndexX=ComponentIndexX1;ComponentIndexX<=ComponentIndexX2;ComponentIndexX++ )
+		for (int32 ComponentIndexX = ComponentIndexX1; ComponentIndexX <= ComponentIndexX2; ComponentIndexX++)
 		{
-			FIntPoint ComponentKey(ComponentIndexX,ComponentIndexY);
+			FIntPoint ComponentKey(ComponentIndexX, ComponentIndexY);
 			ULandscapeComponent* Component = LandscapeInfo->XYtoComponentMap.FindRef(ComponentKey);
 
-			// if NULL, it was painted away
-			if( Component==NULL )
+			// if nullptr, it was painted away
+			if (Component == nullptr)
 			{
-				if( CreateComponents )
+				if (CreateComponents)
 				{
 					// not yet implemented
 					continue;
@@ -276,7 +277,7 @@ void FLandscapeEditDataInterface::SetHeightData(int32 X1, int32 Y1, int32 X2, in
 			FLandscapeTextureDataInfo* TexDataInfo = GetTextureDataInfo(Component->HeightmapTexture);
 			FColor* HeightmapTextureData = (FColor*)TexDataInfo->GetMipData(0);
 
-			FColor* XYOffsetMipData = NULL;
+			FColor* XYOffsetMipData = nullptr;
 			if (Component->XYOffsetmapTexture)
 			{
 				FLandscapeTextureDataInfo* XYTexDataInfo = GetTextureDataInfo(Component->XYOffsetmapTexture);
@@ -305,9 +306,9 @@ void FLandscapeEditDataInterface::SetHeightData(int32 X1, int32 Y1, int32 X2, in
 			uint16 MinHeight = MAX_uint16;
 			uint16 MaxHeight = 0;
 
-			for( int32 SubIndexY=SubIndexY1;SubIndexY<=SubIndexY2;SubIndexY++ )
+			for (int32 SubIndexY = SubIndexY1; SubIndexY <= SubIndexY2; SubIndexY++)
 			{
-				for( int32 SubIndexX=SubIndexX1;SubIndexX<=SubIndexX2;SubIndexX++ )
+				for (int32 SubIndexX = SubIndexX1; SubIndexX <= SubIndexX2; SubIndexX++)
 				{
 					// Find coordinates of box that lies inside subsection
 					int32 SubX1 = FMath::Clamp<int32>(ComponentX1-SubsectionSizeQuads*SubIndexX, 0, SubsectionSizeQuads);
@@ -316,43 +317,44 @@ void FLandscapeEditDataInterface::SetHeightData(int32 X1, int32 Y1, int32 X2, in
 					int32 SubY2 = FMath::Clamp<int32>(ComponentY2-SubsectionSizeQuads*SubIndexY, 0, SubsectionSizeQuads);
 
 					// Update texture data for the box that lies inside subsection
-					for( int32 SubY=SubY1;SubY<=SubY2;SubY++ )
+					for (int32 SubY = SubY1; SubY <= SubY2; SubY++)
 					{
-						for( int32 SubX=SubX1;SubX<=SubX2;SubX++ )
+						for (int32 SubX = SubX1; SubX <= SubX2; SubX++)
 						{
 							int32 LandscapeX = SubIndexX*SubsectionSizeQuads + ComponentIndexX*ComponentSizeQuads + SubX;
 							int32 LandscapeY = SubIndexY*SubsectionSizeQuads + ComponentIndexY*ComponentSizeQuads + SubY;
-							checkSlow( LandscapeX >= X1 && LandscapeX <= X2 );
-							checkSlow( LandscapeY >= Y1 && LandscapeY <= Y2 );
+							checkSlow(LandscapeX >= X1 && LandscapeX <= X2);
+							checkSlow(LandscapeY >= Y1 && LandscapeY <= Y2);
 
 							// Find the input data corresponding to this vertex
 							int32 DataIndex = (LandscapeX-X1) + Stride * (LandscapeY-Y1);
 							const uint16& Height = Data[DataIndex];
 
 							// for bounding box
-							if( Height < MinHeight )
+							if (Height < MinHeight)
 							{
 								MinHeight = Height;
 							}
-							if( Height > MaxHeight )
+							if (Height > MaxHeight)
 							{
 								MaxHeight = Height;
 							}
 
 							int32 TexX = HeightmapOffsetX + (SubsectionSizeQuads+1) * SubIndexX + SubX;
 							int32 TexY = HeightmapOffsetY + (SubsectionSizeQuads+1) * SubIndexY + SubY;
-							FColor& TexData = HeightmapTextureData[ TexX + TexY * SizeU ];
+							FColor& TexData = HeightmapTextureData[TexX + TexY * SizeU];
 
 							// Update the texture
 							TexData.R = Height >> 8;
 							TexData.G = Height & 255;
 
 							// Update normals if we're not on an edge vertex
-							if( VertexNormals && LandscapeX > X1 && LandscapeX < X2 && LandscapeY > Y1 && LandscapeY < Y2 )
+							if (VertexNormals && LandscapeX > X1 && LandscapeX < X2 && LandscapeY > Y1 && LandscapeY < Y2)
 							{
-								FVector Normal = VertexNormals[DataIndex].GetSafeNormal();
-								TexData.B = FMath::RoundToInt( 127.5f * (Normal.X + 1.0f) );
-								TexData.A = FMath::RoundToInt( 127.5f * (Normal.Y + 1.0f) );
+								const int32 NormalDataIndex = (LandscapeX-X1) + NumVertsX * (LandscapeY-Y1);
+								FVector Normal = VertexNormals[NormalDataIndex].GetSafeNormal();
+								TexData.B = FMath::RoundToInt(127.5f * (Normal.X + 1.0f));
+								TexData.A = FMath::RoundToInt(127.5f * (Normal.Y + 1.0f));
 							}
 							else if (NormalData)
 							{
@@ -378,18 +380,18 @@ void FLandscapeEditDataInterface::SetHeightData(int32 X1, int32 Y1, int32 X2, in
 			float MaxLocalZ = LandscapeDataAccess::GetLocalHeight(MaxHeight);
 
 			bool bUpdateBoxSphereBounds = false;
-			if( MinLocalZ < Component->CachedLocalBox.Min.Z )
+			if (MinLocalZ < Component->CachedLocalBox.Min.Z)
 			{
 				Component->CachedLocalBox.Min.Z = MinLocalZ;
 				bUpdateBoxSphereBounds = true;
 			}
-			if( MaxLocalZ > Component->CachedLocalBox.Max.Z )
+			if (MaxLocalZ > Component->CachedLocalBox.Max.Z)
 			{
 				Component->CachedLocalBox.Max.Z = MaxLocalZ;
 				bUpdateBoxSphereBounds = true;
 			}
 
-			if( bUpdateBoxSphereBounds )
+			if (bUpdateBoxSphereBounds)
 			{
 				Component->UpdateComponentToWorld();
 			}
@@ -403,21 +405,21 @@ void FLandscapeEditDataInterface::SetHeightData(int32 X1, int32 Y1, int32 X2, in
 			TArray<FColor*> MipData;
 			MipData.AddUninitialized(BaseNumMips);
 			MipData[0] = HeightmapTextureData;
-			for( int32 MipIdx=1;MipIdx<BaseNumMips;MipIdx++ )
+			for (int32 MipIdx = 1; MipIdx < BaseNumMips; MipIdx++)
 			{
 				MipData[MipIdx] = (FColor*)TexDataInfo->GetMipData(MipIdx);
 			}
-			Component->GenerateHeightmapMips( MipData, ComponentX1, ComponentY1, ComponentX2, ComponentY2, TexDataInfo );
+			Component->GenerateHeightmapMips(MipData, ComponentX1, ComponentY1, ComponentX2, ComponentY2, TexDataInfo);
 
 			// Update collision
-			Component->UpdateCollisionHeightData(MipData[Component->CollisionMipLevel], ComponentX1, ComponentY1, ComponentX2, ComponentY2, bUpdateBoxSphereBounds, XYOffsetMipData );
+			Component->UpdateCollisionHeightData(MipData[Component->CollisionMipLevel], ComponentX1, ComponentY1, ComponentX2, ComponentY2, bUpdateBoxSphereBounds, XYOffsetMipData);
 
 			// Update GUID for Platform Data
 			FPlatformMisc::CreateGuid(Component->StateId);
 		}
 	}
 
-	if( VertexNormals )
+	if (VertexNormals)
 	{
 		delete[] VertexNormals;
 	}
