Fixed glitches when mirroring landscape material layers with the landscape mirror tool

diff --git a/Engine/Source/Editor/LandscapeEditor/Private/LandscapeEdModeMirrorTool.cpp b/Engine/Source/Editor/LandscapeEditor/Private/LandscapeEdModeMirrorTool.cpp
--- a/Engine/Source/Editor/LandscapeEditor/Private/LandscapeEdModeMirrorTool.cpp
+++ b/Engine/Source/Editor/LandscapeEditor/Private/LandscapeEdModeMirrorTool.cpp
@@ -434,7 +434,7 @@ public:
 				TempMaxY = SourceMaxY;
 				LandscapeEdit.GetWeightData(LayerInfo, TempMinX, TempMinY, TempMaxX, TempMaxY, &MirrorWeightData[ReadOffset], ReadStride);
 				ApplyMirrorInternal(MirrorWeightData.GetData(), SizeX, SizeY, MirrorSize);
-				LandscapeEdit.SetAlphaData(LayerInfo, DestMinX, DestMinY, DestMaxX, DestMaxY, &MirrorWeightData[WriteOffset], WriteStride, ELandscapeLayerPaintingRestriction::None);
+				LandscapeEdit.SetAlphaData(LayerInfo, DestMinX, DestMinY, DestMaxX, DestMaxY, &MirrorWeightData[WriteOffset], WriteStride, ELandscapeLayerPaintingRestriction::None, false, false);
 			}
 		}
 
