VR Editor: Landscape editing in VR (from Trung Le)

diff --git a/Engine/Source/Editor/LandscapeEditor/LandscapeEditor.Build.cs b/Engine/Source/Editor/LandscapeEditor/LandscapeEditor.Build.cs
--- a/Engine/Source/Editor/LandscapeEditor/LandscapeEditor.Build.cs
+++ b/Engine/Source/Editor/LandscapeEditor/LandscapeEditor.Build.cs
@@ -23,6 +23,8 @@ public class LandscapeEditor : ModuleRules
 				"ImageWrapper",
                 "EditorWidgets",
                 "Foliage",
+				"ViewportInteraction",
+				"VREditor",
 			}
 			);
 
diff --git a/Engine/Source/Editor/LandscapeEditor/Private/LandscapeEdMode.cpp b/Engine/Source/Editor/LandscapeEditor/Private/LandscapeEdMode.cpp
--- a/Engine/Source/Editor/LandscapeEditor/Private/LandscapeEdMode.cpp
+++ b/Engine/Source/Editor/LandscapeEditor/Private/LandscapeEdMode.cpp
@@ -25,6 +25,11 @@
 #include "Editor/LevelEditor/Public/SLevelViewport.h"
 #include "SLandscapeEditor.h"
 
+// VR Editor
+#include "IVREditorModule.h"
+#include "ViewportWorldInteraction.h"
+#include "VREditorInteractor.h"
+
 // Classes
 #include "Landscape.h"
 #include "LandscapeLayerInfoObject.h"
@@ -174,6 +179,7 @@ FEdModeLandscape::FEdModeLandscape()
 	, LandscapeRenderAddCollision(nullptr)
 	, CachedLandscapeMaterial(nullptr)
 	, bToolActive(false)
+	, bIsPainting(false)
 {
 	GLayerDebugColorMaterial = LandscapeTool::CreateMaterialInstance(LoadObject<UMaterial>(nullptr, TEXT("/Engine/EditorLandscapeResources/LayerVisMaterial.LayerVisMaterial")));
 	GSelectionColorMaterial  = LandscapeTool::CreateMaterialInstance(LoadObject<UMaterialInstanceConstant>(nullptr, TEXT("/Engine/EditorLandscapeResources/SelectBrushMaterial_Selected.SelectBrushMaterial_Selected")));
@@ -461,12 +467,41 @@ void FEdModeLandscape::Enter()
 	{
 		GizmoBrush->EnterBrush();
 	}
+
+	// Register to find out about editor modes being activated
+	GetModeManager()->OnEditorModeChanged().AddRaw(this, &FEdModeLandscape::OnEditorModeChanged);
+
+	// Register to find out about VR input events
+	IVREditorMode* VREditorMode = static_cast<IVREditorMode*>(GetModeManager()->GetActiveMode(IVREditorModule::Get().GetVREditorModeID()));
+	if (VREditorMode != nullptr)
+	{
+		VREditorMode->GetWorldInteraction().OnViewportInteractionInputAction().RemoveAll(this);
+		VREditorMode->GetWorldInteraction().OnViewportInteractionInputAction().AddRaw(this, &FEdModeLandscape::OnVRAction);
+
+		VREditorMode->GetWorldInteraction().OnViewportInteractionHoverUpdate().RemoveAll(this);
+		VREditorMode->GetWorldInteraction().OnViewportInteractionHoverUpdate().AddRaw(this, &FEdModeLandscape::OnVRHoverUpdate);
+	}
 }
 
 
 /** FEdMode: Called when the mode is exited */
 void FEdModeLandscape::Exit()
 {
+	// Unregister VR mode from event handlers
+	if (IVREditorModule::IsAvailable())
+	{
+		IVREditorMode* VREditorMode = static_cast<IVREditorMode*>(GetModeManager()->GetActiveMode(IVREditorModule::Get().GetVREditorModeID()));
+		if (VREditorMode != nullptr)
+		{
+			VREditorMode->GetWorldInteraction().OnViewportInteractionInputAction().RemoveAll(this);
+			VREditorMode->GetWorldInteraction().OnViewportInteractionHoverUpdate().RemoveAll(this);
+		}
+	}
+	
+	// Unregister landscape mode from event handlers
+	GetModeManager()->OnEditorModeChanged().RemoveAll(this);
+
+
 	FEditorSupportDelegates::WorldChange.Remove(OnWorldChangeDelegateHandle);
 	UMaterial::OnMaterialCompilationFinished().Remove(OnMaterialCompilationFinishedDelegateHandle);
 
@@ -537,6 +572,116 @@ void FEdModeLandscape::Exit()
 	FEdMode::Exit();
 }
 
+
+void FEdModeLandscape::OnEditorModeChanged( FEdMode* EditorMode, bool bEntered )
+{
+	// VR Editor mode may have gone away, so re-register for events in case it's a different object than the one we were originally bound to
+	check(EditorMode != nullptr);
+	if (bEntered && EditorMode->GetID() == IVREditorModule::Get().GetVREditorModeID())
+	{
+		IVREditorMode* VREditorMode = static_cast<IVREditorMode*>(EditorMode);
+		if (VREditorMode != nullptr)
+		{
+			VREditorMode->GetWorldInteraction().OnViewportInteractionInputAction().RemoveAll(this);
+			VREditorMode->GetWorldInteraction().OnViewportInteractionInputAction().AddRaw(this, &FEdModeLandscape::OnVRAction);
+
+			VREditorMode->GetWorldInteraction().OnViewportInteractionHoverUpdate().RemoveAll(this);
+			VREditorMode->GetWorldInteraction().OnViewportInteractionHoverUpdate().AddRaw(this, &FEdModeLandscape::OnVRHoverUpdate);
+		}
+	}
+}
+
+void FEdModeLandscape::OnVRHoverUpdate(FEditorViewportClient& ViewportClient, UViewportInteractor* Interactor, FVector& HoverImpactPoint, bool& bWasHandled)
+{
+	IVREditorMode* VREditorMode = static_cast<IVREditorMode*>(GetModeManager()->GetActiveMode(IVREditorModule::Get().GetVREditorModeID()));
+
+	if (VREditorMode != nullptr && Interactor != nullptr)
+	{
+		const UVREditorInteractor* VRInteractor = Cast<UVREditorInteractor>(Interactor);
+		const bool bIsHoveringOverUIVR = VRInteractor->IsHoveringOverUI();
+
+		if (!bIsHoveringOverUIVR && CurrentTool && (CurrentTool->GetSupportedTargetTypes() == ELandscapeToolTargetTypeMask::NA || CurrentToolTarget.TargetType != ELandscapeToolTargetType::Invalid))
+		{
+			FVector HitLocation;
+			FVector LaserPointerStart, LaserPointerEnd;
+			if (Interactor->GetLaserPointer( /* Out */ LaserPointerStart, /* Out */ LaserPointerEnd))
+			{
+				if (LandscapeTrace(&ViewportClient, LaserPointerStart, LaserPointerEnd, HitLocation))
+				{
+					if (bIsPainting && !(CurrentToolTarget.TargetType == ELandscapeToolTargetType::Weightmap && CurrentToolTarget.LayerInfo == NULL))
+					{
+						bToolActive = CurrentTool->BeginTool(&ViewportClient, CurrentToolTarget, HitLocation, Interactor);
+					}
+
+					if (CurrentBrush)
+					{
+						// Inform the brush of the current location, to update the cursor
+						CurrentBrush->MouseMove(HitLocation.X, HitLocation.Y);
+					}
+				}
+			}
+		}
+	}
+}
+
+void FEdModeLandscape::OnVRAction(FEditorViewportClient& ViewportClient, UViewportInteractor* Interactor, const struct FViewportActionKeyInput& Action, bool& bOutIsInputCaptured, bool& bWasHandled)
+{
+	IVREditorMode* VREditorMode = static_cast<IVREditorMode*>(GetModeManager()->GetActiveMode(IVREditorModule::Get().GetVREditorModeID()));
+	
+	// Never show the traditional Unreal transform widget.  It doesn't work in VR because we don't have hit proxies.
+	ViewportClient.EngineShowFlags.SetModeWidgets(false);
+
+	if (VREditorMode != nullptr && Interactor != nullptr)
+	{
+		if (Action.ActionType == ViewportWorldActionTypes::SelectAndMove_LightlyPressed || Action.ActionType == ViewportWorldActionTypes::SelectAndMove)
+		{
+			const UVREditorInteractor* VRInteractor = Cast<UVREditorInteractor>(Interactor);
+
+			// Begin landscape brush
+			if (Action.Event == IE_Pressed && !VRInteractor->IsHoveringOverUI())
+			{
+				if (bToolActive)
+				{
+					CurrentTool->EndTool(&ViewportClient);
+					bToolActive = false;
+				}
+
+				if (CurrentTool && (CurrentTool->GetSupportedTargetTypes() == ELandscapeToolTargetTypeMask::NA || CurrentToolTarget.TargetType != ELandscapeToolTargetType::Invalid))
+				{
+					FVector HitLocation;
+					FVector LaserPointerStart, LaserPointerEnd;
+					if (Interactor->GetLaserPointer( /* Out */ LaserPointerStart, /* Out */ LaserPointerEnd))
+					{
+						if (LandscapeTrace(&ViewportClient, LaserPointerStart, LaserPointerEnd, HitLocation))
+						{
+							if (!(CurrentToolTarget.TargetType == ELandscapeToolTargetType::Weightmap && CurrentToolTarget.LayerInfo == NULL))
+							{
+								bToolActive = CurrentTool->BeginTool(&ViewportClient, CurrentToolTarget, HitLocation, Interactor);
+							}
+
+							bIsPainting = true;
+							bWasHandled = true;
+							bOutIsInputCaptured = false;
+						}
+					}
+				}
+			}
+
+			// End landscape brush
+			else if (Action.Event == IE_Released)
+			{
+				if (CurrentTool && bToolActive)
+				{
+					CurrentTool->EndTool(&ViewportClient);
+					bToolActive = false;
+				}
+				bIsPainting = false;
+			}
+		}
+	}
+
+}
+
 /** FEdMode: Called once per frame */
 void FEdModeLandscape::Tick(FEditorViewportClient* ViewportClient, float DeltaTime)
 {
@@ -759,7 +904,18 @@ bool FEdModeLandscape::LandscapeMouseTrace(FEditorViewportClient* ViewportClient
 
 	FVector End = Start + WORLD_MAX * BrushTraceDirection;
 
-	static FName TraceTag = FName(TEXT("LandscapeMouseTrace"));
+	return LandscapeTrace(ViewportClient, Start, End, OutHitLocation);
+}
+
+bool FEdModeLandscape::LandscapeTrace(const FEditorViewportClient* ViewportClient, const FVector& InRayOrigin, const FVector& InRayEnd, FVector& OutHitLocation)
+{
+	FVector Start = InRayOrigin;
+	FVector End = InRayEnd;
+
+	// Cache a copy of the world pointer
+	UWorld* World = ViewportClient->GetWorld();
+
+	static FName TraceTag = FName(TEXT("LandscapeTrace"));
 	TArray<FHitResult> Results;
 	// Each landscape component has 2 collision shapes, 1 of them is specific to landscape editor
 	// Trace only ECC_Visibility channel, so we do hit only Editor specific shape
@@ -834,8 +990,6 @@ bool FEdModeLandscape::LandscapePlaneTrace(FEditorViewportClient* ViewportClient
 
 bool FEdModeLandscape::LandscapePlaneTrace(FEditorViewportClient* ViewportClient, int32 MouseX, int32 MouseY, const FPlane& Plane, FVector& OutHitLocation)
 {
-	// Cache a copy of the world pointer
-	UWorld* World = ViewportClient->GetWorld();
 	// Compute a world space ray from the screen space mouse coordinates
 	FSceneViewFamilyContext ViewFamily(FSceneViewFamily::ConstructionValues(
 		ViewportClient->Viewport,
@@ -2361,7 +2515,17 @@ void FEdModeLandscape::ForceRealTimeViewports(const bool bEnable, const bool bSt
 				if (bEnable)
 				{
 					Viewport.SetRealtime(bEnable, bStoreCurrentState);
-					Viewport.SetGameView(false);
+
+					// @todo vreditor: Force game view to true in VREditor since we can't use hitproxies and debug objects yet
+					IVREditorMode* VREditorMode = static_cast<IVREditorMode*>(GetModeManager()->GetActiveMode(IVREditorModule::Get().GetVREditorModeID()));
+					if (VREditorMode != nullptr)
+					{
+						Viewport.SetGameView(true);
+					} 
+					else
+					{
+						Viewport.SetGameView(false);
+					}
 				}
 				else
 				{
diff --git a/Engine/Source/Editor/LandscapeEditor/Private/LandscapeEdMode.h b/Engine/Source/Editor/LandscapeEditor/Private/LandscapeEdMode.h
--- a/Engine/Source/Editor/LandscapeEditor/Private/LandscapeEdMode.h
+++ b/Engine/Source/Editor/LandscapeEditor/Private/LandscapeEdMode.h
@@ -14,6 +14,8 @@ DECLARE_LOG_CATEGORY_EXTERN(LogLandscapeEdMode, Log, All);
 // Forward declarations
 class ULandscapeLayerInfoObject;
 class FLandscapeToolSplines;
+class UViewportInteractor;
+struct FViewportActionKeyInput;
 
 struct FHeightmapToolTarget;
 template<typename TargetType> class FLandscapeToolCopyPaste;
@@ -398,6 +400,9 @@ public:
 	bool LandscapePlaneTrace(FEditorViewportClient* ViewportClient, const FPlane& Plane, FVector& OutHitLocation);
 	bool LandscapePlaneTrace(FEditorViewportClient* ViewportClient, int32 MouseX, int32 MouseY, const FPlane& Plane, FVector& OutHitLocation);
 
+	/** Trace under the specified laser start and direction and return the landscape hit and the hit location (in landscape quad space) */
+	bool LandscapeTrace(const FEditorViewportClient* ViewportClient, const FVector& InRayOrigin, const FVector& InRayEnd, FVector& OutHitLocation);
+
 	void SetCurrentToolMode(FName ToolModeName, bool bRestoreCurrentTool = true);
 
 	/** Change current tool */
@@ -421,6 +426,14 @@ public:
 	DECLARE_EVENT(FEdModeLandscape, FTargetsListUpdated);
 	static FTargetsListUpdated TargetsListUpdated;
 
+	/** Called when an editor mode is entered or exited */
+	void OnEditorModeChanged(FEdMode* EditorMode, bool bEntered);
+
+	/** Called when the user presses a button on their motion controller device */
+	void OnVRAction(FEditorViewportClient& ViewportClient, UViewportInteractor* Interactor, const FViewportActionKeyInput& Action, bool& bOutIsInputCaptured, bool& bWasHandled);
+
+	void OnVRHoverUpdate(FEditorViewportClient& ViewportClient, UViewportInteractor* Interactor, FVector& HoverImpactPoint, bool& bWasHandled);
+
 	void OnWorldChange();
 
 	void OnMaterialCompilationFinished(UMaterialInterface* MaterialInterface);
@@ -454,4 +467,6 @@ private:
 
 	FDelegateHandle OnWorldChangeDelegateHandle;
 	FDelegateHandle OnMaterialCompilationFinishedDelegateHandle;
+
+	bool bIsPainting;
 };
diff --git a/Engine/Source/Editor/LandscapeEditor/Private/LandscapeEdModeBrushes.cpp b/Engine/Source/Editor/LandscapeEditor/Private/LandscapeEdModeBrushes.cpp
--- a/Engine/Source/Editor/LandscapeEditor/Private/LandscapeEdModeBrushes.cpp
+++ b/Engine/Source/Editor/LandscapeEditor/Private/LandscapeEdModeBrushes.cpp
@@ -222,7 +222,7 @@ public:
 		LastMousePosition = FVector2D(LandscapeX, LandscapeY);
 	}
 
-	virtual FLandscapeBrushData ApplyBrush(const TArray<FLandscapeToolMousePosition>& InMousePositions) override
+	virtual FLandscapeBrushData ApplyBrush(const TArray<FLandscapeToolInteractorPosition>& InInteractorPositions) override
 	{
 		ULandscapeInfo* LandscapeInfo = EdMode->CurrentToolTarget.LandscapeInfo.Get();
 		const float ScaleXY = FMath::Abs(LandscapeInfo->DrawScale.X);
@@ -231,28 +231,28 @@ public:
 		const float Falloff = EdMode->UISettings->BrushFalloff * TotalRadius;
 
 		// Cap number of mouse positions to a sensible number
-		TArray<FLandscapeToolMousePosition> MousePositions;
-		if (InMousePositions.Num() > 10)
+		TArray<FLandscapeToolInteractorPosition> InteractorPositions;
+		if (InInteractorPositions.Num() > 10)
 		{
 			for (int32 i = 0; i < 10; ++i)
 			{
 				// Scale so we include the first and last of the input positions
-				MousePositions.Add(InMousePositions[(i * (InMousePositions.Num() - 1)) / 9]);
+				InteractorPositions.Add(InInteractorPositions[(i * (InInteractorPositions.Num() - 1)) / 9]);
 			}
 		}
 		else
 		{
-			MousePositions = InMousePositions;
+			InteractorPositions = InInteractorPositions;
 		}
 
 		FIntRect Bounds;
-		for (const FLandscapeToolMousePosition& MousePosition : MousePositions)
+		for (const FLandscapeToolInteractorPosition& InteractorPosition : InteractorPositions)
 		{
 			FIntRect SpotBounds;
-			SpotBounds.Min.X = FMath::FloorToInt(MousePosition.Position.X - TotalRadius);
-			SpotBounds.Min.Y = FMath::FloorToInt(MousePosition.Position.Y - TotalRadius);
-			SpotBounds.Max.X = FMath::CeilToInt( MousePosition.Position.X + TotalRadius);
-			SpotBounds.Max.Y = FMath::CeilToInt( MousePosition.Position.Y + TotalRadius);
+			SpotBounds.Min.X = FMath::FloorToInt(InteractorPosition.Position.X - TotalRadius);
+			SpotBounds.Min.Y = FMath::FloorToInt(InteractorPosition.Position.Y - TotalRadius);
+			SpotBounds.Max.X = FMath::CeilToInt( InteractorPosition.Position.X + TotalRadius);
+			SpotBounds.Max.Y = FMath::CeilToInt( InteractorPosition.Position.Y + TotalRadius);
 
 			if (Bounds.IsEmpty())
 			{
@@ -276,13 +276,13 @@ public:
 
 		FLandscapeBrushData BrushData(Bounds);
 
-		for (const FLandscapeToolMousePosition& MousePosition : MousePositions)
+		for (const FLandscapeToolInteractorPosition& InteractorPosition : InteractorPositions)
 		{
 			FIntRect SpotBounds;
-			SpotBounds.Min.X = FMath::Max(FMath::FloorToInt(MousePosition.Position.X - TotalRadius), Bounds.Min.X);
-			SpotBounds.Min.Y = FMath::Max(FMath::FloorToInt(MousePosition.Position.Y - TotalRadius), Bounds.Min.Y);
-			SpotBounds.Max.X = FMath::Min(FMath::CeilToInt( MousePosition.Position.X + TotalRadius), Bounds.Max.X);
-			SpotBounds.Max.Y = FMath::Min(FMath::CeilToInt( MousePosition.Position.Y + TotalRadius), Bounds.Max.Y);
+			SpotBounds.Min.X = FMath::Max(FMath::FloorToInt(InteractorPosition.Position.X - TotalRadius), Bounds.Min.X);
+			SpotBounds.Min.Y = FMath::Max(FMath::FloorToInt(InteractorPosition.Position.Y - TotalRadius), Bounds.Min.Y);
+			SpotBounds.Max.X = FMath::Min(FMath::CeilToInt( InteractorPosition.Position.X + TotalRadius), Bounds.Max.X);
+			SpotBounds.Max.Y = FMath::Min(FMath::CeilToInt( InteractorPosition.Position.Y + TotalRadius), Bounds.Max.Y);
 
 			for (int32 Y = SpotBounds.Min.Y; Y < SpotBounds.Max.Y; Y++)
 			{
@@ -294,7 +294,7 @@ public:
 					if (PrevAmount < 1.0f)
 					{
 						// Distance from mouse
-						float MouseDist = FMath::Sqrt(FMath::Square(MousePosition.Position.X - (float)X) + FMath::Square(MousePosition.Position.Y - (float)Y));
+						float MouseDist = FMath::Sqrt(FMath::Square(InteractorPosition.Position.X - (float)X) + FMath::Square(InteractorPosition.Position.Y - (float)Y));
 
 						float PaintAmount = CalculateFalloff(MouseDist, Radius, Falloff);
 
@@ -446,7 +446,7 @@ public:
 		LastMousePosition = FVector2D(LandscapeX, LandscapeY);
 	}
 
-	virtual FLandscapeBrushData ApplyBrush(const TArray<FLandscapeToolMousePosition>& MousePositions) override
+	virtual FLandscapeBrushData ApplyBrush(const TArray<FLandscapeToolInteractorPosition>& InteractorPositions) override
 	{
 		// Selection Brush only works for 
 		ULandscapeInfo* LandscapeInfo = EdMode->CurrentToolTarget.LandscapeInfo.Get();
@@ -675,7 +675,7 @@ public:
 		return TOptional<bool>();
 	}
 
-	virtual FLandscapeBrushData ApplyBrush(const TArray<FLandscapeToolMousePosition>& MousePositions) override
+	virtual FLandscapeBrushData ApplyBrush(const TArray<FLandscapeToolInteractorPosition>& InteractorPositions) override
 	{
 		// Selection Brush only works for 
 		ALandscapeGizmoActiveActor* Gizmo = EdMode->CurrentGizmoActor.Get();
@@ -797,7 +797,7 @@ public:
 	{
 	}
 
-	virtual FLandscapeBrushData ApplyBrush(const TArray<FLandscapeToolMousePosition>& MousePositions) override
+	virtual FLandscapeBrushData ApplyBrush(const TArray<FLandscapeToolInteractorPosition>& InteractorPositions) override
 	{
 		return FLandscapeBrushData();
 	}
@@ -829,7 +829,7 @@ public:
 	{
 	}
 
-	virtual FLandscapeBrushData ApplyBrush(const TArray<FLandscapeToolMousePosition>& MousePositions) override
+	virtual FLandscapeBrushData ApplyBrush(const TArray<FLandscapeToolInteractorPosition>& InteractorPositions) override
 	{
 		return FLandscapeBrushData();
 	}
@@ -1013,7 +1013,7 @@ public:
 
 	virtual ELandscapeBrushType GetBrushType() override { return ELandscapeBrushType::Alpha; }
 
-	virtual FLandscapeBrushData ApplyBrush(const TArray<FLandscapeToolMousePosition>& MousePositions) override
+	virtual FLandscapeBrushData ApplyBrush(const TArray<FLandscapeToolInteractorPosition>& InteractorPositions) override
 	{
 		ULandscapeInfo* LandscapeInfo = EdMode->CurrentToolTarget.LandscapeInfo.Get();
 		const float ScaleXY = FMath::Abs(LandscapeInfo->DrawScale.X);
@@ -1208,7 +1208,7 @@ public:
 		return new FLandscapeBrushAlpha(InEdMode, AlphaBrushMaterial);
 	}
 
-	virtual FLandscapeBrushData ApplyBrush(const TArray<FLandscapeToolMousePosition>& MousePositions) override
+	virtual FLandscapeBrushData ApplyBrush(const TArray<FLandscapeToolInteractorPosition>& InteractorPositions) override
 	{
 		ULandscapeInfo* LandscapeInfo = EdMode->CurrentToolTarget.LandscapeInfo.Get();
 		if (EdMode->UISettings->bAlphaBrushAutoRotate && OldMousePosition.IsZero())
diff --git a/Engine/Source/Editor/LandscapeEditor/Private/LandscapeEdModeComponentTools.cpp b/Engine/Source/Editor/LandscapeEditor/Private/LandscapeEdModeComponentTools.cpp
--- a/Engine/Source/Editor/LandscapeEditor/Private/LandscapeEdModeComponentTools.cpp
+++ b/Engine/Source/Editor/LandscapeEditor/Private/LandscapeEdModeComponentTools.cpp
@@ -34,19 +34,19 @@ public:
 	{
 	}
 
-	void Apply(FEditorViewportClient* ViewportClient, FLandscapeBrush* Brush, const ULandscapeEditorObject* UISettings, const TArray<FLandscapeToolMousePosition>& MousePositions)
+	void Apply(FEditorViewportClient* ViewportClient, FLandscapeBrush* Brush, const ULandscapeEditorObject* UISettings, const TArray<FLandscapeToolInteractorPosition>& InteractorPositions)
 	{
 		if (LandscapeInfo)
 		{
 			LandscapeInfo->Modify();
 
 			// Invert when holding Shift
-			bool bInvert = MousePositions[MousePositions.Num() - 1].bShiftDown;
+			bool bInvert = InteractorPositions[InteractorPositions.Num() - 1].bModifierPressed;
 
 			if (Brush->GetBrushType() == ELandscapeBrushType::Component)
 			{
 				// TODO - only retrieve bounds as we don't need the data
-				const FLandscapeBrushData BrushInfo = Brush->ApplyBrush(MousePositions);
+				const FLandscapeBrushData BrushInfo = Brush->ApplyBrush(InteractorPositions);
 				if (!BrushInfo)
 				{
 					return;
@@ -62,8 +62,8 @@ public:
 				if (!bInitializedComponentInvert)
 				{
 					// Get the component under the mouse location. Copied from FLandscapeBrushComponent::ApplyBrush()
-					const float MouseX = MousePositions[0].Position.X;
-					const float MouseY = MousePositions[0].Position.Y;
+					const float MouseX = InteractorPositions[0].Position.X;
+					const float MouseY = InteractorPositions[0].Position.Y;
 					const int32 MouseComponentIndexX = (MouseX >= 0.0f) ? FMath::FloorToInt(MouseX / LandscapeInfo->ComponentSizeQuads) : FMath::CeilToInt(MouseX / LandscapeInfo->ComponentSizeQuads);
 					const int32 MouseComponentIndexY = (MouseY >= 0.0f) ? FMath::FloorToInt(MouseY / LandscapeInfo->ComponentSizeQuads) : FMath::CeilToInt(MouseY / LandscapeInfo->ComponentSizeQuads);
 					ULandscapeComponent* MouseComponent = LandscapeInfo->XYtoComponentMap.FindRef(FIntPoint(MouseComponentIndexX, MouseComponentIndexY));
@@ -107,7 +107,7 @@ public:
 			}
 			else // Select various shape regions
 			{
-				const FLandscapeBrushData BrushInfo = Brush->ApplyBrush(MousePositions);
+				const FLandscapeBrushData BrushInfo = Brush->ApplyBrush(InteractorPositions);
 				if (!BrushInfo)
 				{
 					return;
@@ -218,13 +218,13 @@ public:
 	{
 	}
 
-	void Apply(FEditorViewportClient* ViewportClient, FLandscapeBrush* Brush, const ULandscapeEditorObject* UISettings, const TArray<FLandscapeToolMousePosition>& MousePositions)
+	void Apply(FEditorViewportClient* ViewportClient, FLandscapeBrush* Brush, const ULandscapeEditorObject* UISettings, const TArray<FLandscapeToolInteractorPosition>& InteractorPositions)
 	{
 		if (LandscapeInfo)
 		{
 			LandscapeInfo->Modify();
 			// Get list of verts to update
-			FLandscapeBrushData BrushInfo = Brush->ApplyBrush(MousePositions);
+			FLandscapeBrushData BrushInfo = Brush->ApplyBrush(InteractorPositions);
 			if (!BrushInfo)
 			{
 				return;
@@ -234,7 +234,7 @@ public:
 			BrushInfo.GetInclusiveBounds(X1, Y1, X2, Y2);
 
 			// Invert when holding Shift
-			bool bInvert = MousePositions[MousePositions.Num() - 1].bShiftDown;
+			bool bInvert = InteractorPositions[InteractorPositions.Num() - 1].bModifierPressed;
 
 			// Tablet pressure
 			float Pressure = ViewportClient->Viewport->IsPenActive() ? ViewportClient->Viewport->GetTabletPressure() : 1.0f;
@@ -280,7 +280,7 @@ public:
 	{
 	}
 
-	virtual bool BeginTool(FEditorViewportClient* ViewportClient, const FLandscapeToolTarget& InTarget, const FVector& InHitLocation) override
+	virtual bool BeginTool(FEditorViewportClient* ViewportClient, const FLandscapeToolTarget& InTarget, const FVector& InHitLocation, const UViewportInteractor* Interactor = nullptr) override
 	{
 		ALandscapeProxy* Proxy = InTarget.LandscapeInfo->GetLandscapeProxy();
 		UMaterialInterface* HoleMaterial = Proxy->GetLandscapeHoleMaterial();
@@ -318,7 +318,7 @@ public:
 	{
 	}
 
-	void Apply(FEditorViewportClient* ViewportClient, FLandscapeBrush* Brush, const ULandscapeEditorObject* UISettings, const TArray<FLandscapeToolMousePosition>& MousePositions)
+	void Apply(FEditorViewportClient* ViewportClient, FLandscapeBrush* Brush, const ULandscapeEditorObject* UISettings, const TArray<FLandscapeToolInteractorPosition>& InteractorPositions)
 	{
 		ALandscape* Landscape = LandscapeInfo ? LandscapeInfo->LandscapeActor.Get() : nullptr;
 
@@ -357,7 +357,7 @@ public:
 			{
 				// Get list of verts to update
 				// TODO - only retrieve bounds as we don't need the data
-				FLandscapeBrushData BrushInfo = Brush->ApplyBrush(MousePositions);
+				FLandscapeBrushData BrushInfo = Brush->ApplyBrush(InteractorPositions);
 				if (!BrushInfo)
 				{
 					return;
@@ -773,7 +773,7 @@ public:
 		XYOffsetCache.Flush();
 	}
 
-	virtual void Apply(FEditorViewportClient* ViewportClient, FLandscapeBrush* Brush, const ULandscapeEditorObject* UISettings, const TArray<FLandscapeToolMousePosition>& MousePositions)
+	virtual void Apply(FEditorViewportClient* ViewportClient, FLandscapeBrush* Brush, const ULandscapeEditorObject* UISettings, const TArray<FLandscapeToolInteractorPosition>& InteractorPositions)
 	{
 		ALandscapeProxy* Landscape = LandscapeInfo ? LandscapeInfo->GetCurrentLevelLandscapeProxy(true) : nullptr;
 		if (Landscape && EdMode->LandscapeRenderAddCollision)
@@ -782,7 +782,7 @@ public:
 
 			// Get list of verts to update
 			// TODO - only retrieve bounds as we don't need the data
-			FLandscapeBrushData BrushInfo = Brush->ApplyBrush(MousePositions);
+			FLandscapeBrushData BrushInfo = Brush->ApplyBrush(InteractorPositions);
 			if (!BrushInfo)
 			{
 				return;
@@ -973,7 +973,7 @@ public:
 	{
 	}
 
-	void Apply(FEditorViewportClient* ViewportClient, FLandscapeBrush* Brush, const ULandscapeEditorObject* UISettings, const TArray<FLandscapeToolMousePosition>& MousePositions)
+	void Apply(FEditorViewportClient* ViewportClient, FLandscapeBrush* Brush, const ULandscapeEditorObject* UISettings, const TArray<FLandscapeToolInteractorPosition>& InteractorPositions)
 	{
 		if (LandscapeInfo)
 		{
@@ -984,7 +984,7 @@ public:
 			{
 				// Get list of components to delete from brush
 				// TODO - only retrieve bounds as we don't need the vert data
-				FLandscapeBrushData BrushInfo = Brush->ApplyBrush(MousePositions);
+				FLandscapeBrushData BrushInfo = Brush->ApplyBrush(InteractorPositions);
 				if (!BrushInfo)
 				{
 					return;
@@ -1174,7 +1174,7 @@ public:
 		float Data;
 	};
 
-	void Apply(FEditorViewportClient* ViewportClient, FLandscapeBrush* Brush, const ULandscapeEditorObject* UISettings, const TArray<FLandscapeToolMousePosition>& MousePositions)
+	void Apply(FEditorViewportClient* ViewportClient, FLandscapeBrush* Brush, const ULandscapeEditorObject* UISettings, const TArray<FLandscapeToolInteractorPosition>& InteractorPositions)
 	{
 		//ULandscapeInfo* LandscapeInfo = EdMode->CurrentToolTarget.LandscapeInfo;
 		ALandscapeGizmoActiveActor* Gizmo = EdMode->CurrentGizmoActor.Get();
@@ -1184,7 +1184,7 @@ public:
 
 			// Get list of verts to update
 			// TODO - only retrieve bounds as we don't need the data
-			FLandscapeBrushData BrushInfo = Brush->ApplyBrush(MousePositions);
+			FLandscapeBrushData BrushInfo = Brush->ApplyBrush(InteractorPositions);
 			if (!BrushInfo)
 			{
 				return;
@@ -1488,7 +1488,7 @@ public:
 		return ELandscapeToolTargetTypeMask::FromType(ToolTarget::TargetType);
 	}
 
-	virtual bool BeginTool(FEditorViewportClient* ViewportClient, const FLandscapeToolTarget& InTarget, const FVector& InHitLocation) override
+	virtual bool BeginTool(FEditorViewportClient* ViewportClient, const FLandscapeToolTarget& InTarget, const FVector& InHitLocation, const UViewportInteractor* Interactor = nullptr) override
 	{
 		this->EdMode->GizmoBrush->Tick(ViewportClient, 0.1f);
 
@@ -1523,7 +1523,7 @@ public:
 	{
 	}
 
-	void Apply(FEditorViewportClient* ViewportClient, FLandscapeBrush* Brush, const ULandscapeEditorObject* UISettings, const TArray<FLandscapeToolMousePosition>& MousePositions)
+	void Apply(FEditorViewportClient* ViewportClient, FLandscapeBrush* Brush, const ULandscapeEditorObject* UISettings, const TArray<FLandscapeToolInteractorPosition>& InteractorPositions)
 	{
 		//ULandscapeInfo* LandscapeInfo = EdMode->CurrentToolTarget.LandscapeInfo;
 		ALandscapeGizmoActiveActor* Gizmo = EdMode->CurrentGizmoActor.Get();
@@ -1563,7 +1563,7 @@ public:
 			//LandscapeInfo->Modify();
 
 			// Get list of verts to update
-			FLandscapeBrushData BrushInfo = Brush->ApplyBrush(MousePositions);
+			FLandscapeBrushData BrushInfo = Brush->ApplyBrush(InteractorPositions);
 			if (!BrushInfo)
 			{
 				return;
@@ -1809,7 +1809,7 @@ public:
 		return ELandscapeToolTargetTypeMask::FromType(ToolTarget::TargetType);
 	}
 
-	virtual bool BeginTool(FEditorViewportClient* ViewportClient, const FLandscapeToolTarget& InTarget, const FVector& InHitLocation) override
+	virtual bool BeginTool(FEditorViewportClient* ViewportClient, const FLandscapeToolTarget& InTarget, const FVector& InHitLocation, const UViewportInteractor* Interactor = nullptr) override
 	{
 		this->EdMode->GizmoBrush->Tick(ViewportClient, 0.1f);
 
@@ -1955,7 +1955,7 @@ public:
 		EdMode->NewLandscapePreviewMode = ENewLandscapePreviewMode::None;
 	}
 
-	virtual bool BeginTool(FEditorViewportClient* ViewportClient, const FLandscapeToolTarget& Target, const FVector& InHitLocation) override
+	virtual bool BeginTool(FEditorViewportClient* ViewportClient, const FLandscapeToolTarget& Target, const FVector& InHitLocation, const UViewportInteractor* Interactor = nullptr) override
 	{
 		// does nothing
 		return false;
@@ -2019,7 +2019,7 @@ public:
 	{
 	}
 
-	virtual bool BeginTool(FEditorViewportClient* ViewportClient, const FLandscapeToolTarget& Target, const FVector& InHitLocation) override
+	virtual bool BeginTool(FEditorViewportClient* ViewportClient, const FLandscapeToolTarget& Target, const FVector& InHitLocation, const UViewportInteractor* Interactor = nullptr) override
 	{
 		// does nothing
 		return false;
diff --git a/Engine/Source/Editor/LandscapeEditor/Private/LandscapeEdModeErosionTools.cpp b/Engine/Source/Editor/LandscapeEditor/Private/LandscapeEdModeErosionTools.cpp
--- a/Engine/Source/Editor/LandscapeEditor/Private/LandscapeEdModeErosionTools.cpp
+++ b/Engine/Source/Editor/LandscapeEditor/Private/LandscapeEdModeErosionTools.cpp
@@ -61,7 +61,7 @@ public:
 	{
 	}
 
-	void Apply(FEditorViewportClient* ViewportClient, FLandscapeBrush* Brush, const ULandscapeEditorObject* UISettings, const TArray<FLandscapeToolMousePosition>& MousePositions)
+	void Apply(FEditorViewportClient* ViewportClient, FLandscapeBrush* Brush, const ULandscapeEditorObject* UISettings, const TArray<FLandscapeToolInteractorPosition>& InteractorPositions)
 	{
 		if (!LandscapeInfo)
 		{
@@ -69,7 +69,7 @@ public:
 		}
 
 		// Get list of verts to update
-		FLandscapeBrushData BrushInfo = Brush->ApplyBrush(MousePositions);
+		FLandscapeBrushData BrushInfo = Brush->ApplyBrush(InteractorPositions);
 		if (!BrushInfo)
 		{
 			return;
@@ -293,7 +293,7 @@ public:
 	{
 	}
 
-	void Apply(FEditorViewportClient* ViewportClient, FLandscapeBrush* Brush, const ULandscapeEditorObject* UISettings, const TArray<FLandscapeToolMousePosition>& MousePositions)
+	void Apply(FEditorViewportClient* ViewportClient, FLandscapeBrush* Brush, const ULandscapeEditorObject* UISettings, const TArray<FLandscapeToolInteractorPosition>& InteractorPositions)
 	{
 		if (!LandscapeInfo)
 		{
@@ -301,7 +301,7 @@ public:
 		}
 
 		// Get list of verts to update
-		FLandscapeBrushData BrushInfo = Brush->ApplyBrush(MousePositions);
+		FLandscapeBrushData BrushInfo = Brush->ApplyBrush(InteractorPositions);
 		if (!BrushInfo)
 		{
 			return;
diff --git a/Engine/Source/Editor/LandscapeEditor/Private/LandscapeEdModeMirrorTool.cpp b/Engine/Source/Editor/LandscapeEditor/Private/LandscapeEdModeMirrorTool.cpp
--- a/Engine/Source/Editor/LandscapeEditor/Private/LandscapeEdModeMirrorTool.cpp
+++ b/Engine/Source/Editor/LandscapeEditor/Private/LandscapeEdModeMirrorTool.cpp
@@ -62,7 +62,7 @@ public:
 		GLevelEditorModeTools().SetCoordSystem(SavedCoordSystem);
 	}
 
-	virtual bool BeginTool(FEditorViewportClient* ViewportClient, const FLandscapeToolTarget& Target, const FVector& InHitLocation) override
+	virtual bool BeginTool(FEditorViewportClient* ViewportClient, const FLandscapeToolTarget& Target, const FVector& InHitLocation, const UViewportInteractor* Interactor = nullptr) override
 	{
 		return true;
 	}
diff --git a/Engine/Source/Editor/LandscapeEditor/Private/LandscapeEdModePaintTools.cpp b/Engine/Source/Editor/LandscapeEditor/Private/LandscapeEdModePaintTools.cpp
--- a/Engine/Source/Editor/LandscapeEditor/Private/LandscapeEdModePaintTools.cpp
+++ b/Engine/Source/Editor/LandscapeEditor/Private/LandscapeEdModePaintTools.cpp
@@ -93,17 +93,17 @@ public:
 	{
 	}
 
-	void Apply(FEditorViewportClient* ViewportClient, FLandscapeBrush* Brush, const ULandscapeEditorObject* UISettings, const TArray<FLandscapeToolMousePosition>& MousePositions)
+	void Apply(FEditorViewportClient* ViewportClient, FLandscapeBrush* Brush, const ULandscapeEditorObject* UISettings, const TArray<FLandscapeToolInteractorPosition>& InteractorPositions)
 	{
 		// Invert when holding Shift
 		//UE_LOG(LogLandscape, Log, TEXT("bInvert = %d"), bInvert);
-		bool bInvert = MousePositions.Last().bShiftDown;
+		bool bInvert = InteractorPositions.Last().bModifierPressed;
 
 		if (bIsWhitelistMode)
 		{
 			// Get list of components to delete from brush
 			// TODO - only retrieve bounds as we don't need the vert data
-			FLandscapeBrushData BrushInfo = Brush->ApplyBrush(MousePositions);
+			FLandscapeBrushData BrushInfo = Brush->ApplyBrush(InteractorPositions);
 			if (!BrushInfo)
 			{
 				return;
@@ -142,7 +142,7 @@ public:
 		}
 
 		// Get list of verts to update
-		FLandscapeBrushData BrushInfo = Brush->ApplyBrush(MousePositions);
+		FLandscapeBrushData BrushInfo = Brush->ApplyBrush(InteractorPositions);
 		if (!BrushInfo)
 		{
 			return;
@@ -289,14 +289,14 @@ public:
 	{
 	}
 
-	void Apply(FEditorViewportClient* ViewportClient, FLandscapeBrush* Brush, const ULandscapeEditorObject* UISettings, const TArray<FLandscapeToolMousePosition>& MousePositions)
+	void Apply(FEditorViewportClient* ViewportClient, FLandscapeBrush* Brush, const ULandscapeEditorObject* UISettings, const TArray<FLandscapeToolInteractorPosition>& InteractorPositions)
 	{
 		// Invert when holding Shift
 		//UE_LOG(LogLandscape, Log, TEXT("bInvert = %d"), bInvert);
-		bool bInvert = MousePositions.Last().bShiftDown;
+		bool bInvert = InteractorPositions.Last().bModifierPressed;
 
 		// Get list of verts to update
-		FLandscapeBrushData BrushInfo = Brush->ApplyBrush(MousePositions);
+		FLandscapeBrushData BrushInfo = Brush->ApplyBrush(InteractorPositions);
 		if (!BrushInfo)
 		{
 			return;
@@ -513,12 +513,12 @@ public:
 	{
 	}
 
-	void Apply(FEditorViewportClient* ViewportClient, FLandscapeBrush* Brush, const ULandscapeEditorObject* UISettings, const TArray<FLandscapeToolMousePosition>& MousePositions)
+	void Apply(FEditorViewportClient* ViewportClient, FLandscapeBrush* Brush, const ULandscapeEditorObject* UISettings, const TArray<FLandscapeToolInteractorPosition>& InteractorPositions)
 	{
 		if (!this->LandscapeInfo) return;
 
 		// Get list of verts to update
-		FLandscapeBrushData BrushInfo = Brush->ApplyBrush(MousePositions);
+		FLandscapeBrushData BrushInfo = Brush->ApplyBrush(InteractorPositions);
 		if (!BrushInfo)
 		{
 			return;
@@ -656,15 +656,15 @@ public:
 		}
 	}
 
-	void Apply(FEditorViewportClient* ViewportClient, FLandscapeBrush* Brush, const ULandscapeEditorObject* UISettings, const TArray<FLandscapeToolMousePosition>& MousePositions)
+	void Apply(FEditorViewportClient* ViewportClient, FLandscapeBrush* Brush, const ULandscapeEditorObject* UISettings, const TArray<FLandscapeToolInteractorPosition>& InteractorPositions)
 	{
 		if (!this->LandscapeInfo) return;
 
 		if (!bInitializedFlattenHeight || (UISettings->bPickValuePerApply && bTargetIsHeightmap))
 		{
 			bInitializedFlattenHeight = false;
-			float FlattenX = MousePositions[0].Position.X;
-			float FlattenY = MousePositions[0].Position.Y;
+			float FlattenX = InteractorPositions[0].Position.X;
+			float FlattenY = InteractorPositions[0].Position.Y;
 			int32 FlattenHeightX = FMath::FloorToInt(FlattenX);
 			int32 FlattenHeightY = FMath::FloorToInt(FlattenY);
 
@@ -683,7 +683,7 @@ public:
 
 
 		// Get list of verts to update
-		FLandscapeBrushData BrushInfo = Brush->ApplyBrush(MousePositions);
+		FLandscapeBrushData BrushInfo = Brush->ApplyBrush(InteractorPositions);
 		if (!BrushInfo)
 		{
 			return;
@@ -876,12 +876,12 @@ public:
 	{
 	}
 
-	void Apply(FEditorViewportClient* ViewportClient, FLandscapeBrush* Brush, const ULandscapeEditorObject* UISettings, const TArray<FLandscapeToolMousePosition>& MousePositions)
+	void Apply(FEditorViewportClient* ViewportClient, FLandscapeBrush* Brush, const ULandscapeEditorObject* UISettings, const TArray<FLandscapeToolInteractorPosition>& InteractorPositions)
 	{
 		if (!this->LandscapeInfo) return;
 
 		// Get list of verts to update
-		FLandscapeBrushData BrushInfo = Brush->ApplyBrush(MousePositions);
+		FLandscapeBrushData BrushInfo = Brush->ApplyBrush(InteractorPositions);
 		if (!BrushInfo)
 		{
 			return;
diff --git a/Engine/Source/Editor/LandscapeEditor/Private/LandscapeEdModeRampTool.cpp b/Engine/Source/Editor/LandscapeEditor/Private/LandscapeEdModeRampTool.cpp
--- a/Engine/Source/Editor/LandscapeEditor/Private/LandscapeEdModeRampTool.cpp
+++ b/Engine/Source/Editor/LandscapeEditor/Private/LandscapeEdModeRampTool.cpp
@@ -123,7 +123,7 @@ public:
 		GLevelEditorModeTools().SetWidgetMode(FWidget::WM_Translate);
 	}
 
-	virtual bool BeginTool(FEditorViewportClient* ViewportClient, const FLandscapeToolTarget& Target, const FVector& InHitLocation) override
+	virtual bool BeginTool(FEditorViewportClient* ViewportClient, const FLandscapeToolTarget& Target, const FVector& InHitLocation, const UViewportInteractor* Interactor = nullptr) override
 	{
 		if (NumPoints < 2)
 		{
diff --git a/Engine/Source/Editor/LandscapeEditor/Private/LandscapeEdModeSplineTools.cpp b/Engine/Source/Editor/LandscapeEditor/Private/LandscapeEdModeSplineTools.cpp
--- a/Engine/Source/Editor/LandscapeEditor/Private/LandscapeEdModeSplineTools.cpp
+++ b/Engine/Source/Editor/LandscapeEditor/Private/LandscapeEdModeSplineTools.cpp
@@ -837,7 +837,7 @@ public:
 		}
 	}
 
-	virtual bool BeginTool(FEditorViewportClient* ViewportClient, const FLandscapeToolTarget& InTarget, const FVector& InHitLocation) override
+	virtual bool BeginTool(FEditorViewportClient* ViewportClient, const FLandscapeToolTarget& InTarget, const FVector& InHitLocation, const UViewportInteractor* Interactor = nullptr) override
 	{
 		if (ViewportClient->IsCtrlPressed())
 		{
diff --git a/Engine/Source/Editor/LandscapeEditor/Private/LandscapeEdModeTools.h b/Engine/Source/Editor/LandscapeEditor/Private/LandscapeEdModeTools.h
--- a/Engine/Source/Editor/LandscapeEditor/Private/LandscapeEdModeTools.h
+++ b/Engine/Source/Editor/LandscapeEditor/Private/LandscapeEdModeTools.h
@@ -7,6 +7,11 @@
 #include "InstancedFoliageActor.h"
 #include "AI/Navigation/NavigationSystem.h"
 
+// VR Editor
+#include "IVREditorModule.h"
+#include "ViewportWorldInteraction.h"
+#include "VREditorInteractor.h"
+
 //
 //	FNoiseParameter - Perlin noise
 //
@@ -1303,11 +1308,11 @@ public:
 	{
 	}
 
-	virtual bool BeginTool(FEditorViewportClient* ViewportClient, const FLandscapeToolTarget& InTarget, const FVector& InHitLocation) override
+	virtual bool BeginTool(FEditorViewportClient* ViewportClient, const FLandscapeToolTarget& InTarget, const FVector& InHitLocation, const UViewportInteractor* Interactor = nullptr) override
 	{
-		if (!ensure(MousePositions.Num() == 0))
+		if (!ensure(InteractorPositions.Num() == 0))
 		{
-			MousePositions.Empty(1);
+			InteractorPositions.Empty(1);
 		}
 
 		bToolActive = true;
@@ -1316,13 +1321,13 @@ public:
 		EdMode->CurrentBrush->BeginStroke(InHitLocation.X, InHitLocation.Y, this);
 
 		// Save the mouse position
-		LastMousePosition = FVector2D(InHitLocation);
-		MousePositions.Emplace(LastMousePosition, ViewportClient ? IsShiftDown(ViewportClient->Viewport) : false); // Copy tool sometimes activates without a specific viewport via ctrl+c hotkey
-		TimeSinceLastMouseMove = 0.0f;
+		LastInteractorPosition = FVector2D(InHitLocation);
+		InteractorPositions.Emplace(LastInteractorPosition, ViewportClient ? IsModifierPressed(ViewportClient, Interactor) : false); // Copy tool sometimes activates without a specific viewport via ctrl+c hotkey
+		TimeSinceLastInteractorMove = 0.0f;
 
-		ToolStroke->Apply(ViewportClient, EdMode->CurrentBrush, EdMode->UISettings, MousePositions);
+		ToolStroke->Apply(ViewportClient, EdMode->CurrentBrush, EdMode->UISettings, InteractorPositions);
 
-		MousePositions.Empty(1);
+		InteractorPositions.Empty(1);
 		return true;
 	}
 
@@ -1330,20 +1335,20 @@ public:
 	{
 		if (bToolActive)
 		{
-			if (MousePositions.Num() > 0)
+			if (InteractorPositions.Num() > 0)
 			{
-				ToolStroke->Apply(ViewportClient, EdMode->CurrentBrush, EdMode->UISettings, MousePositions);
+				ToolStroke->Apply(ViewportClient, EdMode->CurrentBrush, EdMode->UISettings, InteractorPositions);
 				ViewportClient->Invalidate(false, false);
-				MousePositions.Empty(1);
+				InteractorPositions.Empty(1);
 			}
-			else if (TStrokeClass::UseContinuousApply && TimeSinceLastMouseMove >= 0.25f)
+			else if (TStrokeClass::UseContinuousApply && TimeSinceLastInteractorMove >= 0.25f)
 			{
-				MousePositions.Emplace(LastMousePosition, IsShiftDown(ViewportClient->Viewport));
-				ToolStroke->Apply(ViewportClient, EdMode->CurrentBrush, EdMode->UISettings, MousePositions);
+				InteractorPositions.Emplace(LastInteractorPosition, IsModifierPressed(ViewportClient));
+				ToolStroke->Apply(ViewportClient, EdMode->CurrentBrush, EdMode->UISettings, InteractorPositions);
 				ViewportClient->Invalidate(false, false);
-				MousePositions.Empty(1);
+				InteractorPositions.Empty(1);
 			}
-			TimeSinceLastMouseMove += DeltaTime;
+			TimeSinceLastInteractorMove += DeltaTime;
 
 			// Prevent landscape from baking textures while tool stroke is active
 			EdMode->CurrentToolTarget.LandscapeInfo->PostponeTextureBaking();
@@ -1352,10 +1357,10 @@ public:
 
 	virtual void EndTool(FEditorViewportClient* ViewportClient) override
 	{
-		if (bToolActive && MousePositions.Num())
+		if (bToolActive && InteractorPositions.Num())
 		{
-			ToolStroke->Apply(ViewportClient, EdMode->CurrentBrush, EdMode->UISettings, MousePositions);
-			MousePositions.Empty(1);
+			ToolStroke->Apply(ViewportClient, EdMode->CurrentBrush, EdMode->UISettings, InteractorPositions);
+			InteractorPositions.Empty(1);
 		}
 
 		ToolStroke.Reset();
@@ -1376,13 +1381,13 @@ public:
 
 			if (bToolActive)
 			{
-				// Save the mouse position
-				if (MousePositions.Num() == 0 || LastMousePosition != FVector2D(HitLocation))
+				// Save the interactor position
+				if (InteractorPositions.Num() == 0 || LastInteractorPosition != FVector2D(HitLocation))
 				{
-					LastMousePosition = FVector2D(HitLocation);
-					MousePositions.Emplace(LastMousePosition, IsShiftDown(ViewportClient->Viewport));
+					LastInteractorPosition = FVector2D(HitLocation);
+					InteractorPositions.Emplace(LastInteractorPosition, IsModifierPressed(ViewportClient));
 				}
-				TimeSinceLastMouseMove = 0.0f;
+				TimeSinceLastInteractorMove = 0.0f;
 			}
 		}
 
@@ -1390,10 +1395,22 @@ public:
 	}
 
 protected:
-	TArray<FLandscapeToolMousePosition> MousePositions;
-	FVector2D LastMousePosition;
-	float TimeSinceLastMouseMove;
+	TArray<FLandscapeToolInteractorPosition> InteractorPositions;
+	FVector2D LastInteractorPosition;
+	float TimeSinceLastInteractorMove;
 	FEdModeLandscape* EdMode;
 	bool bToolActive;
 	TOptional<TStrokeClass> ToolStroke;
+
+	bool IsModifierPressed(const class FEditorViewportClient* ViewportClient, const UViewportInteractor* Interactor = nullptr)
+	{
+		bool bIsInteractorModifierPressed = false;
+		if (Interactor != nullptr)
+		{
+			const UVREditorInteractor* VRInteractor = Cast<UVREditorInteractor>(Interactor);
+			bIsInteractorModifierPressed = VRInteractor->IsModifierPressed();
+		}
+
+		return IsShiftDown(ViewportClient->Viewport) || bIsInteractorModifierPressed;
+	}
 };
diff --git a/Engine/Source/Editor/LandscapeEditor/Private/LandscapeEdModeXYOffset.cpp b/Engine/Source/Editor/LandscapeEditor/Private/LandscapeEdModeXYOffset.cpp
--- a/Engine/Source/Editor/LandscapeEditor/Private/LandscapeEdModeXYOffset.cpp
+++ b/Engine/Source/Editor/LandscapeEditor/Private/LandscapeEdModeXYOffset.cpp
@@ -59,7 +59,7 @@ public:
 	{
 	}
 
-	void Apply(FEditorViewportClient* ViewportClient, FLandscapeBrush* Brush, const ULandscapeEditorObject* UISettings, const TArray<FLandscapeToolMousePosition>& MousePositions)
+	void Apply(FEditorViewportClient* ViewportClient, FLandscapeBrush* Brush, const ULandscapeEditorObject* UISettings, const TArray<FLandscapeToolInteractorPosition>& InteractorPositions)
 	{
 		if (!LandscapeInfo)
 		{
@@ -69,7 +69,7 @@ public:
 		// Get list of verts to update
 		// TODO - only retrieve bounds as we don't need the data
 		// or use the brush data?
-		FLandscapeBrushData BrushInfo = Brush->ApplyBrush(MousePositions);
+		FLandscapeBrushData BrushInfo = Brush->ApplyBrush(InteractorPositions);
 		if (!BrushInfo)
 		{
 			return;
diff --git a/Engine/Source/Editor/LandscapeEditor/Public/LandscapeToolInterface.h b/Engine/Source/Editor/LandscapeEditor/Public/LandscapeToolInterface.h
--- a/Engine/Source/Editor/LandscapeEditor/Public/LandscapeToolInterface.h
+++ b/Engine/Source/Editor/LandscapeEditor/Public/LandscapeToolInterface.h
@@ -5,17 +5,18 @@
 // Forward declarations
 class ULandscapeInfo;
 class ULandscapeLayerInfoObject;
+class UViewportInteractor;
 
 // FLandscapeToolMousePosition - Struct to store mouse positions since the last time we applied the brush
-struct FLandscapeToolMousePosition
+struct FLandscapeToolInteractorPosition
 {
 	// Stored in heightmap space.
 	FVector2D Position;
-	bool bShiftDown;
+	bool bModifierPressed;
 
-	FLandscapeToolMousePosition(FVector2D InPosition, bool InbShiftDown)
+	FLandscapeToolInteractorPosition(FVector2D InPosition, const bool bInModifierPressed)
 		: Position(InPosition)
-		, bShiftDown(InbShiftDown)
+		, bModifierPressed(bInModifierPressed)
 	{
 	}
 };
@@ -85,7 +86,7 @@ class FLandscapeBrush : public FGCObject
 {
 public:
 	virtual void MouseMove(float LandscapeX, float LandscapeY) = 0;
-	virtual FLandscapeBrushData ApplyBrush(const TArray<FLandscapeToolMousePosition>& MousePositions) = 0;
+	virtual FLandscapeBrushData ApplyBrush(const TArray<FLandscapeToolInteractorPosition>& InteractorPositions) = 0;
 	virtual TOptional<bool> InputKey(FEditorViewportClient* InViewportClient, FViewport* InViewport, FKey InKey, EInputEvent InEvent) { return TOptional<bool>(); }
 	virtual void Tick(FEditorViewportClient* ViewportClient, float DeltaTime) {};
 	virtual void BeginStroke(float LandscapeX, float LandscapeY, class FLandscapeTool* CurrentTool);
@@ -214,7 +215,7 @@ class FLandscapeTool : public FGCObject
 public:
 	virtual void EnterTool() {}
 	virtual void ExitTool() {}
-	virtual bool BeginTool(FEditorViewportClient* ViewportClient, const FLandscapeToolTarget& Target, const FVector& InHitLocation) = 0;
+	virtual bool BeginTool(FEditorViewportClient* ViewportClient, const FLandscapeToolTarget& Target, const FVector& InHitLocation, const UViewportInteractor* Interactor = nullptr) = 0;
 	virtual void EndTool(FEditorViewportClient* ViewportClient) = 0;
 	virtual void Tick(FEditorViewportClient* ViewportClient, float DeltaTime) {};
 	virtual bool MouseMove(FEditorViewportClient* ViewportClient, FViewport* Viewport, int32 x, int32 y) = 0;
diff --git a/Engine/Source/Runtime/Slate/Private/Framework/MultiBox/SToolBarButtonBlock.cpp b/Engine/Source/Runtime/Slate/Private/Framework/MultiBox/SToolBarButtonBlock.cpp
--- a/Engine/Source/Runtime/Slate/Private/Framework/MultiBox/SToolBarButtonBlock.cpp
+++ b/Engine/Source/Runtime/Slate/Private/Framework/MultiBox/SToolBarButtonBlock.cpp
@@ -309,8 +309,7 @@ FReply SToolBarButtonBlock::OnClicked()
 	const bool ClosingMenu = MultiBox->ShouldCloseWindowAfterMenuSelection();
 	if( ClosingMenu )
 	{
-		TSharedRef<SWindow> ParentContextMenuWindow = FSlateApplication::Get().FindWidgetWindow( AsShared() ).ToSharedRef();
-		FSlateApplication::Get().RequestDestroyWindow( ParentContextMenuWindow );
+		FSlateApplication::Get().DismissMenuByWidget(AsShared());
 	}
 
 	return FReply::Handled();
