Copying //UE4/Dev-Landscape to //UE4/Dev-Main (Source: //UE4/Dev-Landscape @ 3053591)

==========================
MAJOR FEATURES + CHANGES
==========================

Change 3016060 on 2016/06/16 by [REDACTED_PERSON]

	Fixed check being hit in AActor::InvalidateLightingCacheDetailed when building lighting with grass (cloned from 4.12)

Change 3016346 on 2016/06/16 by [REDACTED_PERSON]

	Fixed crash when editing landscape component properties and then clicking in the editor viewport (cloned from 4.12)

Change 3016366 on 2016/06/16 by [REDACTED_PERSON]

	Refactor Landscape Select / Mask tools so that they don't share a common brush stroke class
	- they had no common code, there was just a massive if statement in there selecting between the two behaviours

Change 3017787 on 2016/06/17 by [REDACTED_PERSON]

	Added Copy/Transform Algos

Change 3019479 on 2016/06/20 by [REDACTED_PERSON]

	Fix for user-reported crash in landscape circle brush
	(https://udn.unrealengine.com/questions/298966/crash-in-flandscapebrushcircletick.html)

Change 3020927 on 2016/06/21 by [REDACTED_PERSON]

	Fixed "Obj in another map" error when packaging a project containing Landscape for Android

Change 3021318 on 2016/06/21 by [REDACTED_PERSON]

	Fixed crash in PIE while using Landscape Splines with bPlaceSplineMeshesInStreamingLevels

Change 3024769 on 2016/06/23 by [REDACTED_PERSON]

	Allow creation of transient objects during saving, as they wouldn't be saved anyway
	(and the slate font system did during autosave and hit this check)

Change 3026457 on 2016/06/24 by [REDACTED_PERSON]

	Fixed crash when setting landscape component collision mip to an invalid value while in collision viewmode (cloned from 4.12)

Change 3032336 on 2016/06/29 by [REDACTED_PERSON]

	Fixed a long-standing issue with undo/redo nulling lazy pointers
	- Two objects are created which reference each other with lazy object pointers (e.g. ULandscapeComponent and ULandscapeHeightfieldCollisionComponent)
	- Undo
	- During undo, the first object saves its state and restores itself to a deleted (pendingkill) state
	- the second object then tries to save its state, but as the other object is marked pending kill it fails to resolve its lazy ptr and saves null instead
	- Redo
	- Lazy ptr is restored to the saved null instead of the actual original value

Change 3032818 on 2016/06/29 by [REDACTED_PERSON]

	Removed all use of deferred exec commands from the landscape code
	This solves numerous crashes during commandlets etc with landscape infos not being initialized because they don't tick the engine

	Notes:
	Split CreateLandscapeInfo from GetLandscapeInfo (removing the default-true bool param from GetLandscapeInfo in the process), so it only gets created where intended
	Removed some unused properties/functions, e.g. ULandscapeInfo::bIsValid, bCurrentlyEditing
	Removed legacy landscape world composition code
	General cleanup

Change 3032820 on 2016/06/29 by [REDACTED_PERSON]

	File missing from CL 3032818 (P4! Grr)

Change 3037033 on 2016/07/04 by [REDACTED_PERSON]

	Capitalised variable names in Algo::Accumulate at Core's request

Change 3037035 on 2016/07/04 by [REDACTED_PERSON]

	Improved Algo::Accumulate when used with movable types (e.g. FString)

Change 3037373 on 2016/07/05 by [REDACTED_PERSON]

	Added FText property handle type

Change 3039183 on 2016/07/06 by [REDACTED_PERSON]

	Landscape data plugin API - part 1: Import

Change 3042473 on 2016/07/08 by [REDACTED_PERSON]

	Landscape data plugin API - part 2: Re-Import

Change 3042484 on 2016/07/08 by [REDACTED_PERSON]

	Pull request from AndrewScheidecker https://github.com/EpicGames/UnrealEngine/pull/1279 and updated to 4.13.

	Adapted GetStaticBatchElementVisibility to support changing visibility of batches with a single element (for view-dependent culling of static mesh elements)
	Added GetStaticBatchElementShadowVisibility that allows culling static mesh elements when rendering shadow depths using information about the light

	** Any user-created vertex factory that implements GetStaticBatchElementVisibility will need to set bRequiresPerElementVisibility on their FMeshBatch.

	#1279

Change 3044665 on 2016/07/11 by [REDACTED_PERSON]

	Renamed Engine/FixedSizeArrayView to Core/Containers/ArrayView
	Note: The new and improved TArrayView (as agreed with core team) is *not* const-propogating. const TArrayView<T> does not imply "const T" and will not bind to a const TArray! You want TArrayView<const T>

Change 3045936 on 2016/07/12 by [REDACTED_PERSON]

	Worked around Visual Studio 2013 bug with the new TArrayView

Change 3045964 on 2016/07/12 by [REDACTED_PERSON]

	Fix "cannot find shader" crash when using landscape with XY offsets

Change 3046311 on 2016/07/12 by [REDACTED_PERSON]

	Landscape data plugin API - part 3: Tiled landscapes

Change 3047701 on 2016/07/13 by [REDACTED_PERSON]

	Landscape data plugin API - cleanup

Change 3047894 on 2016/07/13 by [REDACTED_PERSON]

	Landscape data plugin API - part 4: Support inherent scale in landscape heightmap formats

Change 3049656 on 2016/07/14 by [REDACTED_PERSON]

	Landscape data plugin API - Fix right-click layer import option to show plugin filetypes in the file selector

Change 3050073 on 2016/07/14 by [REDACTED_PERSON]

	Update sequence recorder to new TArrayView following merge from main

Change 3050096 on 2016/07/14 by [REDACTED_PERSON]

	Landscape data plugin API - part 5: Export!

Change 3051521 on 2016/07/15 by [REDACTED_PERSON]

	Landscape data plugin API - part 6: Finishing up

Change 3051783 on 2016/07/15 by [REDACTED_PERSON]

	Fixed the landscape info map not being transacted correctly, resulting in a crash if you undo the deletion of a landscape

Change 3053322 on 2016/07/18 by [REDACTED_PERSON]

	Fix landscapes not being registered correctly with their landscape info when hiding/showing them in world composition

Change 3053513 on 2016/07/18 by [REDACTED_PERSON]

	Added LANDSCAPE_API to Landscape MaterialExpression classes

	#1927

Change 3053591 on 2016/07/18 by [REDACTED_PERSON]

	Attempt to fix VS 2013 ICE

diff --git a/Engine/Source/Editor/LandscapeEditor/Private/LandscapeEdMode.h b/Engine/Source/Editor/LandscapeEditor/Private/LandscapeEdMode.h
--- a/Engine/Source/Editor/LandscapeEditor/Private/LandscapeEdMode.h
+++ b/Engine/Source/Editor/LandscapeEditor/Private/LandscapeEdMode.h
@@ -12,6 +12,7 @@
 DECLARE_LOG_CATEGORY_EXTERN(LogLandscapeEdMode, Log, All);
 
 // Forward declarations
+class ULandscapeEditorObject;
 class ULandscapeLayerInfoObject;
 class FLandscapeToolSplines;
 
@@ -218,7 +219,7 @@ class FEdModeLandscape : public FEdMode
 {
 public:
 
-	class ULandscapeEditorObject* UISettings;
+	ULandscapeEditorObject* UISettings;
 
 	FLandscapeToolMode* CurrentToolMode;
 	FLandscapeTool* CurrentTool;
diff --git a/Engine/Source/Editor/LandscapeEditor/Private/LandscapeEdModeBrushes.cpp b/Engine/Source/Editor/LandscapeEditor/Private/LandscapeEdModeBrushes.cpp
--- a/Engine/Source/Editor/LandscapeEditor/Private/LandscapeEdModeBrushes.cpp
+++ b/Engine/Source/Editor/LandscapeEditor/Private/LandscapeEdModeBrushes.cpp
@@ -73,16 +73,34 @@ public:
 		// Allow any currently unused material instances to be GC'd
 		BrushMaterialFreeInstances.Empty();
 
+		Collector.AddReferencedObjects(BrushMaterialComponents);
 		Collector.AddReferencedObjects(BrushMaterialInstanceMap);
+
+		// If a user tool removes any components then we will have bad (null) entries in our TSet/TMap, remove them
+		// We can't just call .Remove(nullptr) because the entries were hashed as non-null values so a hash lookup of nullptr won't find them
+		for (auto It = BrushMaterialComponents.CreateIterator(); It; ++It)
+		{
+			if (*It == nullptr)
+			{
+				It.RemoveCurrent();
+			}
+		}
+		for (auto It = BrushMaterialInstanceMap.CreateIterator(); It; ++It)
+		{
+			if (It->Key == nullptr || It->Value == nullptr)
+			{
+				It.RemoveCurrent();
+			}
+		}
 	}
 
 	virtual void LeaveBrush() override
 	{
-		for (TSet<ULandscapeComponent*>::TIterator It(BrushMaterialComponents); It; ++It)
+		for (ULandscapeComponent* Component : BrushMaterialComponents)
 		{
-			if ((*It) != nullptr && (*It)->EditToolRenderData != nullptr)
+			if (Component && Component->EditToolRenderData)
 			{
-				(*It)->EditToolRenderData->Update(nullptr);
+				Component->EditToolRenderData->Update(nullptr);
 			}
 		}
 		TArray<UMaterialInstanceDynamic*> BrushMaterialInstances;
@@ -126,19 +144,11 @@ public:
 		TSet<ULandscapeComponent*> RemovedComponents = BrushMaterialComponents.Difference(NewComponents);
 		for (ULandscapeComponent* RemovedComponent : RemovedComponents)
 		{
-			if (RemovedComponent == nullptr)
-			{
-				// can occur if a component was removed by a user tool
-				BrushMaterialInstanceMap.Remove(nullptr);
-			}
-			else
-			{
-				BrushMaterialFreeInstances.Push(BrushMaterialInstanceMap.FindAndRemoveChecked(RemovedComponent));
+			BrushMaterialFreeInstances.Push(BrushMaterialInstanceMap.FindAndRemoveChecked(RemovedComponent));
 
-				if (ensure(RemovedComponent->EditToolRenderData != nullptr))
-				{
-					RemovedComponent->EditToolRenderData->Update(nullptr);
-				}
+			if (ensure(RemovedComponent->EditToolRenderData != nullptr))
+			{
+				RemovedComponent->EditToolRenderData->Update(nullptr);
 			}
 		}
 
diff --git a/Engine/Source/Editor/LandscapeEditor/Private/LandscapeEdModeComponentTools.cpp b/Engine/Source/Editor/LandscapeEditor/Private/LandscapeEdModeComponentTools.cpp
--- a/Engine/Source/Editor/LandscapeEditor/Private/LandscapeEdModeComponentTools.cpp
+++ b/Engine/Source/Editor/LandscapeEditor/Private/LandscapeEdModeComponentTools.cpp
@@ -18,155 +18,100 @@
 #include "ComponentReregisterContext.h"
 #include "PhysicalMaterials/PhysicalMaterial.h"
 #include "Materials/MaterialExpressionLandscapeVisibilityMask.h"
+#include "Algo/Copy.h"
 
 #define LOCTEXT_NAMESPACE "Landscape"
 
+//
+// FLandscapeToolSelect
+//
 class FLandscapeToolStrokeSelect : public FLandscapeToolStrokeBase
 {
 	bool bInitializedComponentInvert;
-	bool bComponentInvert;
+	bool bInvert;
+	bool bNeedsSelectionUpdate;
 
 public:
 	FLandscapeToolStrokeSelect(FEdModeLandscape* InEdMode, FEditorViewportClient* InViewportClient, const FLandscapeToolTarget& InTarget)
 		: FLandscapeToolStrokeBase(InEdMode, InViewportClient, InTarget)
 		, bInitializedComponentInvert(false)
+		, bNeedsSelectionUpdate(false)
 		, Cache(InTarget)
 	{
 	}
 
+	~FLandscapeToolStrokeSelect()
+	{
+		if (bNeedsSelectionUpdate)
+		{
+			TArray<UObject*> Objects;
+			if (LandscapeInfo)
+			{
+				TSet<ULandscapeComponent*> SelectedComponents = LandscapeInfo->GetSelectedComponents();
+				Objects.Reset(SelectedComponents.Num());
+				Algo::Copy(SelectedComponents, Objects);
+			}
+			FPropertyEditorModule& PropertyModule = FModuleManager::Get().LoadModuleChecked<FPropertyEditorModule>(TEXT("PropertyEditor"));
+			PropertyModule.UpdatePropertyViews(Objects);
+		}
+	}
+
 	void Apply(FEditorViewportClient* ViewportClient, FLandscapeBrush* Brush, const ULandscapeEditorObject* UISettings, const TArray<FLandscapeToolMousePosition>& MousePositions)
 	{
 		if (LandscapeInfo)
 		{
 			LandscapeInfo->Modify();
 
-			// Invert when holding Shift
-			bool bInvert = MousePositions[MousePositions.Num() - 1].bShiftDown;
-
-			if (Brush->GetBrushType() == ELandscapeBrushType::Component)
+			// TODO - only retrieve bounds as we don't need the data
+			const FLandscapeBrushData BrushInfo = Brush->ApplyBrush(MousePositions);
+			if (!BrushInfo)
 			{
-				// TODO - only retrieve bounds as we don't need the data
-				const FLandscapeBrushData BrushInfo = Brush->ApplyBrush(MousePositions);
-				if (!BrushInfo)
-				{
-					return;
-				}
-
-				int32 X1, Y1, X2, Y2;
-				BrushInfo.GetInclusiveBounds(X1, Y1, X2, Y2);
-
-				// Shrink bounds by 1,1 to avoid GetComponentsInRegion picking up extra components on all sides due to the overlap between components
-				TSet<ULandscapeComponent*> NewComponents;
-				LandscapeInfo->GetComponentsInRegion(X1 + 1, Y1 + 1, X2 - 1, Y2 - 1, NewComponents);
-
-				if (!bInitializedComponentInvert)
-				{
-					// Get the component under the mouse location. Copied from FLandscapeBrushComponent::ApplyBrush()
-					const float MouseX = MousePositions[0].Position.X;
-					const float MouseY = MousePositions[0].Position.Y;
-					const int32 MouseComponentIndexX = (MouseX >= 0.0f) ? FMath::FloorToInt(MouseX / LandscapeInfo->ComponentSizeQuads) : FMath::CeilToInt(MouseX / LandscapeInfo->ComponentSizeQuads);
-					const int32 MouseComponentIndexY = (MouseY >= 0.0f) ? FMath::FloorToInt(MouseY / LandscapeInfo->ComponentSizeQuads) : FMath::CeilToInt(MouseY / LandscapeInfo->ComponentSizeQuads);
-					ULandscapeComponent* MouseComponent = LandscapeInfo->XYtoComponentMap.FindRef(FIntPoint(MouseComponentIndexX, MouseComponentIndexY));
-
-					if (MouseComponent != nullptr)
-					{
-						bComponentInvert = LandscapeInfo->GetSelectedComponents().Contains(MouseComponent);
-					}
-					else
-					{
-						bComponentInvert = false;
-					}
+				return;
+			}
 
-					bInitializedComponentInvert = true;
-				}
+			int32 X1, Y1, X2, Y2;
+			BrushInfo.GetInclusiveBounds(X1, Y1, X2, Y2);
 
-				bInvert = bComponentInvert;
+			// Shrink bounds by 1,1 to avoid GetComponentsInRegion picking up extra components on all sides due to the overlap between components
+			TSet<ULandscapeComponent*> NewComponents;
+			LandscapeInfo->GetComponentsInRegion(X1 + 1, Y1 + 1, X2 - 1, Y2 - 1, NewComponents);
 
-				TSet<ULandscapeComponent*> NewSelection;
-				if (bInvert)
+			if (!bInitializedComponentInvert)
+			{
+				// Get the component under the mouse location. Copied from FLandscapeBrushComponent::ApplyBrush()
+				const float MouseX = MousePositions[0].Position.X;
+				const float MouseY = MousePositions[0].Position.Y;
+				const int32 MouseComponentIndexX = (MouseX >= 0.0f) ? FMath::FloorToInt(MouseX / LandscapeInfo->ComponentSizeQuads) : FMath::CeilToInt(MouseX / LandscapeInfo->ComponentSizeQuads);
+				const int32 MouseComponentIndexY = (MouseY >= 0.0f) ? FMath::FloorToInt(MouseY / LandscapeInfo->ComponentSizeQuads) : FMath::CeilToInt(MouseY / LandscapeInfo->ComponentSizeQuads);
+				ULandscapeComponent* MouseComponent = LandscapeInfo->XYtoComponentMap.FindRef(FIntPoint(MouseComponentIndexX, MouseComponentIndexY));
+
+				if (MouseComponent != nullptr)
 				{
-					NewSelection = LandscapeInfo->GetSelectedComponents().Difference(NewComponents);
+					bInvert = LandscapeInfo->GetSelectedComponents().Contains(MouseComponent);
 				}
 				else
 				{
-					NewSelection = LandscapeInfo->GetSelectedComponents().Union(NewComponents);
+					bInvert = false;
 				}
 
-				LandscapeInfo->Modify();
-				LandscapeInfo->UpdateSelectedComponents(NewSelection);
+				bInitializedComponentInvert = true;
+			}
 
-				// Update Details tab with selection
-				TArray<UObject*> Objects;
-				Objects.Reset(NewSelection.Num());
-				for (auto It = NewSelection.CreateConstIterator(); It; ++It)
-				{
-					Objects.Add(*It);
-				}
-				FPropertyEditorModule& PropertyModule = FModuleManager::Get().LoadModuleChecked<FPropertyEditorModule>(TEXT("PropertyEditor"));
-				PropertyModule.UpdatePropertyViews(Objects);
+			TSet<ULandscapeComponent*> NewSelection;
+			if (bInvert)
+			{
+				NewSelection = LandscapeInfo->GetSelectedComponents().Difference(NewComponents);
 			}
-			else // Select various shape regions
+			else
 			{
-				const FLandscapeBrushData BrushInfo = Brush->ApplyBrush(MousePositions);
-				if (!BrushInfo)
-				{
-					return;
-				}
-
-				int32 X1, Y1, X2, Y2;
-				BrushInfo.GetInclusiveBounds(X1, Y1, X2, Y2);
-
-				// Tablet pressure
-				float Pressure = ViewportClient->Viewport->IsPenActive() ? ViewportClient->Viewport->GetTabletPressure() : 1.0f;
-
-				Cache.CacheData(X1, Y1, X2, Y2);
-				TArray<uint8> Data;
-				Cache.GetCachedData(X1, Y1, X2, Y2, Data);
-
-				TSet<ULandscapeComponent*> NewComponents;
-				LandscapeInfo->GetComponentsInRegion(X1, Y1, X2, Y2, NewComponents);
-				LandscapeInfo->UpdateSelectedComponents(NewComponents, false);
-
-				for (int32 Y = BrushInfo.GetBounds().Min.Y; Y < BrushInfo.GetBounds().Max.Y; Y++)
-				{
-					const float* BrushScanline = BrushInfo.GetDataPtr(FIntPoint(0, Y));
-					uint8* DataScanline = Data.GetData() + (Y - Y1) * (X2 - X1 + 1) + (0 - X1);
-
-					for (int32 X = BrushInfo.GetBounds().Min.X; X < BrushInfo.GetBounds().Max.X; X++)
-					{
-						const FIntPoint Key = ALandscape::MakeKey(X, Y);
-						const float BrushValue = BrushScanline[X];
-
-						if (BrushValue > 0.0f && LandscapeInfo->IsValidPosition(X, Y))
-						{
-							float PaintValue = BrushValue * UISettings->ToolStrength * Pressure;
-							float Value = DataScanline[X] / 255.0f;
-							checkSlow(FMath::IsNearlyEqual(Value, LandscapeInfo->SelectedRegion.FindRef(Key), 1 / 255.0f));
-							if (bInvert)
-							{
-								Value = FMath::Max(Value - PaintValue, 0.0f);
-							}
-							else
-							{
-								Value = FMath::Min(Value + PaintValue, 1.0f);
-							}
-							if (Value > 0.0f)
-							{
-								LandscapeInfo->SelectedRegion.Add(Key, Value);
-							}
-							else
-							{
-								LandscapeInfo->SelectedRegion.Remove(Key);
-							}
+				NewSelection = LandscapeInfo->GetSelectedComponents().Union(NewComponents);
+			}
 
-							DataScanline[X] = FMath::Clamp<int32>(FMath::RoundToInt(Value * 255), 0, 255);
-						}
-					}
-				}
+			LandscapeInfo->Modify();
+			LandscapeInfo->UpdateSelectedComponents(NewSelection);
 
-				Cache.SetCachedData(X1, Y1, X2, Y2, Data);
-				Cache.Flush();
-			}
+			// Update Details tab with selection
+			bNeedsSelectionUpdate = true;
 		}
 	}
 
@@ -174,15 +119,11 @@ protected:
 	FLandscapeDataCache Cache;
 };
 
-// 
-// FLandscapeToolSelect
-//
-template<class TStrokeClass>
-class FLandscapeToolSelect : public FLandscapeToolBase<TStrokeClass>
+class FLandscapeToolSelect : public FLandscapeToolBase<FLandscapeToolStrokeSelect>
 {
 public:
 	FLandscapeToolSelect(FEdModeLandscape* InEdMode)
-		: FLandscapeToolBase<TStrokeClass>(InEdMode)
+		: FLandscapeToolBase<FLandscapeToolStrokeSelect>(InEdMode)
 	{
 	}
 
@@ -190,16 +131,100 @@ public:
 	virtual FText GetDisplayName() override { return NSLOCTEXT("UnrealEd", "LandscapeMode_Selection", "Component Selection"); };
 	virtual void SetEditRenderType() override { GLandscapeEditRenderMode = ELandscapeEditRenderMode::SelectComponent | (GLandscapeEditRenderMode & ELandscapeEditRenderMode::BitMaskForMask); }
 	virtual bool SupportsMask() override { return false; }
+};
 
-	virtual ELandscapeToolType GetToolType() override { return ELandscapeToolType::Mask; }
+//
+// FLandscapeToolMask
+//
+class FLandscapeToolStrokeMask : public FLandscapeToolStrokeBase
+{
+public:
+	FLandscapeToolStrokeMask(FEdModeLandscape* InEdMode, FEditorViewportClient* InViewportClient, const FLandscapeToolTarget& InTarget)
+		: FLandscapeToolStrokeBase(InEdMode, InViewportClient, InTarget)
+		, Cache(InTarget)
+	{
+	}
+
+	void Apply(FEditorViewportClient* ViewportClient, FLandscapeBrush* Brush, const ULandscapeEditorObject* UISettings, const TArray<FLandscapeToolMousePosition>& MousePositions)
+	{
+		if (LandscapeInfo)
+		{
+			LandscapeInfo->Modify();
+
+			// Invert when holding Shift
+			bool bInvert = MousePositions[MousePositions.Num() - 1].bShiftDown;
+
+			const FLandscapeBrushData BrushInfo = Brush->ApplyBrush(MousePositions);
+			if (!BrushInfo)
+			{
+				return;
+			}
+
+			int32 X1, Y1, X2, Y2;
+			BrushInfo.GetInclusiveBounds(X1, Y1, X2, Y2);
+
+			// Tablet pressure
+			float Pressure = ViewportClient->Viewport->IsPenActive() ? ViewportClient->Viewport->GetTabletPressure() : 1.0f;
+
+			Cache.CacheData(X1, Y1, X2, Y2);
+			TArray<uint8> Data;
+			Cache.GetCachedData(X1, Y1, X2, Y2, Data);
+
+			TSet<ULandscapeComponent*> NewComponents;
+			LandscapeInfo->GetComponentsInRegion(X1, Y1, X2, Y2, NewComponents);
+			LandscapeInfo->UpdateSelectedComponents(NewComponents, false);
+
+			for (int32 Y = BrushInfo.GetBounds().Min.Y; Y < BrushInfo.GetBounds().Max.Y; Y++)
+			{
+				const float* BrushScanline = BrushInfo.GetDataPtr(FIntPoint(0, Y));
+				uint8* DataScanline = Data.GetData() + (Y - Y1) * (X2 - X1 + 1) + (0 - X1);
+
+				for (int32 X = BrushInfo.GetBounds().Min.X; X < BrushInfo.GetBounds().Max.X; X++)
+				{
+					const FIntPoint Key = ALandscape::MakeKey(X, Y);
+					const float BrushValue = BrushScanline[X];
+
+					if (BrushValue > 0.0f && LandscapeInfo->IsValidPosition(X, Y))
+					{
+						float PaintValue = BrushValue * UISettings->ToolStrength * Pressure;
+						float Value = DataScanline[X] / 255.0f;
+						checkSlow(FMath::IsNearlyEqual(Value, LandscapeInfo->SelectedRegion.FindRef(Key), 1 / 255.0f));
+						if (bInvert)
+						{
+							Value = FMath::Max(Value - PaintValue, 0.0f);
+						}
+						else
+						{
+							Value = FMath::Min(Value + PaintValue, 1.0f);
+						}
+						if (Value > 0.0f)
+						{
+							LandscapeInfo->SelectedRegion.Add(Key, Value);
+						}
+						else
+						{
+							LandscapeInfo->SelectedRegion.Remove(Key);
+						}
+
+						DataScanline[X] = FMath::Clamp<int32>(FMath::RoundToInt(Value * 255), 0, 255);
+					}
+				}
+			}
+
+			Cache.SetCachedData(X1, Y1, X2, Y2, Data);
+			Cache.Flush();
+		}
+	}
+
+protected:
+	FLandscapeDataCache Cache;
 };
 
-template<class TStrokeClass>
-class FLandscapeToolMask : public FLandscapeToolSelect<TStrokeClass>
+class FLandscapeToolMask : public FLandscapeToolBase<FLandscapeToolStrokeMask>
 {
 public:
 	FLandscapeToolMask(FEdModeLandscape* InEdMode)
-		: FLandscapeToolSelect<TStrokeClass>(InEdMode)
+		: FLandscapeToolBase<FLandscapeToolStrokeMask>(InEdMode)
 	{
 	}
 
@@ -207,8 +232,13 @@ public:
 	virtual FText GetDisplayName() override { return NSLOCTEXT("UnrealEd", "LandscapeMode_Mask", "Region Selection"); };
 	virtual void SetEditRenderType() override { GLandscapeEditRenderMode = ELandscapeEditRenderMode::SelectRegion | (GLandscapeEditRenderMode & ELandscapeEditRenderMode::BitMaskForMask); }
 	virtual bool SupportsMask() override { return true; }
+
+	virtual ELandscapeToolType GetToolType() override { return ELandscapeToolType::Mask; }
 };
 
+//
+// FLandscapeToolVisibility
+//
 class FLandscapeToolStrokeVisibility : public FLandscapeToolStrokeBase
 {
 public:
@@ -269,9 +299,6 @@ protected:
 	FLandscapeVisCache Cache;
 };
 
-// 
-// FLandscapeToolVisibility
-//
 class FLandscapeToolVisibility : public FLandscapeToolBase<FLandscapeToolStrokeVisibility>
 {
 public:
@@ -310,6 +337,9 @@ public:
 	}
 };
 
+//
+// FLandscapeToolMoveToLevel
+//
 class FLandscapeToolStrokeMoveToLevel : public FLandscapeToolStrokeBase
 {
 public:
@@ -737,9 +767,6 @@ public:
 	}
 };
 
-// 
-// FLandscapeToolMoveToLevel
-//
 class FLandscapeToolMoveToLevel : public FLandscapeToolBase<FLandscapeToolStrokeMoveToLevel>
 {
 public:
@@ -755,7 +782,9 @@ public:
 	virtual bool SupportsMask() override { return false; }
 };
 
-
+//
+// FLandscapeToolAddComponent
+//
 class FLandscapeToolStrokeAddComponent : public FLandscapeToolStrokeBase
 {
 public:
@@ -933,9 +962,6 @@ protected:
 	FLandscapeXYOffsetCache<true> XYOffsetCache;
 };
 
-// 
-// FLandscapeToolAddComponent
-//
 class FLandscapeToolAddComponent : public FLandscapeToolBase<FLandscapeToolStrokeAddComponent>
 {
 public:
@@ -965,6 +991,9 @@ public:
 	}
 };
 
+//
+// FLandscapeToolDeleteComponent
+//
 class FLandscapeToolStrokeDeleteComponent : public FLandscapeToolStrokeBase
 {
 public:
@@ -1138,9 +1167,6 @@ public:
 	}
 };
 
-// 
-// FLandscapeToolDeleteComponent
-//
 class FLandscapeToolDeleteComponent : public FLandscapeToolBase<FLandscapeToolStrokeDeleteComponent>
 {
 public:
@@ -1156,6 +1182,9 @@ public:
 	virtual bool SupportsMask() override { return false; }
 };
 
+//
+// FLandscapeToolCopy
+//
 template<class ToolTarget>
 class FLandscapeToolStrokeCopy : public FLandscapeToolStrokeBase
 {
@@ -1461,9 +1490,6 @@ protected:
 	FLandscapeFullWeightCache WeightCache;
 };
 
-// 
-// FLandscapeToolCopy
-//
 template<class ToolTarget>
 class FLandscapeToolCopy : public FLandscapeToolBase<FLandscapeToolStrokeCopy<ToolTarget>>
 {
@@ -1511,6 +1537,9 @@ protected:
 	FLandscapeBrush* BackupCurrentBrush;
 };
 
+//
+// FLandscapeToolPaste
+//
 template<class ToolTarget>
 class FLandscapeToolStrokePaste : public FLandscapeToolStrokeBase
 {
@@ -1776,9 +1805,6 @@ protected:
 	FLandscapeFullWeightCache WeightCache;
 };
 
-// 
-// FLandscapeToolPaste
-//
 template<class ToolTarget>
 class FLandscapeToolPaste : public FLandscapeToolBase<FLandscapeToolStrokePaste<ToolTarget>>
 {
@@ -1850,6 +1876,9 @@ protected:
 	FLandscapeBrush* BackupCurrentBrush;
 };
 
+//
+// FLandscapeToolCopyPaste
+//
 template<class ToolTarget>
 class FLandscapeToolCopyPaste : public FLandscapeToolPaste<ToolTarget>
 {
@@ -1922,7 +1951,7 @@ void FEdModeLandscape::PasteDataFromGizmo()
 	}
 }
 
-// 
+//
 // FLandscapeToolNewLandscape
 //
 class FLandscapeToolNewLandscape : public FLandscapeTool
@@ -1974,7 +2003,7 @@ public:
 };
 
 
-// 
+//
 // FLandscapeToolResizeLandscape
 //
 class FLandscapeToolResizeLandscape : public FLandscapeTool
@@ -2055,7 +2084,7 @@ void FEdModeLandscape::InitializeTool_ResizeLandscape()
 
 void FEdModeLandscape::InitializeTool_Select()
 {
-	auto Tool_Select = MakeUnique<FLandscapeToolSelect<FLandscapeToolStrokeSelect>>(this);
+	auto Tool_Select = MakeUnique<FLandscapeToolSelect>(this);
 	Tool_Select->ValidBrushes.Add("BrushSet_Component");
 	LandscapeTools.Add(MoveTemp(Tool_Select));
 }
@@ -2083,7 +2112,7 @@ void FEdModeLandscape::InitializeTool_MoveToLevel()
 
 void FEdModeLandscape::InitializeTool_Mask()
 {
-	auto Tool_Mask = MakeUnique<FLandscapeToolMask<FLandscapeToolStrokeSelect>>(this);
+	auto Tool_Mask = MakeUnique<FLandscapeToolMask>(this);
 	Tool_Mask->ValidBrushes.Add("BrushSet_Circle");
 	Tool_Mask->ValidBrushes.Add("BrushSet_Alpha");
 	Tool_Mask->ValidBrushes.Add("BrushSet_Pattern");
diff --git a/Engine/Source/Editor/LandscapeEditor/Private/LandscapeEdModeSplineTools.cpp b/Engine/Source/Editor/LandscapeEditor/Private/LandscapeEdModeSplineTools.cpp
--- a/Engine/Source/Editor/LandscapeEditor/Private/LandscapeEdModeSplineTools.cpp
+++ b/Engine/Source/Editor/LandscapeEditor/Private/LandscapeEdModeSplineTools.cpp
@@ -19,6 +19,7 @@
 #include "ControlPointMeshComponent.h"
 #include "EditorUndoClient.h"
 #include "EngineUtils.h"
+#include "Algo/Copy.h"
 
 
 #define LOCTEXT_NAMESPACE "Landscape"
@@ -85,16 +86,8 @@ public:
 		{
 			TArray<UObject*> Objects;
 			Objects.Reset(SelectedSplineControlPoints.Num() + SelectedSplineSegments.Num());
-
-			for (ULandscapeSplineControlPoint* ControlPoint : SelectedSplineControlPoints)
-			{
-				Objects.Add(ControlPoint);
-			}
-			for (ULandscapeSplineSegment* Segment : SelectedSplineSegments)
-			{
-				Objects.Add(Segment);
-			}
-
+			Algo::Copy(SelectedSplineControlPoints, Objects);
+			Algo::Copy(SelectedSplineSegments, Objects);
 
 			FPropertyEditorModule& PropertyModule = FModuleManager::Get().LoadModuleChecked<FPropertyEditorModule>(TEXT("PropertyEditor"));
 			PropertyModule.UpdatePropertyViews(Objects);
@@ -632,7 +625,7 @@ public:
 			{
 				if (!ToLandscape)
 				{
-					ULandscapeInfo* ProxyLandscapeInfo = FromProxy->GetLandscapeInfo(false);
+					ULandscapeInfo* ProxyLandscapeInfo = FromProxy->GetLandscapeInfo();
 					check(ProxyLandscapeInfo);
 					ToLandscape = ProxyLandscapeInfo->GetCurrentLevelLandscapeProxy(true);
 					if (!ToLandscape)
@@ -722,7 +715,7 @@ public:
 			{
 				if (!ToLandscape)
 				{
-					ULandscapeInfo* ProxyLandscapeInfo = FromProxy->GetLandscapeInfo(false);
+					ULandscapeInfo* ProxyLandscapeInfo = FromProxy->GetLandscapeInfo();
 					check(ProxyLandscapeInfo);
 					ToLandscape = ProxyLandscapeInfo->GetCurrentLevelLandscapeProxy(true);
 					if (!ToLandscape)
@@ -816,15 +809,8 @@ public:
 	{
 		TArray<UObject*> Objects;
 		Objects.Reset(SelectedSplineControlPoints.Num() + SelectedSplineSegments.Num());
-
-		for (ULandscapeSplineControlPoint* ControlPoint : SelectedSplineControlPoints)
-		{
-			Objects.Add(ControlPoint);
-		}
-		for (ULandscapeSplineSegment* Segment : SelectedSplineSegments)
-		{
-			Objects.Add(Segment);
-		}
+		Algo::Copy(SelectedSplineControlPoints, Objects);
+		Algo::Copy(SelectedSplineSegments, Objects);
 
 		FPropertyEditorModule& PropertyModule = FModuleManager::Get().LoadModuleChecked<FPropertyEditorModule>(TEXT("PropertyEditor"));
 		if (!PropertyModule.HasUnlockedDetailViews())
diff --git a/Engine/Source/Editor/LandscapeEditor/Private/LandscapeEditorDetailCustomization_Base.h b/Engine/Source/Editor/LandscapeEditor/Private/LandscapeEditorDetailCustomization_Base.h
--- a/Engine/Source/Editor/LandscapeEditor/Private/LandscapeEditorDetailCustomization_Base.h
+++ b/Engine/Source/Editor/LandscapeEditor/Private/LandscapeEditorDetailCustomization_Base.h
@@ -48,7 +48,7 @@ type FLandscapeEditorDetailCustomization_Base::GetPropertyValue(TSharedRef<IProp
 	}
 
 	// Couldn't get, return null / 0
-	return (type)0;
+	return type{};
 }
 
 template<typename type>
@@ -60,7 +60,7 @@ TOptional<type> FLandscapeEditorDetailCustomization_Base::GetOptionalPropertyVal
 		return Value;
 	}
 
-	// Couldn't get, return null / 0
+	// Couldn't get, return unset optional
 	return TOptional<type>();
 }
 
@@ -74,7 +74,7 @@ type* FLandscapeEditorDetailCustomization_Base::GetObjectPropertyValue(TSharedRe
 	}
 
 	// Couldn't get, return null
-	return NULL;
+	return nullptr;
 }
 
 inline FText FLandscapeEditorDetailCustomization_Base::GetPropertyValueText(TSharedRef<IPropertyHandle> PropertyHandle)
diff --git a/Engine/Source/Editor/LandscapeEditor/Private/LandscapeEditorDetailCustomization_NewLandscape.cpp b/Engine/Source/Editor/LandscapeEditor/Private/LandscapeEditorDetailCustomization_NewLandscape.cpp
--- a/Engine/Source/Editor/LandscapeEditor/Private/LandscapeEditorDetailCustomization_NewLandscape.cpp
+++ b/Engine/Source/Editor/LandscapeEditor/Private/LandscapeEditorDetailCustomization_NewLandscape.cpp
@@ -20,18 +20,16 @@
 #include "DlgPickAssetPath.h"
 #include "SVectorInputBox.h"
 #include "SRotatorInputBox.h"
-#include "PackageTools.h"
-//#include "ObjectTools.h"
 #include "ScopedTransaction.h"
 #include "DesktopPlatformModule.h"
 #include "MainFrame.h"
 #include "AssetRegistryModule.h"
-#include "ImageWrapper.h"
 
 #include "Landscape.h"
 #include "LandscapeLayerInfoObject.h"
 #include "TutorialMetaData.h"
 #include "SNumericEntryBox.h"
+#include "LandscapeDataAccess.h"
 
 #define LOCTEXT_NAMESPACE "LandscapeEditor.NewLandscape"
 
@@ -82,43 +80,42 @@ void FLandscapeEditorDetailCustomization_NewLandscape::CustomizeDetails(IDetailL
 	];
 
 	TSharedRef<IPropertyHandle> PropertyHandle_HeightmapFilename = DetailBuilder.GetProperty(GET_MEMBER_NAME_CHECKED(ULandscapeEditorObject, ImportLandscape_HeightmapFilename));
-	TSharedRef<IPropertyHandle> PropertyHandle_HeightmapError = DetailBuilder.GetProperty(GET_MEMBER_NAME_CHECKED(ULandscapeEditorObject, ImportLandscape_HeightmapError));
-	DetailBuilder.HideProperty(PropertyHandle_HeightmapError);
+	TSharedRef<IPropertyHandle> PropertyHandle_HeightmapImportResult = DetailBuilder.GetProperty(GET_MEMBER_NAME_CHECKED(ULandscapeEditorObject, ImportLandscape_HeightmapImportResult));
+	TSharedRef<IPropertyHandle> PropertyHandle_HeightmapErrorMessage = DetailBuilder.GetProperty(GET_MEMBER_NAME_CHECKED(ULandscapeEditorObject, ImportLandscape_HeightmapErrorMessage));
+	DetailBuilder.HideProperty(PropertyHandle_HeightmapImportResult);
+	DetailBuilder.HideProperty(PropertyHandle_HeightmapErrorMessage);
 	PropertyHandle_HeightmapFilename->SetOnPropertyValueChanged(FSimpleDelegate::CreateSP(this, &FLandscapeEditorDetailCustomization_NewLandscape::OnImportHeightmapFilenameChanged));
 	NewLandscapeCategory.AddProperty(PropertyHandle_HeightmapFilename)
 	.Visibility(TAttribute<EVisibility>::Create(TAttribute<EVisibility>::FGetter::CreateStatic(&GetVisibilityOnlyInNewLandscapeMode, ENewLandscapePreviewMode::ImportLandscape)))
 	.CustomWidget()
 	.NameContent()
+	[
+		PropertyHandle_HeightmapFilename->CreatePropertyNameWidget()
+	]
+	.ValueContent()
+	.MinDesiredWidth(250.0f)
+	.MaxDesiredWidth(0)
 	[
 		SNew(SHorizontalBox)
 		+ SHorizontalBox::Slot()
 		.AutoWidth()
+		.Padding(0,0,2,0)
 		[
 			SNew(SErrorText)
-			.Visibility_Static(&GetHeightmapErrorVisibility, PropertyHandle_HeightmapError)
+			.Visibility_Static(&GetHeightmapErrorVisibility, PropertyHandle_HeightmapImportResult)
+			.BackgroundColor_Static(&GetHeightmapErrorColor, PropertyHandle_HeightmapImportResult)
 			.ErrorText(NSLOCTEXT("UnrealEd", "Error", "!"))
 			.ToolTip(
 				SNew(SToolTip)
-				.Text_Static(&GetHeightmapErrorText, PropertyHandle_HeightmapError)
+				.Text_Static(&GetPropertyValue<FText>, PropertyHandle_HeightmapErrorMessage)
 			)
 		]
 		+ SHorizontalBox::Slot()
 		.FillWidth(1)
-		[
-			PropertyHandle_HeightmapFilename->CreatePropertyNameWidget()
-		]
-	]
-	.ValueContent()
-	.MinDesiredWidth(250.0f)
-	.MaxDesiredWidth(0)
-	[
-		SNew(SHorizontalBox)
-		+ SHorizontalBox::Slot()
-		.FillWidth(1)
 		[
 			SNew(SEditableTextBox)
 			.Font(DetailBuilder.GetDetailFont())
-			.Text_Static(&GetImportHeightmapFilenameString, PropertyHandle_HeightmapFilename)
+			.Text_Static(&GetPropertyValueText, PropertyHandle_HeightmapFilename)
 			.OnTextCommitted_Static(&SetImportHeightmapFilenameString, PropertyHandle_HeightmapFilename)
 			.HintText(LOCTEXT("Import_HeightmapNotSet", "(Please specify a heightmap)"))
 		]
@@ -533,7 +530,7 @@ void FLandscapeEditorDetailCustomization_NewLandscape::SetScale(float NewValue,
 
 TSharedRef<SWidget> FLandscapeEditorDetailCustomization_NewLandscape::GetSectionSizeMenu(TSharedRef<IPropertyHandle> PropertyHandle)
 {
-	FMenuBuilder MenuBuilder(true, NULL);
+	FMenuBuilder MenuBuilder(true, nullptr);
 
 	for (int32 i = 0; i < ARRAY_COUNT(SectionSizes); i++)
 	{
@@ -564,7 +561,7 @@ FText FLandscapeEditorDetailCustomization_NewLandscape::GetSectionSize(TSharedRe
 
 TSharedRef<SWidget> FLandscapeEditorDetailCustomization_NewLandscape::GetSectionsPerComponentMenu(TSharedRef<IPropertyHandle> PropertyHandle)
 {
-	FMenuBuilder MenuBuilder(true, NULL);
+	FMenuBuilder MenuBuilder(true, nullptr);
 
 	for (int32 i = 0; i < ARRAY_COUNT(NumSections); i++)
 	{
@@ -602,7 +599,7 @@ FText FLandscapeEditorDetailCustomization_NewLandscape::GetSectionsPerComponent(
 TOptional<int32> FLandscapeEditorDetailCustomization_NewLandscape::GetLandscapeResolutionX() const
 {
 	FEdModeLandscape* LandscapeEdMode = GetEditorMode();
-	if (LandscapeEdMode != NULL)
+	if (LandscapeEdMode != nullptr)
 	{
 		return (LandscapeEdMode->UISettings->NewLandscape_ComponentCount.X * LandscapeEdMode->UISettings->NewLandscape_SectionsPerComponent * LandscapeEdMode->UISettings->NewLandscape_QuadsPerSection + 1);
 	}
@@ -613,7 +610,7 @@ TOptional<int32> FLandscapeEditorDetailCustomization_NewLandscape::GetLandscapeR
 void FLandscapeEditorDetailCustomization_NewLandscape::OnChangeLandscapeResolutionX(int32 NewValue)
 {
 	FEdModeLandscape* LandscapeEdMode = GetEditorMode();
-	if (LandscapeEdMode != NULL)
+	if (LandscapeEdMode != nullptr)
 	{
 		int32 NewComponentCountX = LandscapeEdMode->UISettings->CalcComponentsCount(NewValue);
 		if (NewComponentCountX != LandscapeEdMode->UISettings->NewLandscape_ComponentCount.X)
@@ -632,7 +629,7 @@ void FLandscapeEditorDetailCustomization_NewLandscape::OnChangeLandscapeResoluti
 void FLandscapeEditorDetailCustomization_NewLandscape::OnCommitLandscapeResolutionX(int32 NewValue, ETextCommit::Type CommitInfo)
 {
 	FEdModeLandscape* LandscapeEdMode = GetEditorMode();
-	if (LandscapeEdMode != NULL)
+	if (LandscapeEdMode != nullptr)
 	{
 		if (!GEditor->IsTransactionActive())
 		{
@@ -647,7 +644,7 @@ void FLandscapeEditorDetailCustomization_NewLandscape::OnCommitLandscapeResoluti
 TOptional<int32> FLandscapeEditorDetailCustomization_NewLandscape::GetLandscapeResolutionY() const
 {
 	FEdModeLandscape* LandscapeEdMode = GetEditorMode();
-	if (LandscapeEdMode != NULL)
+	if (LandscapeEdMode != nullptr)
 	{
 		return (LandscapeEdMode->UISettings->NewLandscape_ComponentCount.Y * LandscapeEdMode->UISettings->NewLandscape_SectionsPerComponent * LandscapeEdMode->UISettings->NewLandscape_QuadsPerSection + 1);
 	}
@@ -658,7 +655,7 @@ TOptional<int32> FLandscapeEditorDetailCustomization_NewLandscape::GetLandscapeR
 void FLandscapeEditorDetailCustomization_NewLandscape::OnChangeLandscapeResolutionY(int32 NewValue)
 {
 	FEdModeLandscape* LandscapeEdMode = GetEditorMode();
-	if (LandscapeEdMode != NULL)
+	if (LandscapeEdMode != nullptr)
 	{
 		int32 NewComponentCountY = LandscapeEdMode->UISettings->CalcComponentsCount(NewValue);
 		if (NewComponentCountY != LandscapeEdMode->UISettings->NewLandscape_ComponentCount.Y)
@@ -677,7 +674,7 @@ void FLandscapeEditorDetailCustomization_NewLandscape::OnChangeLandscapeResoluti
 void FLandscapeEditorDetailCustomization_NewLandscape::OnCommitLandscapeResolutionY(int32 NewValue, ETextCommit::Type CommitInfo)
 {
 	FEdModeLandscape* LandscapeEdMode = GetEditorMode();
-	if (LandscapeEdMode != NULL)
+	if (LandscapeEdMode != nullptr)
 	{
 		if (!GEditor->IsTransactionActive())
 		{
@@ -692,7 +689,7 @@ void FLandscapeEditorDetailCustomization_NewLandscape::OnCommitLandscapeResoluti
 TOptional<int32> FLandscapeEditorDetailCustomization_NewLandscape::GetMinLandscapeResolution() const
 {
 	FEdModeLandscape* LandscapeEdMode = GetEditorMode();
-	if (LandscapeEdMode != NULL)
+	if (LandscapeEdMode != nullptr)
 	{
 		// Min size is one component
 		return (LandscapeEdMode->UISettings->NewLandscape_SectionsPerComponent * LandscapeEdMode->UISettings->NewLandscape_QuadsPerSection + 1);
@@ -704,7 +701,7 @@ TOptional<int32> FLandscapeEditorDetailCustomization_NewLandscape::GetMinLandsca
 TOptional<int32> FLandscapeEditorDetailCustomization_NewLandscape::GetMaxLandscapeResolution() const
 {
 	FEdModeLandscape* LandscapeEdMode = GetEditorMode();
-	if (LandscapeEdMode != NULL)
+	if (LandscapeEdMode != nullptr)
 	{
 		// Max size is either whole components below 8192 verts, or 32 components
 		const int32 QuadsPerComponent = (LandscapeEdMode->UISettings->NewLandscape_SectionsPerComponent * LandscapeEdMode->UISettings->NewLandscape_QuadsPerSection);
@@ -718,7 +715,7 @@ TOptional<int32> FLandscapeEditorDetailCustomization_NewLandscape::GetMaxLandsca
 FText FLandscapeEditorDetailCustomization_NewLandscape::GetTotalComponentCount() const
 {
 	FEdModeLandscape* LandscapeEdMode = GetEditorMode();
-	if (LandscapeEdMode != NULL)
+	if (LandscapeEdMode != nullptr)
 	{
 		return FText::AsNumber(LandscapeEdMode->UISettings->NewLandscape_ComponentCount.X * LandscapeEdMode->UISettings->NewLandscape_ComponentCount.Y);
 	}
@@ -730,7 +727,7 @@ FText FLandscapeEditorDetailCustomization_NewLandscape::GetTotalComponentCount()
 EVisibility FLandscapeEditorDetailCustomization_NewLandscape::GetVisibilityOnlyInNewLandscapeMode(ENewLandscapePreviewMode::Type value)
 {
 	FEdModeLandscape* LandscapeEdMode = GetEditorMode();
-	if (LandscapeEdMode != NULL)
+	if (LandscapeEdMode != nullptr)
 	{
 		if (LandscapeEdMode->NewLandscapePreviewMode == value)
 		{
@@ -743,7 +740,7 @@ EVisibility FLandscapeEditorDetailCustomization_NewLandscape::GetVisibilityOnlyI
 ECheckBoxState FLandscapeEditorDetailCustomization_NewLandscape::NewLandscapeModeIsChecked(ENewLandscapePreviewMode::Type value) const
 {
 	FEdModeLandscape* LandscapeEdMode = GetEditorMode();
-	if (LandscapeEdMode != NULL)
+	if (LandscapeEdMode != nullptr)
 	{
 		if (LandscapeEdMode->NewLandscapePreviewMode == value)
 		{
@@ -758,7 +755,7 @@ void FLandscapeEditorDetailCustomization_NewLandscape::OnNewLandscapeModeChanged
 	if (NewCheckedState == ECheckBoxState::Checked)
 	{
 		FEdModeLandscape* LandscapeEdMode = GetEditorMode();
-		if (LandscapeEdMode != NULL)
+		if (LandscapeEdMode != nullptr)
 		{
 			LandscapeEdMode->NewLandscapePreviewMode = value;
 
@@ -777,15 +774,13 @@ FReply FLandscapeEditorDetailCustomization_NewLandscape::OnCreateButtonClicked()
 		LandscapeEdMode->GetWorld() != nullptr && 
 		LandscapeEdMode->GetWorld()->GetCurrentLevel()->bIsVisible)
 	{
-		FScopedTransaction Transaction(LOCTEXT("Undo", "Creating New Landscape"));
-
 		// Initialize heightmap data
 		TArray<uint16> Data;
 		const int32 ComponentCountX = LandscapeEdMode->UISettings->NewLandscape_ComponentCount.X;
 		const int32 ComponentCountY = LandscapeEdMode->UISettings->NewLandscape_ComponentCount.Y;
 		const int32 QuadsPerComponent = LandscapeEdMode->UISettings->NewLandscape_SectionsPerComponent * LandscapeEdMode->UISettings->NewLandscape_QuadsPerSection;
-		const int32 SizeX = LandscapeEdMode->UISettings->NewLandscape_ComponentCount.X * QuadsPerComponent + 1;
-		const int32 SizeY = LandscapeEdMode->UISettings->NewLandscape_ComponentCount.Y * QuadsPerComponent + 1;
+		const int32 SizeX = ComponentCountX * QuadsPerComponent + 1;
+		const int32 SizeY = ComponentCountY * QuadsPerComponent + 1;
 		Data.AddUninitialized(SizeX * SizeY);
 		uint16* WordData = Data.GetData();
 
@@ -826,43 +821,49 @@ FReply FLandscapeEditorDetailCustomization_NewLandscape::OnCreateButtonClicked()
 		}
 		else if (LandscapeEdMode->NewLandscapePreviewMode == ENewLandscapePreviewMode::ImportLandscape)
 		{
-			const int32 ImportSizeX = LandscapeEdMode->UISettings->ImportLandscape_Width;
-			const int32 ImportSizeY = LandscapeEdMode->UISettings->ImportLandscape_Height;
+			const uint32 ImportSizeX = LandscapeEdMode->UISettings->ImportLandscape_Width;
+			const uint32 ImportSizeY = LandscapeEdMode->UISettings->ImportLandscape_Height;
+
+			if (LandscapeEdMode->UISettings->ImportLandscape_HeightmapImportResult == ELandscapeImportResult::Error)
+			{
+				// Cancel import
+				return FReply::Handled();
+			}
 
-			const TArray<FLandscapeImportLayer>& ImportLandscapeLayersList = LandscapeEdMode->UISettings->ImportLandscape_Layers;
+			TArray<FLandscapeImportLayer>& ImportLandscapeLayersList = LandscapeEdMode->UISettings->ImportLandscape_Layers;
 			ImportLayers.Reserve(ImportLandscapeLayersList.Num());
 
 			// Fill in LayerInfos array and allocate data
-			for (const FLandscapeImportLayer& UIImportLayer : ImportLandscapeLayersList)
+			for (FLandscapeImportLayer& UIImportLayer : ImportLandscapeLayersList)
 			{
 				ImportLayers.Add((const FLandscapeImportLayer&)UIImportLayer); //slicing is fine here
 				FLandscapeImportLayerInfo& ImportLayer = ImportLayers.Last();
 
 				if (ImportLayer.LayerInfo != nullptr && ImportLayer.SourceFilePath != "")
 				{
-					FFileHelper::LoadFileToArray(ImportLayer.LayerData, *ImportLayer.SourceFilePath, FILEREAD_Silent);
+					ILandscapeEditorModule& LandscapeEditorModule = FModuleManager::GetModuleChecked<ILandscapeEditorModule>("LandscapeEditor");
+					const ILandscapeWeightmapFileFormat* WeightmapFormat = LandscapeEditorModule.GetWeightmapFormatByExtension(*FPaths::GetExtension(ImportLayer.SourceFilePath, true));
 
-					if (ImportLayer.SourceFilePath.EndsWith(".png"))
+					if (WeightmapFormat)
 					{
-						IImageWrapperModule& ImageWrapperModule = FModuleManager::LoadModuleChecked<IImageWrapperModule>("ImageWrapper");
-						IImageWrapperPtr ImageWrapper = ImageWrapperModule.CreateImageWrapper(EImageFormat::PNG);
-
-						const TArray<uint8>* RawData = nullptr;
-						if (ImageWrapper->SetCompressed(ImportLayer.LayerData.GetData(), ImportLayer.LayerData.Num()) &&
-							ImageWrapper->GetRaw(ERGBFormat::Gray, 8, RawData))
-						{
-							ImportLayer.LayerData = *RawData; // agh I want to use MoveTemp() here
-						}
-						else
-						{
-							ImportLayer.LayerData.Empty();
-						}
+						FLandscapeWeightmapImportData WeightmapImportData = WeightmapFormat->Import(*ImportLayer.SourceFilePath, ImportLayer.LayerName, {ImportSizeX, ImportSizeY});
+						UIImportLayer.ImportResult = WeightmapImportData.ResultCode;
+						UIImportLayer.ErrorMessage = WeightmapImportData.ErrorMessage;
+						ImportLayer.LayerData = MoveTemp(WeightmapImportData.Data);
+					}
+					else
+					{
+						UIImportLayer.ImportResult = ELandscapeImportResult::Error;
+						UIImportLayer.ErrorMessage = LOCTEXT("Import_UnknownFileType", "File type not recognised");
 					}
 
-					// Remove invalid raw weightmap data
-					if (ImportLayer.LayerData.Num() != (ImportSizeX * ImportSizeY))
+					if (UIImportLayer.ImportResult == ELandscapeImportResult::Error)
 					{
 						ImportLayer.LayerData.Empty();
+						FMessageDialog::Open(EAppMsgType::Ok, UIImportLayer.ErrorMessage);
+
+						// Cancel import
+						return FReply::Handled();
 					}
 				}
 			}
@@ -890,10 +891,10 @@ FReply FLandscapeEditorDetailCustomization_NewLandscape::OnCreateButtonClicked()
 					}
 				}
 			}
-
-			LandscapeEdMode->UISettings->ClearImportLandscapeData();
 		}
 
+		FScopedTransaction Transaction(LOCTEXT("Undo", "Creating New Landscape"));
+
 		const FVector Offset = FTransform(LandscapeEdMode->UISettings->NewLandscape_Rotation, FVector::ZeroVector, LandscapeEdMode->UISettings->NewLandscape_Scale).TransformVector(FVector(-ComponentCountX * QuadsPerComponent / 2, -ComponentCountY * QuadsPerComponent / 2, 0));
 		ALandscape* Landscape = LandscapeEdMode->GetWorld()->SpawnActor<ALandscape>(LandscapeEdMode->UISettings->NewLandscape_Location + Offset, LandscapeEdMode->UISettings->NewLandscape_Rotation);
 		Landscape->LandscapeMaterial = LandscapeEdMode->UISettings->NewLandscape_Material.Get();
@@ -912,7 +913,7 @@ FReply FLandscapeEditorDetailCustomization_NewLandscape::OnCreateButtonClicked()
 			Landscape->ReimportHeightmapFilePath = LandscapeEdMode->UISettings->ImportLandscape_HeightmapFilename;
 		}
 
-		ULandscapeInfo* LandscapeInfo = Landscape->GetLandscapeInfo(true);
+		ULandscapeInfo* LandscapeInfo = Landscape->CreateLandscapeInfo();
 		LandscapeInfo->UpdateLayerInfoMap(Landscape);
 
 		// Import doesn't fill in the LayerInfo for layers with no data, do that now
@@ -940,7 +941,7 @@ FReply FLandscapeEditorDetailCustomization_NewLandscape::OnCreateButtonClicked()
 		}
 
 		LandscapeEdMode->UpdateLandscapeList();
-		LandscapeEdMode->CurrentToolTarget.LandscapeInfo = Landscape->GetLandscapeInfo();
+		LandscapeEdMode->CurrentToolTarget.LandscapeInfo = LandscapeInfo;
 		LandscapeEdMode->CurrentToolTarget.TargetType = ELandscapeToolTargetType::Heightmap;
 		LandscapeEdMode->CurrentToolTarget.LayerInfo = nullptr;
 		LandscapeEdMode->CurrentToolTarget.LayerName = NAME_None;
@@ -956,7 +957,7 @@ FReply FLandscapeEditorDetailCustomization_NewLandscape::OnCreateButtonClicked()
 FReply FLandscapeEditorDetailCustomization_NewLandscape::OnFillWorldButtonClicked()
 {
 	FEdModeLandscape* LandscapeEdMode = GetEditorMode();
-	if (LandscapeEdMode != NULL)
+	if (LandscapeEdMode != nullptr)
 	{
 		FVector& NewLandscapeLocation = LandscapeEdMode->UISettings->NewLandscape_Location;
 		NewLandscapeLocation.X = 0;
@@ -974,7 +975,7 @@ FReply FLandscapeEditorDetailCustomization_NewLandscape::OnFillWorldButtonClicke
 FReply FLandscapeEditorDetailCustomization_NewLandscape::OnFitImportDataButtonClicked()
 {
 	FEdModeLandscape* LandscapeEdMode = GetEditorMode();
-	if (LandscapeEdMode != NULL)
+	if (LandscapeEdMode != nullptr)
 	{
 		ChooseBestComponentSizeForImport(LandscapeEdMode);
 	}
@@ -985,19 +986,16 @@ FReply FLandscapeEditorDetailCustomization_NewLandscape::OnFitImportDataButtonCl
 bool FLandscapeEditorDetailCustomization_NewLandscape::GetImportButtonIsEnabled() const
 {
 	FEdModeLandscape* LandscapeEdMode = GetEditorMode();
-	if (LandscapeEdMode != NULL)
+	if (LandscapeEdMode != nullptr)
 	{
-		if (LandscapeEdMode->UISettings->ImportLandscape_HeightmapError != ELandscapeImportHeightmapError::None &&
-			LandscapeEdMode->UISettings->ImportLandscape_HeightmapError != ELandscapeImportHeightmapError::ColorPng &&
-			LandscapeEdMode->UISettings->ImportLandscape_HeightmapError != ELandscapeImportHeightmapError::LowBitDepth)
+		if (LandscapeEdMode->UISettings->ImportLandscape_HeightmapImportResult == ELandscapeImportResult::Error)
 		{
 			return false;
 		}
 
 		for (int32 i = 0; i < LandscapeEdMode->UISettings->ImportLandscape_Layers.Num(); ++i)
 		{
-			if (LandscapeEdMode->UISettings->ImportLandscape_Layers[i].ImportError != ELandscapeImportLayerError::None &&
-				LandscapeEdMode->UISettings->ImportLandscape_Layers[i].ImportError != ELandscapeImportLayerError::ColorPng)
+			if (LandscapeEdMode->UISettings->ImportLandscape_Layers[i].ImportResult == ELandscapeImportResult::Error)
 			{
 				return false;
 			}
@@ -1008,13 +1006,12 @@ bool FLandscapeEditorDetailCustomization_NewLandscape::GetImportButtonIsEnabled(
 	return false;
 }
 
-
-EVisibility FLandscapeEditorDetailCustomization_NewLandscape::GetHeightmapErrorVisibility(TSharedRef<IPropertyHandle> PropertyHandle_HeightmapError)
+EVisibility FLandscapeEditorDetailCustomization_NewLandscape::GetHeightmapErrorVisibility(TSharedRef<IPropertyHandle> PropertyHandle_HeightmapImportResult)
 {
-	uint8 HeightmapErrorAsByte = 0;
-	FPropertyAccess::Result Result = PropertyHandle_HeightmapError->GetValue(HeightmapErrorAsByte);
+	ELandscapeImportResult HeightmapImportResult;
+	FPropertyAccess::Result Result = PropertyHandle_HeightmapImportResult->GetValue((uint8&)HeightmapImportResult);
 
-	if (!ensure(Result == FPropertyAccess::Success))
+	if (Result == FPropertyAccess::Fail)
 	{
 		return EVisibility::Collapsed;
 	}
@@ -1024,7 +1021,7 @@ EVisibility FLandscapeEditorDetailCustomization_NewLandscape::GetHeightmapErrorV
 		return EVisibility::Visible;
 	}
 
-	if (HeightmapErrorAsByte != ELandscapeImportHeightmapError::None)
+	if (HeightmapImportResult != ELandscapeImportResult::Success)
 	{
 		return EVisibility::Visible;
 	}
@@ -1032,58 +1029,29 @@ EVisibility FLandscapeEditorDetailCustomization_NewLandscape::GetHeightmapErrorV
 	return EVisibility::Collapsed;
 }
 
-FText FLandscapeEditorDetailCustomization_NewLandscape::GetHeightmapErrorText(TSharedRef<IPropertyHandle> PropertyHandle_HeightmapError)
+FSlateColor FLandscapeEditorDetailCustomization_NewLandscape::GetHeightmapErrorColor(TSharedRef<IPropertyHandle> PropertyHandle_HeightmapImportResult)
 {
-	uint8 HeightmapErrorAsByte = 0;
-	FPropertyAccess::Result Result = PropertyHandle_HeightmapError->GetValue(HeightmapErrorAsByte);
-	
-	if (!ensure(Result == FPropertyAccess::Success))
-	{
-		return FText::GetEmpty();
-	}
+	ELandscapeImportResult HeightmapImportResult;
+	FPropertyAccess::Result Result = PropertyHandle_HeightmapImportResult->GetValue((uint8&)HeightmapImportResult);
 
-	if (Result == FPropertyAccess::MultipleValues)
+	if (Result == FPropertyAccess::Fail ||
+		Result == FPropertyAccess::MultipleValues)
 	{
-		return NSLOCTEXT("PropertyEditor", "MultipleValues", "Multiple Values");
+		return FCoreStyle::Get().GetColor("ErrorReporting.BackgroundColor");
 	}
 
-	switch (HeightmapErrorAsByte)
+	switch (HeightmapImportResult)
 	{
-	case ELandscapeImportHeightmapError::None:
-		return FText::GetEmpty();
-	case ELandscapeImportHeightmapError::FileNotFound:
-		return LOCTEXT("Import_HeightmapFileDoesNotExist", "The Heightmap file does not exist");
-	case ELandscapeImportHeightmapError::InvalidSize:
-		return LOCTEXT("Import_HeightmapFileInvalidSize", "The Heightmap file has an invalid size (possibly not 16-bit?)");
-	case ELandscapeImportHeightmapError::CorruptFile:
-		return LOCTEXT("Import_HeightmapFileCorruptPng", "The Heightmap file cannot be read (corrupt png?)");
-	case ELandscapeImportHeightmapError::ColorPng:
-		return LOCTEXT("Import_HeightmapFileColorPng", "The Heightmap file appears to be a color png, grayscale is expected. The import *can* continue, but the result may not be what you expect...");
-	case ELandscapeImportHeightmapError::LowBitDepth:
-		return LOCTEXT("Import_HeightmapFileLowBitDepth", "The Heightmap file appears to be an 8-bit png, 16-bit is preferred. The import *can* continue, but the result may be lower quality than desired.");
+	case ELandscapeImportResult::Success:
+		return FCoreStyle::Get().GetColor("InfoReporting.BackgroundColor");
+	case ELandscapeImportResult::Warning:
+		return FCoreStyle::Get().GetColor("ErrorReporting.WarningBackgroundColor");
+	case ELandscapeImportResult::Error:
+		return FCoreStyle::Get().GetColor("ErrorReporting.BackgroundColor");
 	default:
 		check(0);
+		return FSlateColor();
 	}
-
-	return FText::GetEmpty();
-}
-
-FText FLandscapeEditorDetailCustomization_NewLandscape::GetImportHeightmapFilenameString(TSharedRef<IPropertyHandle> PropertyHandle_HeightmapFilename)
-{
-	FString HeightmapFilename;
-	FPropertyAccess::Result Result = PropertyHandle_HeightmapFilename->GetValue(HeightmapFilename);
-	
-	if (!ensure(Result == FPropertyAccess::Success))
-	{
-		return FText::GetEmpty();
-	}
-
-	if (Result == FPropertyAccess::MultipleValues)
-	{
-		return NSLOCTEXT("PropertyEditor", "MultipleValues", "Multiple Values");
-	}
-
-	return FText::FromString(HeightmapFilename);
 }
 
 void FLandscapeEditorDetailCustomization_NewLandscape::SetImportHeightmapFilenameString(const FText& NewValue, ETextCommit::Type CommitInfo, TSharedRef<IPropertyHandle> PropertyHandle_HeightmapFilename)
@@ -1095,100 +1063,36 @@ void FLandscapeEditorDetailCustomization_NewLandscape::SetImportHeightmapFilenam
 void FLandscapeEditorDetailCustomization_NewLandscape::OnImportHeightmapFilenameChanged()
 {
 	FEdModeLandscape* LandscapeEdMode = GetEditorMode();
-	if (LandscapeEdMode != NULL)
+	if (LandscapeEdMode != nullptr)
 	{
 		ImportResolutions.Reset(1);
 		LandscapeEdMode->UISettings->ImportLandscape_Width = 0;
 		LandscapeEdMode->UISettings->ImportLandscape_Height = 0;
 		LandscapeEdMode->UISettings->ClearImportLandscapeData();
-		LandscapeEdMode->UISettings->ImportLandscape_HeightmapError = ELandscapeImportHeightmapError::None;
+		LandscapeEdMode->UISettings->ImportLandscape_HeightmapImportResult = ELandscapeImportResult::Success;
+		LandscapeEdMode->UISettings->ImportLandscape_HeightmapErrorMessage = FText();
 
 		if (!LandscapeEdMode->UISettings->ImportLandscape_HeightmapFilename.IsEmpty())
 		{
-			if (LandscapeEdMode->UISettings->ImportLandscape_HeightmapFilename.EndsWith(".png"))
+			ILandscapeEditorModule& LandscapeEditorModule = FModuleManager::GetModuleChecked<ILandscapeEditorModule>("LandscapeEditor");
+			const ILandscapeHeightmapFileFormat* HeightmapFormat = LandscapeEditorModule.GetHeightmapFormatByExtension(*FPaths::GetExtension(LandscapeEdMode->UISettings->ImportLandscape_HeightmapFilename, true));
+
+			if (HeightmapFormat)
 			{
-				TArray<uint8> ImportData;
-				if (!FFileHelper::LoadFileToArray(ImportData, *LandscapeEdMode->UISettings->ImportLandscape_HeightmapFilename, FILEREAD_Silent))
+				FLandscapeHeightmapInfo HeightmapImportInfo = HeightmapFormat->Validate(*LandscapeEdMode->UISettings->ImportLandscape_HeightmapFilename);
+				LandscapeEdMode->UISettings->ImportLandscape_HeightmapImportResult = HeightmapImportInfo.ResultCode;
+				LandscapeEdMode->UISettings->ImportLandscape_HeightmapErrorMessage = HeightmapImportInfo.ErrorMessage;
+				ImportResolutions = MoveTemp(HeightmapImportInfo.PossibleResolutions);
+				if (HeightmapImportInfo.DataScale.IsSet())
 				{
-					LandscapeEdMode->UISettings->ImportLandscape_HeightmapError = ELandscapeImportHeightmapError::FileNotFound;
-				}
-				else
-				{
-					IImageWrapperModule& ImageWrapperModule = FModuleManager::LoadModuleChecked<IImageWrapperModule>("ImageWrapper");
-					IImageWrapperPtr ImageWrapper = ImageWrapperModule.CreateImageWrapper(EImageFormat::PNG);
-
-					if (!ImageWrapper->SetCompressed(ImportData.GetData(), ImportData.Num()))
-					{
-						LandscapeEdMode->UISettings->ImportLandscape_HeightmapError = ELandscapeImportHeightmapError::CorruptFile;
-					}
-					else if (ImageWrapper->GetWidth() <= 0 || ImageWrapper->GetHeight() <= 0)
-					{
-						LandscapeEdMode->UISettings->ImportLandscape_HeightmapError = ELandscapeImportHeightmapError::InvalidSize;
-					}
-					else
-					{
-						if (ImageWrapper->GetFormat() != ERGBFormat::Gray)
-						{
-							LandscapeEdMode->UISettings->ImportLandscape_HeightmapError = ELandscapeImportHeightmapError::ColorPng;
-							FMessageDialog::Open(EAppMsgType::Ok, LOCTEXT("Import_HeightmapFileColorPng", "The Heightmap file appears to be a color png, grayscale is expected. The import *can* continue, but the result may not be what you expect..."));
-						}
-						else if (ImageWrapper->GetBitDepth() != 16)
-						{
-							LandscapeEdMode->UISettings->ImportLandscape_HeightmapError = ELandscapeImportHeightmapError::LowBitDepth;
-							FMessageDialog::Open(EAppMsgType::Ok, LOCTEXT("Import_HeightmapFileLowBitDepth", "The Heightmap file appears to be an 8-bit png, 16-bit is preferred. The import *can* continue, but the result may be lower quality than desired."));
-						}
-						FLandscapeImportResolution ImportResolution;
-						ImportResolution.Width = ImageWrapper->GetWidth();
-						ImportResolution.Height = ImageWrapper->GetHeight();
-						ImportResolutions.Add(ImportResolution);
-					}
+					LandscapeEdMode->UISettings->NewLandscape_Scale = HeightmapImportInfo.DataScale.GetValue();
+					LandscapeEdMode->UISettings->NewLandscape_Scale.Z *= LANDSCAPE_INV_ZSCALE;
 				}
 			}
 			else
 			{
-				int64 ImportFileSize = IFileManager::Get().FileSize(*LandscapeEdMode->UISettings->ImportLandscape_HeightmapFilename);
-
-				if (ImportFileSize < 0)
-				{
-					LandscapeEdMode->UISettings->ImportLandscape_HeightmapError = ELandscapeImportHeightmapError::FileNotFound;
-				}
-				else if (ImportFileSize == 0 || ImportFileSize % 2 != 0)
-				{
-					LandscapeEdMode->UISettings->ImportLandscape_HeightmapError = ELandscapeImportHeightmapError::InvalidSize;
-				}
-				else
-				{
-					ImportFileSize /= 2;
-
-					// Find all possible heightmap sizes, up to 8192 width/height
-					int32 InsertIndex = 0;
-					const int32 MinWidth = FMath::DivideAndRoundUp(ImportFileSize, (int64)8192);
-					const int32 MaxWidth = FMath::TruncToInt(FMath::Sqrt(ImportFileSize));
-					for (int32 Width = MinWidth; Width <= MaxWidth; Width++)
-					{
-						if (ImportFileSize % Width == 0)
-						{
-							const int32 Height = ImportFileSize / Width;
-
-							FLandscapeImportResolution ImportResolution;
-							ImportResolution.Width = Width;
-							ImportResolution.Height = Height;
-							ImportResolutions.Insert(ImportResolution, InsertIndex++);
-
-							if (Width != Height)
-							{
-								ImportResolution.Width = Height;
-								ImportResolution.Height = Width;
-								ImportResolutions.Insert(ImportResolution, InsertIndex);
-							}
-						}
-					}
-
-					if (ImportResolutions.Num() == 0)
-					{
-						LandscapeEdMode->UISettings->ImportLandscape_HeightmapError = ELandscapeImportHeightmapError::InvalidSize;
-					}
-				}
+				LandscapeEdMode->UISettings->ImportLandscape_HeightmapImportResult = ELandscapeImportResult::Error;
+				LandscapeEdMode->UISettings->ImportLandscape_HeightmapErrorMessage = LOCTEXT("Import_UnknownFileType", "File type not recognised");
 			}
 		}
 
@@ -1197,28 +1101,21 @@ void FLandscapeEditorDetailCustomization_NewLandscape::OnImportHeightmapFilename
 			int32 i = ImportResolutions.Num() / 2;
 			LandscapeEdMode->UISettings->ImportLandscape_Width = ImportResolutions[i].Width;
 			LandscapeEdMode->UISettings->ImportLandscape_Height = ImportResolutions[i].Height;
-			LandscapeEdMode->UISettings->ClearImportLandscapeData();
 			ChooseBestComponentSizeForImport(LandscapeEdMode);
 		}
-		else
-		{
-			LandscapeEdMode->UISettings->ImportLandscape_Width = 0;
-			LandscapeEdMode->UISettings->ImportLandscape_Height = 0;
-			LandscapeEdMode->UISettings->ClearImportLandscapeData();
-		}
 	}
 }
 
 FReply FLandscapeEditorDetailCustomization_NewLandscape::OnImportHeightmapFilenameButtonClicked(TSharedRef<IPropertyHandle> PropertyHandle_HeightmapFilename)
 {
 	FEdModeLandscape* LandscapeEdMode = GetEditorMode();
-	check(LandscapeEdMode != NULL);
+	check(LandscapeEdMode != nullptr);
 
 	// Prompt the user for the Filenames
 	IDesktopPlatform* DesktopPlatform = FDesktopPlatformModule::Get();
-	if (DesktopPlatform != NULL)
+	if (DesktopPlatform != nullptr)
 	{
-		void* ParentWindowWindowHandle = NULL;
+		void* ParentWindowWindowHandle = nullptr;
 
 		IMainFrameModule& MainFrameModule = FModuleManager::LoadModuleChecked<IMainFrameModule>(TEXT("MainFrame"));
 		const TSharedPtr<SWindow>& MainFrameParentWindow = MainFrameModule.GetParentWindow();
@@ -1227,13 +1124,16 @@ FReply FLandscapeEditorDetailCustomization_NewLandscape::OnImportHeightmapFilena
 			ParentWindowWindowHandle = MainFrameParentWindow->GetNativeWindow()->GetOSWindowHandle();
 		}
 
+		ILandscapeEditorModule& LandscapeEditorModule = FModuleManager::GetModuleChecked<ILandscapeEditorModule>("LandscapeEditor");
+		const TCHAR* FileTypes = LandscapeEditorModule.GetHeightmapImportDialogTypeString();
+
 		TArray<FString> OpenFilenames;
 		bool bOpened = DesktopPlatform->OpenFileDialog(
 			ParentWindowWindowHandle,
 			NSLOCTEXT("UnrealEd", "Import", "Import").ToString(),
 			LandscapeEdMode->UISettings->LastImportPath,
 			TEXT(""),
-			TEXT("All Heightmap files|*.png;*.raw;*.r16|Heightmap png files|*.png|Heightmap .raw/.r16 files|*.raw;*.r16|All files (*.*)|*.*"),
+			FileTypes,
 			EFileDialogFlags::None,
 			OpenFilenames);
 
@@ -1250,7 +1150,7 @@ FReply FLandscapeEditorDetailCustomization_NewLandscape::OnImportHeightmapFilena
 
 TSharedRef<SWidget> FLandscapeEditorDetailCustomization_NewLandscape::GetImportLandscapeResolutionMenu()
 {
-	FMenuBuilder MenuBuilder(true, NULL);
+	FMenuBuilder MenuBuilder(true, nullptr);
 
 	for (int32 i = 0; i < ImportResolutions.Num(); i++)
 	{
@@ -1266,7 +1166,7 @@ TSharedRef<SWidget> FLandscapeEditorDetailCustomization_NewLandscape::GetImportL
 void FLandscapeEditorDetailCustomization_NewLandscape::OnChangeImportLandscapeResolution(int32 Index)
 {
 	FEdModeLandscape* LandscapeEdMode = GetEditorMode();
-	if (LandscapeEdMode != NULL)
+	if (LandscapeEdMode != nullptr)
 	{
 		LandscapeEdMode->UISettings->ImportLandscape_Width = ImportResolutions[Index].Width;
 		LandscapeEdMode->UISettings->ImportLandscape_Height = ImportResolutions[Index].Height;
@@ -1278,7 +1178,7 @@ void FLandscapeEditorDetailCustomization_NewLandscape::OnChangeImportLandscapeRe
 FText FLandscapeEditorDetailCustomization_NewLandscape::GetImportLandscapeResolution() const
 {
 	FEdModeLandscape* LandscapeEdMode = GetEditorMode();
-	if (LandscapeEdMode != NULL)
+	if (LandscapeEdMode != nullptr)
 	{
 		const int32 Width = LandscapeEdMode->UISettings->ImportLandscape_Width;
 		const int32 Height = LandscapeEdMode->UISettings->ImportLandscape_Height;
@@ -1382,7 +1282,7 @@ void FLandscapeEditorDetailCustomization_NewLandscape::ChooseBestComponentSizeFo
 EVisibility FLandscapeEditorDetailCustomization_NewLandscape::GetMaterialTipVisibility() const
 {
 	FEdModeLandscape* LandscapeEdMode = GetEditorMode();
-	if (LandscapeEdMode != NULL)
+	if (LandscapeEdMode != nullptr)
 	{
 		if (LandscapeEdMode->UISettings->ImportLandscape_Layers.Num() == 0)
 		{
@@ -1409,7 +1309,8 @@ void FLandscapeEditorStructCustomization_FLandscapeImportLayer::CustomizeChildre
 	TSharedRef<IPropertyHandle> PropertyHandle_LayerInfo = StructPropertyHandle->GetChildHandle(GET_MEMBER_NAME_CHECKED(FLandscapeImportLayer, LayerInfo)).ToSharedRef();
 	TSharedRef<IPropertyHandle> PropertyHandle_SourceFilePath = StructPropertyHandle->GetChildHandle(GET_MEMBER_NAME_CHECKED(FLandscapeImportLayer, SourceFilePath)).ToSharedRef();
 	TSharedRef<IPropertyHandle> PropertyHandle_ThumbnailMIC = StructPropertyHandle->GetChildHandle(GET_MEMBER_NAME_CHECKED(FLandscapeImportLayer, ThumbnailMIC)).ToSharedRef();
-	TSharedRef<IPropertyHandle> PropertyHandle_LayerError = StructPropertyHandle->GetChildHandle(GET_MEMBER_NAME_CHECKED(FLandscapeImportLayer, ImportError)).ToSharedRef();
+	TSharedRef<IPropertyHandle> PropertyHandle_ImportResult = StructPropertyHandle->GetChildHandle(GET_MEMBER_NAME_CHECKED(FLandscapeImportLayer, ImportResult)).ToSharedRef();
+	TSharedRef<IPropertyHandle> PropertyHandle_ErrorMessage = StructPropertyHandle->GetChildHandle(GET_MEMBER_NAME_CHECKED(FLandscapeImportLayer, ErrorMessage)).ToSharedRef();
 
 	FName LayerName;
 	FText LayerNameText;
@@ -1422,7 +1323,7 @@ void FLandscapeEditorStructCustomization_FLandscapeImportLayer::CustomizeChildre
 		LayerNameText = NSLOCTEXT("PropertyEditor", "MultipleValues", "Multiple Values");
 	}
 
-	UObject* ThumbnailMIC = NULL;
+	UObject* ThumbnailMIC = nullptr;
 	Result = PropertyHandle_ThumbnailMIC->GetValue(ThumbnailMIC);
 	checkSlow(Result == FPropertyAccess::Success);
 
@@ -1431,19 +1332,6 @@ void FLandscapeEditorStructCustomization_FLandscapeImportLayer::CustomizeChildre
 	[
 		SNew(SHorizontalBox)
 		+ SHorizontalBox::Slot()
-		.AutoWidth()
-		.VAlign(VAlign_Center)
-		.Padding(FMargin(2))
-		[
-			SNew(SErrorText)
-			.Visibility_Static(&GetErrorVisibility, PropertyHandle_LayerError)
-			.ErrorText(NSLOCTEXT("UnrealEd", "Error", "!"))
-			.ToolTip(
-				SNew(SToolTip)
-				.Text_Static(&GetErrorText, PropertyHandle_LayerError)
-			)
-		]
-		+ SHorizontalBox::Slot()
 		.FillWidth(1)
 		.VAlign(VAlign_Center)
 		.Padding(FMargin(2))
@@ -1510,11 +1398,25 @@ void FLandscapeEditorStructCustomization_FLandscapeImportLayer::CustomizeChildre
 				SNew(SHorizontalBox)
 				.Visibility_Static(&FLandscapeEditorDetailCustomization_NewLandscape::GetVisibilityOnlyInNewLandscapeMode, ENewLandscapePreviewMode::ImportLandscape)
 				+ SHorizontalBox::Slot()
+				.AutoWidth()
+				.Padding(0, 0, 2, 0)
+				[
+					SNew(SErrorText)
+					.Visibility_Static(&GetErrorVisibility, PropertyHandle_ImportResult)
+					.BackgroundColor_Static(&GetErrorColor, PropertyHandle_ImportResult)
+					.ErrorText(NSLOCTEXT("UnrealEd", "Error", "!"))
+					.ToolTip(
+						SNew(SToolTip)
+						.Text_Static(&GetErrorText, PropertyHandle_ErrorMessage)
+					)
+				]
+				+ SHorizontalBox::Slot()
 				[
 					PropertyHandle_SourceFilePath->CreatePropertyValueWidget()
 				]
 				+ SHorizontalBox::Slot()
 				.AutoWidth()
+				.Padding(1, 0, 0, 0)
 				[
 					SNew(SButton)
 					.ContentPadding(FMargin(4, 0))
@@ -1529,13 +1431,13 @@ void FLandscapeEditorStructCustomization_FLandscapeImportLayer::CustomizeChildre
 FReply FLandscapeEditorStructCustomization_FLandscapeImportLayer::OnLayerFilenameButtonClicked(TSharedRef<IPropertyHandle> PropertyHandle_LayerFilename)
 {
 	FEdModeLandscape* LandscapeEdMode = GetEditorMode();
-	check(LandscapeEdMode != NULL);
+	check(LandscapeEdMode != nullptr);
 
 	// Prompt the user for the Filenames
 	IDesktopPlatform* DesktopPlatform = FDesktopPlatformModule::Get();
-	if (DesktopPlatform != NULL)
+	if (DesktopPlatform != nullptr)
 	{
-		void* ParentWindowWindowHandle = NULL;
+		void* ParentWindowWindowHandle = nullptr;
 
 		IMainFrameModule& MainFrameModule = FModuleManager::LoadModuleChecked<IMainFrameModule>(TEXT("MainFrame"));
 		const TSharedPtr<SWindow>& MainFrameParentWindow = MainFrameModule.GetParentWindow();
@@ -1544,13 +1446,16 @@ FReply FLandscapeEditorStructCustomization_FLandscapeImportLayer::OnLayerFilenam
 			ParentWindowWindowHandle = MainFrameParentWindow->GetNativeWindow()->GetOSWindowHandle();
 		}
 
+		ILandscapeEditorModule& LandscapeEditorModule = FModuleManager::GetModuleChecked<ILandscapeEditorModule>("LandscapeEditor");
+		const TCHAR* FileTypes = LandscapeEditorModule.GetWeightmapImportDialogTypeString();
+
 		TArray<FString> OpenFilenames;
 		bool bOpened = DesktopPlatform->OpenFileDialog(
 			ParentWindowWindowHandle,
 			NSLOCTEXT("UnrealEd", "Import", "Import").ToString(),
 			LandscapeEdMode->UISettings->LastImportPath,
 			TEXT(""),
-			TEXT("All Layer files|*.png;*.raw;*.r16|Layer png files|*.png|Layer files (*.raw,*.r16)|*.raw;*.r16|All files (*.*)|*.*"),
+			FileTypes,
 			EFileDialogFlags::None,
 			OpenFilenames);
 
@@ -1578,9 +1483,9 @@ bool FLandscapeEditorStructCustomization_FLandscapeImportLayer::ShouldFilterLaye
 
 EVisibility FLandscapeEditorStructCustomization_FLandscapeImportLayer::GetImportLayerCreateVisibility(TSharedRef<IPropertyHandle> PropertyHandle_LayerInfo)
 {
-	UObject* LayerInfoAsUObject = NULL;
+	UObject* LayerInfoAsUObject = nullptr;
 	if (PropertyHandle_LayerInfo->GetValue(LayerInfoAsUObject) != FPropertyAccess::Fail &&
-		LayerInfoAsUObject == NULL)
+		LayerInfoAsUObject == nullptr)
 	{
 		return EVisibility::Visible;
 	}
@@ -1590,7 +1495,7 @@ EVisibility FLandscapeEditorStructCustomization_FLandscapeImportLayer::GetImport
 
 TSharedRef<SWidget> FLandscapeEditorStructCustomization_FLandscapeImportLayer::OnGetImportLayerCreateMenu(TSharedRef<IPropertyHandle> PropertyHandle_LayerInfo, FName LayerName)
 {
-	FMenuBuilder MenuBuilder(true, NULL);
+	FMenuBuilder MenuBuilder(true, nullptr);
 
 	MenuBuilder.AddMenuEntry(LOCTEXT("Target_Create_Blended", "Weight-Blended Layer (normal)"), FText(), FSlateIcon(),
 		FUIAction(FExecuteAction::CreateStatic(&OnImportLayerCreateClicked, PropertyHandle_LayerInfo, LayerName, false)));
@@ -1604,7 +1509,7 @@ TSharedRef<SWidget> FLandscapeEditorStructCustomization_FLandscapeImportLayer::O
 void FLandscapeEditorStructCustomization_FLandscapeImportLayer::OnImportLayerCreateClicked(TSharedRef<IPropertyHandle> PropertyHandle_LayerInfo, FName LayerName, bool bNoWeightBlend)
 {
 	FEdModeLandscape* LandscapeEdMode = GetEditorMode();
-	if (LandscapeEdMode != NULL)
+	if (LandscapeEdMode != nullptr)
 	{
 		// Hack as we don't have a direct world pointer in the EdMode...
 		ULevel* Level = LandscapeEdMode->CurrentGizmoActor->GetWorld()->GetCurrentLevel();
@@ -1628,7 +1533,7 @@ void FLandscapeEditorStructCustomization_FLandscapeImportLayer::OnImportLayerCre
 			PackageName = NewLayerDlg->GetFullAssetPath().ToString();
 			LayerObjectName = FName(*NewLayerDlg->GetAssetName().ToString());
 
-			UPackage* Package = CreatePackage(NULL, *PackageName);
+			UPackage* Package = CreatePackage(nullptr, *PackageName);
 			ULandscapeLayerInfoObject* LayerInfo = NewObject<ULandscapeLayerInfoObject>(Package, LayerObjectName, RF_Public | RF_Standalone | RF_Transactional);
 			LayerInfo->LayerName = LayerName;
 			LayerInfo->bNoWeightBlend = bNoWeightBlend;
@@ -1650,37 +1555,53 @@ void FLandscapeEditorStructCustomization_FLandscapeImportLayer::OnImportLayerCre
 	}
 }
 
-EVisibility FLandscapeEditorStructCustomization_FLandscapeImportLayer::GetErrorVisibility(TSharedRef<IPropertyHandle> PropertyHandle_LayerError)
+EVisibility FLandscapeEditorStructCustomization_FLandscapeImportLayer::GetErrorVisibility(TSharedRef<IPropertyHandle> PropertyHandle_ImportResult)
 {
-	auto LayerError = ELandscapeImportLayerError::None;
-	FPropertyAccess::Result Result = PropertyHandle_LayerError->GetValue((uint8&)LayerError);
+	ELandscapeImportResult WeightmapImportResult;
+	FPropertyAccess::Result Result = PropertyHandle_ImportResult->GetValue((uint8&)WeightmapImportResult);
+
 	if (Result == FPropertyAccess::Fail ||
 		Result == FPropertyAccess::MultipleValues)
 	{
 		return EVisibility::Visible;
 	}
 
-	// None of the errors are currently relevant outside of "Import" mode
-	FEdModeLandscape* LandscapeEdMode = GetEditorMode();
-	if (LandscapeEdMode != NULL)
+	if (WeightmapImportResult != ELandscapeImportResult::Success)
 	{
-		if (LandscapeEdMode->NewLandscapePreviewMode != ENewLandscapePreviewMode::ImportLandscape)
-		{
-			return EVisibility::Collapsed;
-		}
+		return EVisibility::Visible;
 	}
+	return EVisibility::Collapsed;
+}
+
+FSlateColor FLandscapeEditorStructCustomization_FLandscapeImportLayer::GetErrorColor(TSharedRef<IPropertyHandle> PropertyHandle_ImportResult)
+{
+	ELandscapeImportResult WeightmapImportResult;
+	FPropertyAccess::Result Result = PropertyHandle_ImportResult->GetValue((uint8&)WeightmapImportResult);
+	check(Result == FPropertyAccess::Success);
 
-	if (LayerError != ELandscapeImportLayerError::None)
+	if (Result == FPropertyAccess::MultipleValues)
 	{
-		return EVisibility::Visible;
+		return FCoreStyle::Get().GetColor("ErrorReporting.BackgroundColor");
+	}
+
+	switch (WeightmapImportResult)
+	{
+	case ELandscapeImportResult::Success:
+		return FCoreStyle::Get().GetColor("InfoReporting.BackgroundColor");
+	case ELandscapeImportResult::Warning:
+		return FCoreStyle::Get().GetColor("ErrorReporting.WarningBackgroundColor");
+	case ELandscapeImportResult::Error:
+		return FCoreStyle::Get().GetColor("ErrorReporting.BackgroundColor");
+	default:
+		check(0);
+		return FSlateColor();
 	}
-	return EVisibility::Collapsed;
 }
 
-FText FLandscapeEditorStructCustomization_FLandscapeImportLayer::GetErrorText(TSharedRef<IPropertyHandle> PropertyHandle_LayerError)
+FText FLandscapeEditorStructCustomization_FLandscapeImportLayer::GetErrorText(TSharedRef<IPropertyHandle> PropertyHandle_ErrorMessage)
 {
-	auto LayerError = ELandscapeImportLayerError::None;
-	FPropertyAccess::Result Result = PropertyHandle_LayerError->GetValue((uint8&)LayerError);
+	FText ErrorMessage;
+	FPropertyAccess::Result Result = PropertyHandle_ErrorMessage->GetValue(ErrorMessage);
 	if (Result == FPropertyAccess::Fail)
 	{
 		return LOCTEXT("Import_LayerUnknownError", "Unknown Error");
@@ -1690,26 +1611,7 @@ FText FLandscapeEditorStructCustomization_FLandscapeImportLayer::GetErrorText(TS
 		return NSLOCTEXT("PropertyEditor", "MultipleValues", "Multiple Values");
 	}
 
-	switch (LayerError)
-	{
-	case ELandscapeImportLayerError::None:
-		return FText::GetEmpty();
-	case ELandscapeImportLayerError::MissingLayerInfo:
-		return LOCTEXT("Import_LayerInfoNotSet", "Can't import a layer file without a layer info");
-	case ELandscapeImportLayerError::FileNotFound:
-		return LOCTEXT("Import_LayerFileDoesNotExist", "The Layer file does not exist");
-	case ELandscapeImportLayerError::FileSizeMismatch:
-		return LOCTEXT("Import_LayerSizeMismatch", "Size of the layer file does not match size of heightmap file");
-	case ELandscapeImportLayerError::CorruptFile:
-		return LOCTEXT("Import_LayerCorruptPng", "The Layer file cannot be read (corrupt png?)");
-	case ELandscapeImportLayerError::ColorPng:
-		return LOCTEXT("Import_LayerColorPng", "The Layer file appears to be a color png, grayscale is expected. The import *can* continue, but the result may not be what you expect...");
-	default:
-		checkSlow(0);
-		return LOCTEXT("Import_LayerUnknownError", "Unknown Error");
-	}
-
-	return FText::GetEmpty();
+	return ErrorMessage;
 }
 
 #undef LOCTEXT_NAMESPACE
diff --git a/Engine/Source/Editor/LandscapeEditor/Private/LandscapeEditorDetailCustomization_NewLandscape.h b/Engine/Source/Editor/LandscapeEditor/Private/LandscapeEditorDetailCustomization_NewLandscape.h
--- a/Engine/Source/Editor/LandscapeEditor/Private/LandscapeEditorDetailCustomization_NewLandscape.h
+++ b/Engine/Source/Editor/LandscapeEditor/Private/LandscapeEditorDetailCustomization_NewLandscape.h
@@ -9,12 +9,6 @@
  * Slate widgets customizer for the "New Landscape" tool
  */
 
-struct FLandscapeImportResolution
-{
-	int32 Width;
-	int32 Height;
-};
-
 class FLandscapeEditorDetailCustomization_NewLandscape : public FLandscapeEditorDetailCustomization_Base
 {
 public:
@@ -56,9 +50,8 @@ public:
 	void OnNewLandscapeModeChanged(ECheckBoxState NewCheckedState, ENewLandscapePreviewMode::Type value);
 
 	// Import
-	static EVisibility GetHeightmapErrorVisibility(TSharedRef<IPropertyHandle> PropertyHandle_HeightmapError);
-	static FText GetHeightmapErrorText(TSharedRef<IPropertyHandle> PropertyHandle_HeightmapError);
-	static FText GetImportHeightmapFilenameString(TSharedRef<IPropertyHandle> PropertyHandle_HeightmapFilename);
+	static EVisibility GetHeightmapErrorVisibility(TSharedRef<IPropertyHandle> PropertyHandle_HeightmapImportResult);
+	static FSlateColor GetHeightmapErrorColor(TSharedRef<IPropertyHandle> PropertyHandle_HeightmapImportResult);
 	static void SetImportHeightmapFilenameString(const FText& NewValue, ETextCommit::Type CommitInfo, TSharedRef<IPropertyHandle> PropertyHandle_HeightmapFilename);
 	void OnImportHeightmapFilenameChanged();
 	static FReply OnImportHeightmapFilenameButtonClicked(TSharedRef<IPropertyHandle> PropertyHandle_HeightmapFilename);
@@ -80,7 +73,7 @@ protected:
 	static const int32 SectionSizes[];
 	static const int32 NumSections[];
 
-	TArray<FLandscapeImportResolution> ImportResolutions;
+	TArray<FLandscapeFileResolution> ImportResolutions;
 };
 
 class FLandscapeEditorStructCustomization_FLandscapeImportLayer : public FLandscapeEditorStructCustomization_Base
@@ -99,6 +92,7 @@ public:
 	static TSharedRef<SWidget> OnGetImportLayerCreateMenu(TSharedRef<IPropertyHandle> PropertyHandle_LayerInfo, FName LayerName);
 	static void OnImportLayerCreateClicked(TSharedRef<IPropertyHandle> PropertyHandle_LayerInfo, FName LayerName, bool bNoWeightBlend);
 
-	static EVisibility GetErrorVisibility(TSharedRef<IPropertyHandle> PropertyHandle_LayerError);
-	static FText GetErrorText(TSharedRef<IPropertyHandle> PropertyHandle_LayerError);
+	static EVisibility GetErrorVisibility(TSharedRef<IPropertyHandle> PropertyHandle_ImportResult);
+	static FSlateColor GetErrorColor(TSharedRef<IPropertyHandle> PropertyHandle_ImportResult);
+	static FText GetErrorText(TSharedRef<IPropertyHandle> PropertyHandle_ErrorMessage);
 };
diff --git a/Engine/Source/Editor/LandscapeEditor/Private/LandscapeEditorDetailCustomization_TargetLayers.cpp b/Engine/Source/Editor/LandscapeEditor/Private/LandscapeEditorDetailCustomization_TargetLayers.cpp
--- a/Engine/Source/Editor/LandscapeEditor/Private/LandscapeEditorDetailCustomization_TargetLayers.cpp
+++ b/Engine/Source/Editor/LandscapeEditor/Private/LandscapeEditorDetailCustomization_TargetLayers.cpp
@@ -546,19 +546,21 @@ void FLandscapeEditorCustomNodeBuilder_TargetLayers::OnExportLayer(const TShared
 		// Prompt for filename
 		FString SaveDialogTitle;
 		FString DefaultFileName;
-		FString FileTypes;
+		const TCHAR* FileTypes = nullptr;
+
+		ILandscapeEditorModule& LandscapeEditorModule = FModuleManager::GetModuleChecked<ILandscapeEditorModule>("LandscapeEditor");
 
 		if (Target->TargetType == ELandscapeToolTargetType::Heightmap)
 		{
 			SaveDialogTitle = *LOCTEXT("ExportHeightmap", "Export Landscape Heightmap").ToString();
 			DefaultFileName = TEXT("Heightmap.png");
-			FileTypes = TEXT("Heightmap png files|*.png|Heightmap .raw files|*.raw|Heightmap .r16 files|*.r16|All files|*.*");
+			FileTypes = LandscapeEditorModule.GetHeightmapExportDialogTypeString();
 		}
 		else //if (Target->TargetType == ELandscapeToolTargetType::Weightmap)
 		{
 			SaveDialogTitle = *FText::Format(LOCTEXT("ExportLayer", "Export Landscape Layer: {0}"), FText::FromName(LayerInfoObj->LayerName)).ToString();
 			DefaultFileName = *FString::Printf(TEXT("%s.png"), *(LayerInfoObj->LayerName.ToString()));
-			FileTypes = TEXT("Layer png files|*.png|Layer .raw files|*.raw|Layer .r8 files|*.r8|All files|*.*");
+			FileTypes = LandscapeEditorModule.GetWeightmapExportDialogTypeString();
 		}
 
 		// Prompt the user for the filenames
@@ -568,7 +570,7 @@ void FLandscapeEditorCustomNodeBuilder_TargetLayers::OnExportLayer(const TShared
 			*SaveDialogTitle,
 			*LandscapeEdMode->UISettings->LastImportPath,
 			*DefaultFileName,
-			*FileTypes,
+			FileTypes,
 			EFileDialogFlags::None,
 			SaveFilenames
 			);
@@ -610,19 +612,21 @@ void FLandscapeEditorCustomNodeBuilder_TargetLayers::OnImportLayer(const TShared
 		// Prompt for filename
 		FString OpenDialogTitle;
 		FString DefaultFileName;
-		FString FileTypes;
+		const TCHAR* FileTypes = nullptr;
+
+		ILandscapeEditorModule& LandscapeEditorModule = FModuleManager::GetModuleChecked<ILandscapeEditorModule>("LandscapeEditor");
 
 		if (Target->TargetType == ELandscapeToolTargetType::Heightmap)
 		{
 			OpenDialogTitle = *LOCTEXT("ImportHeightmap", "Import Landscape Heightmap").ToString();
 			DefaultFileName = TEXT("Heightmap.png");
-			FileTypes = TEXT("All Heightmap files|*.png;*.raw;*.r16|Heightmap png files|*.png|Heightmap .raw/.r16 files|*.raw;*.r16|All files|*.*");
+			FileTypes = LandscapeEditorModule.GetHeightmapImportDialogTypeString();
 		}
 		else //if (Target->TargetType == ELandscapeToolTargetType::Weightmap)
 		{
 			OpenDialogTitle = *FText::Format(LOCTEXT("ImportLayer", "Import Landscape Layer: {0}"), FText::FromName(LayerInfoObj->LayerName)).ToString();
 			DefaultFileName = *FString::Printf(TEXT("%s.png"), *(LayerInfoObj->LayerName.ToString()));
-			FileTypes = TEXT("All Layer files|*.png;*.raw;*.r8|Layer png files|*.png|Layer .raw/.r8 files|*.raw;*.r8|All files|*.*");
+			FileTypes = LandscapeEditorModule.GetWeightmapImportDialogTypeString();
 		}
 
 		// Prompt the user for the filenames
@@ -632,7 +636,7 @@ void FLandscapeEditorCustomNodeBuilder_TargetLayers::OnImportLayer(const TShared
 			*OpenDialogTitle,
 			*LandscapeEdMode->UISettings->LastImportPath,
 			*DefaultFileName,
-			*FileTypes,
+			FileTypes,
 			EFileDialogFlags::None,
 			OpenFilenames
 			);
diff --git a/Engine/Source/Editor/LandscapeEditor/Private/LandscapeEditorModule.cpp b/Engine/Source/Editor/LandscapeEditor/Private/LandscapeEditorModule.cpp
--- a/Engine/Source/Editor/LandscapeEditor/Private/LandscapeEditorModule.cpp
+++ b/Engine/Source/Editor/LandscapeEditor/Private/LandscapeEditorModule.cpp
@@ -4,7 +4,8 @@
 #include "LandscapeEditorCommands.h"
 #include "LandscapeEdMode.h"
 #include "Classes/ActorFactoryLandscape.h"
-
+#include "LandscapeFileFormatPng.h"
+#include "LandscapeFileFormatRaw.h"
 
 #include "LandscapeEditorDetails.h"
 #include "LandscapeEditorDetailCustomizations.h"
@@ -14,15 +15,32 @@
 
 #include "Editor/LevelEditor/Public/LevelEditor.h"
 
-#include "Landscape.h"
 #include "LandscapeRender.h"
 
 #define LOCTEXT_NAMESPACE "LandscapeEditor"
 
+struct FRegisteredLandscapeHeightmapFileFormat
+{
+	TSharedRef<ILandscapeHeightmapFileFormat> FileFormat;
+	FLandscapeFileTypeInfo FileTypeInfo;
+	FString ConcatenatedFileExtensions;
+
+	FRegisteredLandscapeHeightmapFileFormat(TSharedRef<ILandscapeHeightmapFileFormat> InFileFormat);
+};
+
+struct FRegisteredLandscapeWeightmapFileFormat
+{
+	TSharedRef<ILandscapeWeightmapFileFormat> FileFormat;
+	FLandscapeFileTypeInfo FileTypeInfo;
+	FString ConcatenatedFileExtensions;
+
+	FRegisteredLandscapeWeightmapFileFormat(TSharedRef<ILandscapeWeightmapFileFormat> InFileFormat);
+};
+
 class FLandscapeEditorModule : public ILandscapeEditorModule
 {
 public:
-	
+
 	/**
 	 * Called right after the module's DLL has been loaded and the module object has been created
 	 */
@@ -41,22 +59,22 @@ public:
 
 		// register customizations
 		FPropertyEditorModule& PropertyModule = FModuleManager::GetModuleChecked<FPropertyEditorModule>("PropertyEditor");
-		PropertyModule.RegisterCustomClassLayout("LandscapeEditorObject",       FOnGetDetailCustomizationInstance::CreateStatic(&FLandscapeEditorDetails::MakeInstance));
-		PropertyModule.RegisterCustomPropertyTypeLayout("GizmoImportLayer",     FOnGetPropertyTypeCustomizationInstance::CreateStatic(&FLandscapeEditorStructCustomization_FGizmoImportLayer::MakeInstance));
+		PropertyModule.RegisterCustomClassLayout("LandscapeEditorObject", FOnGetDetailCustomizationInstance::CreateStatic(&FLandscapeEditorDetails::MakeInstance));
+		PropertyModule.RegisterCustomPropertyTypeLayout("GizmoImportLayer", FOnGetPropertyTypeCustomizationInstance::CreateStatic(&FLandscapeEditorStructCustomization_FGizmoImportLayer::MakeInstance));
 		PropertyModule.RegisterCustomPropertyTypeLayout("LandscapeImportLayer", FOnGetPropertyTypeCustomizationInstance::CreateStatic(&FLandscapeEditorStructCustomization_FLandscapeImportLayer::MakeInstance));
 
 		PropertyModule.RegisterCustomClassLayout("LandscapeSplineControlPoint", FOnGetDetailCustomizationInstance::CreateStatic(&FLandscapeSplineDetails::MakeInstance));
-		PropertyModule.RegisterCustomClassLayout("LandscapeSplineSegment",      FOnGetDetailCustomizationInstance::CreateStatic(&FLandscapeSplineDetails::MakeInstance));
+		PropertyModule.RegisterCustomClassLayout("LandscapeSplineSegment", FOnGetDetailCustomizationInstance::CreateStatic(&FLandscapeSplineDetails::MakeInstance));
 
 		// add menu extension
 		TSharedRef<FUICommandList> CommandList = MakeShareable(new FUICommandList);
 		const FLandscapeEditorCommands& LandscapeActions = FLandscapeEditorCommands::Get();
-		CommandList->MapAction(LandscapeActions.ViewModeNormal,       FExecuteAction::CreateStatic(&ChangeLandscapeViewMode, ELandscapeViewMode::Normal),       FCanExecuteAction(), FIsActionChecked::CreateStatic(&IsLandscapeViewModeSelected, ELandscapeViewMode::Normal));
-		CommandList->MapAction(LandscapeActions.ViewModeLOD,          FExecuteAction::CreateStatic(&ChangeLandscapeViewMode, ELandscapeViewMode::LOD),          FCanExecuteAction(), FIsActionChecked::CreateStatic(&IsLandscapeViewModeSelected, ELandscapeViewMode::LOD));
+		CommandList->MapAction(LandscapeActions.ViewModeNormal, FExecuteAction::CreateStatic(&ChangeLandscapeViewMode, ELandscapeViewMode::Normal), FCanExecuteAction(), FIsActionChecked::CreateStatic(&IsLandscapeViewModeSelected, ELandscapeViewMode::Normal));
+		CommandList->MapAction(LandscapeActions.ViewModeLOD, FExecuteAction::CreateStatic(&ChangeLandscapeViewMode, ELandscapeViewMode::LOD), FCanExecuteAction(), FIsActionChecked::CreateStatic(&IsLandscapeViewModeSelected, ELandscapeViewMode::LOD));
 		CommandList->MapAction(LandscapeActions.ViewModeLayerDensity, FExecuteAction::CreateStatic(&ChangeLandscapeViewMode, ELandscapeViewMode::LayerDensity), FCanExecuteAction(), FIsActionChecked::CreateStatic(&IsLandscapeViewModeSelected, ELandscapeViewMode::LayerDensity));
-		CommandList->MapAction(LandscapeActions.ViewModeLayerDebug,   FExecuteAction::CreateStatic(&ChangeLandscapeViewMode, ELandscapeViewMode::DebugLayer),   FCanExecuteAction(), FIsActionChecked::CreateStatic(&IsLandscapeViewModeSelected, ELandscapeViewMode::DebugLayer));
-		CommandList->MapAction(LandscapeActions.ViewModeWireframeOnTop,FExecuteAction::CreateStatic(&ChangeLandscapeViewMode, ELandscapeViewMode::WireframeOnTop), FCanExecuteAction(), FIsActionChecked::CreateStatic(&IsLandscapeViewModeSelected, ELandscapeViewMode::WireframeOnTop));
-		CommandList->MapAction(LandscapeActions.ViewModeLayerUsage,   FExecuteAction::CreateStatic(&ChangeLandscapeViewMode, ELandscapeViewMode::LayerUsage), FCanExecuteAction(), FIsActionChecked::CreateStatic(&IsLandscapeViewModeSelected, ELandscapeViewMode::LayerUsage));
+		CommandList->MapAction(LandscapeActions.ViewModeLayerDebug, FExecuteAction::CreateStatic(&ChangeLandscapeViewMode, ELandscapeViewMode::DebugLayer), FCanExecuteAction(), FIsActionChecked::CreateStatic(&IsLandscapeViewModeSelected, ELandscapeViewMode::DebugLayer));
+		CommandList->MapAction(LandscapeActions.ViewModeWireframeOnTop, FExecuteAction::CreateStatic(&ChangeLandscapeViewMode, ELandscapeViewMode::WireframeOnTop), FCanExecuteAction(), FIsActionChecked::CreateStatic(&IsLandscapeViewModeSelected, ELandscapeViewMode::WireframeOnTop));
+		CommandList->MapAction(LandscapeActions.ViewModeLayerUsage, FExecuteAction::CreateStatic(&ChangeLandscapeViewMode, ELandscapeViewMode::LayerUsage), FCanExecuteAction(), FIsActionChecked::CreateStatic(&IsLandscapeViewModeSelected, ELandscapeViewMode::LayerUsage));
 
 		ViewportMenuExtender = MakeShareable(new FExtender);
 		ViewportMenuExtender->AddMenuExtension("LevelViewportLandscape", EExtensionHook::First, CommandList, FMenuExtensionDelegate::CreateStatic(&ConstructLandscapeViewportMenu));
@@ -71,6 +89,12 @@ public:
 		UActorFactoryLandscape* LandscapeProxyActorFactory = NewObject<UActorFactoryLandscape>();
 		LandscapeProxyActorFactory->NewActorClass = ALandscapeProxy::StaticClass();
 		GEditor->ActorFactories.Add(LandscapeProxyActorFactory);
+
+		// Built-in File Formats
+		RegisterHeightmapFileFormat(MakeShareable(new FLandscapeHeightmapFileFormat_Png()));
+		RegisterWeightmapFileFormat(MakeShareable(new FLandscapeWeightmapFileFormat_Png()));
+		RegisterHeightmapFileFormat(MakeShareable(new FLandscapeHeightmapFileFormat_Raw()));
+		RegisterWeightmapFileFormat(MakeShareable(new FLandscapeWeightmapFileFormat_Raw()));
 	}
 
 	/**
@@ -138,10 +162,228 @@ public:
 		return GLandscapeViewMode == ViewMode;
 	}
 
+	virtual void RegisterHeightmapFileFormat(TSharedRef<ILandscapeHeightmapFileFormat> FileFormat) override
+	{
+		HeightmapFormats.Emplace(FileFormat);
+		HeightmapImportDialogTypeString.Reset();
+		HeightmapExportDialogTypeString.Reset();
+	}
+
+	virtual void RegisterWeightmapFileFormat(TSharedRef<ILandscapeWeightmapFileFormat> FileFormat) override
+	{
+		WeightmapFormats.Emplace(FileFormat);
+		WeightmapImportDialogTypeString.Reset();
+		WeightmapExportDialogTypeString.Reset();
+	}
+
+	virtual void UnregisterHeightmapFileFormat(TSharedRef<ILandscapeHeightmapFileFormat> FileFormat) override
+	{
+		int32 Index = HeightmapFormats.IndexOfByPredicate(
+			[FileFormat](const FRegisteredLandscapeHeightmapFileFormat& RegisteredFileFormat)
+			{
+				return RegisteredFileFormat.FileFormat == FileFormat;
+			});
+		if (Index != INDEX_NONE)
+		{
+			HeightmapFormats.RemoveAt(Index);
+			HeightmapImportDialogTypeString.Reset();
+			HeightmapExportDialogTypeString.Reset();
+		}
+	}
+
+	virtual void UnregisterWeightmapFileFormat(TSharedRef<ILandscapeWeightmapFileFormat> FileFormat) override
+	{
+		int32 Index = WeightmapFormats.IndexOfByPredicate(
+			[FileFormat](const FRegisteredLandscapeWeightmapFileFormat& RegisteredFileFormat)
+			{
+				return RegisteredFileFormat.FileFormat == FileFormat;
+			});
+		if (Index != INDEX_NONE)
+		{
+			WeightmapFormats.RemoveAt(Index);
+			WeightmapImportDialogTypeString.Reset();
+			WeightmapExportDialogTypeString.Reset();
+		}
+	}
+
+	virtual const TCHAR* GetHeightmapImportDialogTypeString() const override;
+	virtual const TCHAR* GetWeightmapImportDialogTypeString() const override;
+
+	virtual const TCHAR* GetHeightmapExportDialogTypeString() const override;
+	virtual const TCHAR* GetWeightmapExportDialogTypeString() const override;
+
+	virtual const ILandscapeHeightmapFileFormat* GetHeightmapFormatByExtension(const TCHAR* Extension) override;
+	virtual const ILandscapeWeightmapFileFormat* GetWeightmapFormatByExtension(const TCHAR* Extension) override;
+
 protected:
 	TSharedPtr<FExtender> ViewportMenuExtender;
+	TArray<FRegisteredLandscapeHeightmapFileFormat> HeightmapFormats;
+	TArray<FRegisteredLandscapeWeightmapFileFormat> WeightmapFormats;
+	mutable FString HeightmapImportDialogTypeString;
+	mutable FString WeightmapImportDialogTypeString;
+	mutable FString HeightmapExportDialogTypeString;
+	mutable FString WeightmapExportDialogTypeString;
 };
 
-IMPLEMENT_MODULE( FLandscapeEditorModule, LandscapeEditor );
+IMPLEMENT_MODULE(FLandscapeEditorModule, LandscapeEditor);
+
+FRegisteredLandscapeHeightmapFileFormat::FRegisteredLandscapeHeightmapFileFormat(TSharedRef<ILandscapeHeightmapFileFormat> InFileFormat)
+	: FileFormat(MoveTemp(InFileFormat))
+	, FileTypeInfo(FileFormat->GetInfo())
+{
+	bool bJoin = false;
+	for (const FString& Extension : FileTypeInfo.Extensions)
+	{
+		if (bJoin)
+		{
+			ConcatenatedFileExtensions += TEXT(';');
+		}
+		ConcatenatedFileExtensions += TEXT('*');
+		ConcatenatedFileExtensions += Extension;
+		bJoin = true;
+	}
+}
+
+FRegisteredLandscapeWeightmapFileFormat::FRegisteredLandscapeWeightmapFileFormat(TSharedRef<ILandscapeWeightmapFileFormat> InFileFormat)
+	: FileFormat(MoveTemp(InFileFormat))
+	, FileTypeInfo(FileFormat->GetInfo())
+{
+	bool bJoin = false;
+	for (const FString& Extension : FileTypeInfo.Extensions)
+	{
+		if (bJoin)
+		{
+			ConcatenatedFileExtensions += TEXT(';');
+		}
+		ConcatenatedFileExtensions += TEXT('*');
+		ConcatenatedFileExtensions += Extension;
+		bJoin = true;
+	}
+}
+
+const TCHAR* FLandscapeEditorModule::GetHeightmapImportDialogTypeString() const
+{
+	if (HeightmapImportDialogTypeString.IsEmpty())
+	{
+		HeightmapImportDialogTypeString = TEXT("All Heightmap files|");
+		bool bJoin = false;
+		for (const FRegisteredLandscapeHeightmapFileFormat& HeightmapFormat : HeightmapFormats)
+		{
+			if (bJoin)
+			{
+				HeightmapImportDialogTypeString += TEXT(';');
+			}
+			HeightmapImportDialogTypeString += HeightmapFormat.ConcatenatedFileExtensions;
+			bJoin = true;
+		}
+		HeightmapImportDialogTypeString += TEXT('|');
+		for (const FRegisteredLandscapeHeightmapFileFormat& HeightmapFormat : HeightmapFormats)
+		{
+			HeightmapImportDialogTypeString += HeightmapFormat.FileTypeInfo.Description.ToString();
+			HeightmapImportDialogTypeString += TEXT('|');
+			HeightmapImportDialogTypeString += HeightmapFormat.ConcatenatedFileExtensions;
+			HeightmapImportDialogTypeString += TEXT('|');
+		}
+		HeightmapImportDialogTypeString += TEXT("All Files (*.*)|*.*");
+	}
+
+	return *HeightmapImportDialogTypeString;
+}
+
+const TCHAR* FLandscapeEditorModule::GetWeightmapImportDialogTypeString() const
+{
+	if (WeightmapImportDialogTypeString.IsEmpty())
+	{
+		WeightmapImportDialogTypeString = TEXT("All Layer files|");
+		bool bJoin = false;
+		for (const FRegisteredLandscapeWeightmapFileFormat& WeightmapFormat : WeightmapFormats)
+		{
+			if (bJoin)
+			{
+				WeightmapImportDialogTypeString += TEXT(';');
+			}
+			WeightmapImportDialogTypeString += WeightmapFormat.ConcatenatedFileExtensions;
+			bJoin = true;
+		}
+		WeightmapImportDialogTypeString += TEXT('|');
+		for (const FRegisteredLandscapeWeightmapFileFormat& WeightmapFormat : WeightmapFormats)
+		{
+			WeightmapImportDialogTypeString += WeightmapFormat.FileTypeInfo.Description.ToString();
+			WeightmapImportDialogTypeString += TEXT('|');
+			WeightmapImportDialogTypeString += WeightmapFormat.ConcatenatedFileExtensions;
+			WeightmapImportDialogTypeString += TEXT('|');
+		}
+		WeightmapImportDialogTypeString += TEXT("All Files (*.*)|*.*");
+	}
+
+	return *WeightmapImportDialogTypeString;
+}
+
+const TCHAR* FLandscapeEditorModule::GetHeightmapExportDialogTypeString() const
+{
+	if (HeightmapExportDialogTypeString.IsEmpty())
+	{
+		HeightmapExportDialogTypeString = TEXT("");
+
+		for (const FRegisteredLandscapeHeightmapFileFormat& HeightmapFormat : HeightmapFormats)
+		{
+			if (!HeightmapFormat.FileTypeInfo.bSupportsExport)
+			{
+				continue;
+			}
+			HeightmapExportDialogTypeString += HeightmapFormat.FileTypeInfo.Description.ToString();
+			HeightmapExportDialogTypeString += TEXT('|');
+			HeightmapExportDialogTypeString += HeightmapFormat.ConcatenatedFileExtensions;
+			HeightmapExportDialogTypeString += TEXT('|');
+		}
+		HeightmapExportDialogTypeString += TEXT("All Files (*.*)|*.*");
+	}
+
+	return *HeightmapExportDialogTypeString;
+}
+
+const TCHAR* FLandscapeEditorModule::GetWeightmapExportDialogTypeString() const
+{
+	if (WeightmapExportDialogTypeString.IsEmpty())
+	{
+		WeightmapExportDialogTypeString = TEXT("");
+		for (const FRegisteredLandscapeWeightmapFileFormat& WeightmapFormat : WeightmapFormats)
+		{
+			if (!WeightmapFormat.FileTypeInfo.bSupportsExport)
+			{
+				continue;
+			}
+			WeightmapExportDialogTypeString += WeightmapFormat.FileTypeInfo.Description.ToString();
+			WeightmapExportDialogTypeString += TEXT('|');
+			WeightmapExportDialogTypeString += WeightmapFormat.ConcatenatedFileExtensions;
+			WeightmapExportDialogTypeString += TEXT('|');
+		}
+		WeightmapExportDialogTypeString += TEXT("All Files (*.*)|*.*");
+	}
+
+	return *WeightmapExportDialogTypeString;
+}
+
+const ILandscapeHeightmapFileFormat* FLandscapeEditorModule::GetHeightmapFormatByExtension(const TCHAR* Extension)
+{
+	auto* FoundFormat = HeightmapFormats.FindByPredicate(
+		[Extension](const FRegisteredLandscapeHeightmapFileFormat& HeightmapFormat)
+		{
+			return HeightmapFormat.FileTypeInfo.Extensions.Contains(Extension);
+		});
+
+	return FoundFormat ? &FoundFormat->FileFormat.Get() : nullptr;
+}
+
+const ILandscapeWeightmapFileFormat* FLandscapeEditorModule::GetWeightmapFormatByExtension(const TCHAR* Extension)
+{
+	auto* FoundFormat = WeightmapFormats.FindByPredicate(
+		[Extension](const FRegisteredLandscapeWeightmapFileFormat& WeightmapFormat)
+	{
+		return WeightmapFormat.FileTypeInfo.Extensions.Contains(Extension);
+	});
+
+	return FoundFormat ? &FoundFormat->FileFormat.Get() : nullptr;
+}
 
 #undef LOCTEXT_NAMESPACE
diff --git a/Engine/Source/Editor/LandscapeEditor/Private/LandscapeEditorObject.cpp b/Engine/Source/Editor/LandscapeEditor/Private/LandscapeEditorObject.cpp
--- a/Engine/Source/Editor/LandscapeEditor/Private/LandscapeEditorObject.cpp
+++ b/Engine/Source/Editor/LandscapeEditor/Private/LandscapeEditorObject.cpp
@@ -3,12 +3,12 @@
 #include "LandscapeEditorPrivatePCH.h"
 #include "LandscapeEditorObject.h"
 #include "LandscapeRender.h"
-#include "ImageWrapper.h"
-
 #include "LandscapeMaterialInstanceConstant.h"
 #include "EngineUtils.h"
 #include "LandscapeGizmoActor.h"
 
+//#define LOCTEXT_NAMESPACE "LandscapeEditor"
+
 ULandscapeEditorObject::ULandscapeEditorObject(const FObjectInitializer& ObjectInitializer)
 	: Super(ObjectInitializer)
 
@@ -530,72 +530,25 @@ bool ULandscapeEditorObject::SetAlphaTexture(UTexture2D* InTexture, EColorChanne
 
 const TArray<uint16>& ULandscapeEditorObject::GetImportLandscapeData()
 {
-	if (ImportLandscape_HeightmapFilename.EndsWith(".png"))
-	{
-		if (ImportLandscape_Data.Num() == 0)
-		{
-			ImportLandscape_Data.Empty();
-
-			TArray<uint8> ImportData;
-			if (!FFileHelper::LoadFileToArray(ImportData, *ImportLandscape_HeightmapFilename, FILEREAD_Silent))
-			{
-				ImportLandscape_HeightmapError = ELandscapeImportHeightmapError::FileNotFound;
-			}
-			else
-			{
-				IImageWrapperModule& ImageWrapperModule = FModuleManager::LoadModuleChecked<IImageWrapperModule>("ImageWrapper");
-				IImageWrapperPtr ImageWrapper = ImageWrapperModule.CreateImageWrapper(EImageFormat::PNG);
-
-				const TArray<uint8>* RawData = NULL;
-				if (ImageWrapper->SetCompressed(ImportData.GetData(), ImportData.Num()) &&
-					ImageWrapper->GetWidth() == ImportLandscape_Width && ImageWrapper->GetHeight() == ImportLandscape_Height)
-				{
-					if (ImageWrapper->GetBitDepth() <= 8)
-					{
-						ImageWrapper->GetRaw(ERGBFormat::Gray, 8, RawData);
+	ILandscapeEditorModule& LandscapeEditorModule = FModuleManager::GetModuleChecked<ILandscapeEditorModule>("LandscapeEditor");
+	const ILandscapeHeightmapFileFormat* HeightmapFormat = LandscapeEditorModule.GetHeightmapFormatByExtension(*FPaths::GetExtension(ImportLandscape_HeightmapFilename, true));
 
-						if (ensure(RawData->Num() == ImportLandscape_Width * ImportLandscape_Height))
-						{
-							ImportLandscape_Data.Empty(ImportLandscape_Width * ImportLandscape_Height);
-							ImportLandscape_Data.AddUninitialized(ImportLandscape_Width * ImportLandscape_Height);
-
-							for (int32 i = 0; i < RawData->Num(); i++)
-							{
-								ImportLandscape_Data[i] = (*RawData)[i] * 0x101; // Expand to 16-bit
-							}
-						}
-					}
-					else
-					{
-						ImageWrapper->GetRaw(ERGBFormat::Gray, 16, RawData);
-
-						if (ensure(RawData->Num() == ImportLandscape_Width * ImportLandscape_Height * 2))
-						{
-							ImportLandscape_Data.Empty(ImportLandscape_Width * ImportLandscape_Height);
-							ImportLandscape_Data.AddUninitialized(ImportLandscape_Width * ImportLandscape_Height);
-							FMemory::Memcpy(ImportLandscape_Data.GetData(), RawData->GetData(), RawData->Num());
-						}
-					}
-				}
-			}
-		}
+	if (HeightmapFormat)
+	{
+		FLandscapeHeightmapImportData HeightmapImportData = HeightmapFormat->Import(*ImportLandscape_HeightmapFilename, {ImportLandscape_Width, ImportLandscape_Height});
+		ImportLandscape_HeightmapImportResult = HeightmapImportData.ResultCode;
+		ImportLandscape_HeightmapErrorMessage = HeightmapImportData.ErrorMessage;
+		ImportLandscape_Data = MoveTemp(HeightmapImportData.Data);
 	}
 	else
 	{
-		if (ImportLandscape_Width * ImportLandscape_Height != ImportLandscape_Data.Num())
-		{
-			ImportLandscape_Data.Empty();
+		ImportLandscape_HeightmapImportResult = ELandscapeImportResult::Error;
+		ImportLandscape_HeightmapErrorMessage = NSLOCTEXT("LandscapeEditor.NewLandscape", "Import_UnknownFileType", "File type not recognised");
+	}
 
-			TArray<uint8> ImportData;
-			if (ImportLandscape_Width > 0 && ImportLandscape_Height > 0 &&
-				FFileHelper::LoadFileToArray(ImportData, *ImportLandscape_HeightmapFilename, FILEREAD_Silent) &&
-				ImportData.Num() == ImportLandscape_Width * ImportLandscape_Height * 2)
-			{
-				ImportLandscape_Data.Empty(ImportLandscape_Width * ImportLandscape_Height);
-				ImportLandscape_Data.AddUninitialized(ImportLandscape_Width * ImportLandscape_Height);
-				FMemory::Memcpy(ImportLandscape_Data.GetData(), ImportData.GetData(), ImportData.Num());
-			}
-		}
+	if (ImportLandscape_HeightmapImportResult == ELandscapeImportResult::Error)
+	{
+		ImportLandscape_Data.Empty();
 	}
 
 	return ImportLandscape_Data;
@@ -637,53 +590,38 @@ void ULandscapeEditorObject::RefreshImportLayersList()
 				Context.AddMaterialInterface(NewImportLayer.ThumbnailMIC);
 			}
 
-			NewImportLayer.ImportError = ELandscapeImportLayerError::None;
+			NewImportLayer.ImportResult = ELandscapeImportResult::Success;
+			NewImportLayer.ErrorMessage = FText();
+
 			if (!NewImportLayer.SourceFilePath.IsEmpty())
 			{
-				if (NewImportLayer.LayerInfo == NULL)
+				if (!NewImportLayer.LayerInfo)
 				{
-					NewImportLayer.ImportError = ELandscapeImportLayerError::MissingLayerInfo;
+					NewImportLayer.ImportResult = ELandscapeImportResult::Error;
+					NewImportLayer.ErrorMessage = NSLOCTEXT("LandscapeEditor.NewLandscape", "Import_LayerInfoNotSet", "Can't import a layer file without a layer info");
 				}
 				else
 				{
-					if (NewImportLayer.SourceFilePath.EndsWith(".png"))
+					ILandscapeEditorModule& LandscapeEditorModule = FModuleManager::GetModuleChecked<ILandscapeEditorModule>("LandscapeEditor");
+					const ILandscapeWeightmapFileFormat* WeightmapFormat = LandscapeEditorModule.GetWeightmapFormatByExtension(*FPaths::GetExtension(NewImportLayer.SourceFilePath, true));
+
+					if (!WeightmapFormat)
 					{
-						TArray<uint8> ImportData;
-						if (!FFileHelper::LoadFileToArray(ImportData, *NewImportLayer.SourceFilePath, FILEREAD_Silent))
-						{
-							NewImportLayer.ImportError = ELandscapeImportLayerError::FileNotFound;
-						}
-						else
+						FLandscapeWeightmapInfo WeightmapImportInfo = WeightmapFormat->Validate(*NewImportLayer.SourceFilePath, NewImportLayer.LayerName);
+						NewImportLayer.ImportResult = WeightmapImportInfo.ResultCode;
+						NewImportLayer.ErrorMessage = WeightmapImportInfo.ErrorMessage;
+
+						if (WeightmapImportInfo.ResultCode != ELandscapeImportResult::Error &&
+							!WeightmapImportInfo.PossibleResolutions.Contains(FLandscapeFileResolution{ImportLandscape_Width, ImportLandscape_Height}))
 						{
-							IImageWrapperModule& ImageWrapperModule = FModuleManager::LoadModuleChecked<IImageWrapperModule>("ImageWrapper");
-							IImageWrapperPtr ImageWrapper = ImageWrapperModule.CreateImageWrapper(EImageFormat::PNG);
-
-							if (!ImageWrapper->SetCompressed(ImportData.GetData(), ImportData.Num()))
-							{
-								NewImportLayer.ImportError = ELandscapeImportLayerError::CorruptFile;
-							}
-							else if (ImageWrapper->GetWidth() != ImportLandscape_Width || ImageWrapper->GetHeight() != ImportLandscape_Height)
-							{
-								NewImportLayer.ImportError = ELandscapeImportLayerError::FileSizeMismatch;
-							}
-							else if (ImageWrapper->GetFormat() != ERGBFormat::Gray)
-							{
-								NewImportLayer.ImportError = ELandscapeImportLayerError::ColorPng;
-							}
+							NewImportLayer.ImportResult = ELandscapeImportResult::Error;
+							NewImportLayer.ErrorMessage = NSLOCTEXT("LandscapeEditor.NewLandscape", "Import_LayerSizeMismatch", "Size of the layer file does not match size of heightmap file");
 						}
 					}
 					else
 					{
-						int64 ImportFileSize = IFileManager::Get().FileSize(*NewImportLayer.SourceFilePath);
-
-						if (ImportFileSize < 0)
-						{
-							NewImportLayer.ImportError = ELandscapeImportLayerError::FileNotFound;
-						}
-						else if (ImportFileSize != ImportLandscape_Width * ImportLandscape_Height)
-						{
-							NewImportLayer.ImportError = ELandscapeImportLayerError::FileSizeMismatch;
-						}
+						NewImportLayer.ImportResult = ELandscapeImportResult::Error;
+						NewImportLayer.ErrorMessage = NSLOCTEXT("LandscapeEditor.NewLandscape", "Import_UnknownFileType", "File type not recognised");
 					}
 				}
 			}
