Fixed LandscapeFileFormatInterface.h passing FString to %s in checkf call

diff --git a/Engine/Source/Editor/LandscapeEditor/Public/LandscapeFileFormatInterface.h b/Engine/Source/Editor/LandscapeEditor/Public/LandscapeFileFormatInterface.h
--- a/Engine/Source/Editor/LandscapeEditor/Public/LandscapeFileFormatInterface.h
+++ b/Engine/Source/Editor/LandscapeEditor/Public/LandscapeFileFormatInterface.h
@@ -127,7 +127,7 @@ public:
 	 */
 	virtual void Export(const TCHAR* HeightmapFilename, TArrayView<const uint16> Data, FLandscapeFileResolution DataResolution, FVector Scale) const
 	{
-		checkf(0, TEXT("File type hasn't implemented support for heightmap export - %s"), FPaths::GetExtension(HeightmapFilename, true));
+		checkf(0, TEXT("File type hasn't implemented support for heightmap export - %s"), *FPaths::GetExtension(HeightmapFilename, true));
 	}
 };
 
@@ -165,6 +165,6 @@ public:
 	 */
 	virtual void Export(const TCHAR* WeightmapFilename, FName LayerName, TArrayView<const uint8> Data, FLandscapeFileResolution DataResolution) const
 	{
-		checkf(0, TEXT("File type hasn't implemented support for weightmap export - %s"), FPaths::GetExtension(WeightmapFilename, true));
+		checkf(0, TEXT("File type hasn't implemented support for weightmap export - %s"), *FPaths::GetExtension(WeightmapFilename, true));
 	}
 };
