Fix CIS

diff --git a/Engine/Source/Editor/LandscapeEditor/Private/LandscapeFileFormatPng.cpp b/Engine/Source/Editor/LandscapeEditor/Private/LandscapeFileFormatPng.cpp
--- a/Engine/Source/Editor/LandscapeEditor/Private/LandscapeFileFormatPng.cpp
+++ b/Engine/Source/Editor/LandscapeEditor/Private/LandscapeFileFormatPng.cpp
@@ -248,3 +248,5 @@ void FLandscapeWeightmapFileFormat_Png::Export(const TCHAR* WeightmapFilename, F
 		FFileHelper::SaveArrayToFile(TempData, WeightmapFilename);
 	}
 }
+
+#undef LOCTEXT_NAMESPACE
diff --git a/Engine/Source/Editor/LandscapeEditor/Private/LandscapeFileFormatRaw.cpp b/Engine/Source/Editor/LandscapeEditor/Private/LandscapeFileFormatRaw.cpp
--- a/Engine/Source/Editor/LandscapeEditor/Private/LandscapeFileFormatRaw.cpp
+++ b/Engine/Source/Editor/LandscapeEditor/Private/LandscapeFileFormatRaw.cpp
@@ -171,3 +171,5 @@ void FLandscapeWeightmapFileFormat_Raw::Export(const TCHAR* WeightmapFilename, F
 {
 	FFileHelper::SaveArrayToFile(Data, WeightmapFilename);
 }
+
+#undef LOCTEXT_NAMESPACE
