Added virtual destructors to ILandscapeHeightmapFileFormat and ILandscapeWeightmapFileFormat as they're deleted up by TSharedPointer and causes compile errors on Linux.

diff --git a/Engine/Source/Editor/LandscapeEditor/Public/LandscapeFileFormatInterface.h b/Engine/Source/Editor/LandscapeEditor/Public/LandscapeFileFormatInterface.h
--- a/Engine/Source/Editor/LandscapeEditor/Public/LandscapeFileFormatInterface.h
+++ b/Engine/Source/Editor/LandscapeEditor/Public/LandscapeFileFormatInterface.h
@@ -129,6 +129,11 @@ public:
 	{
 		checkf(0, TEXT("File type hasn't implemented support for heightmap export - %s"), *FPaths::GetExtension(HeightmapFilename, true));
 	}
+
+	/**
+	 * Note: Even though this is an interface class we need a virtual destructor as derived objects are deleted via a pointer to this interface
+	 */
+	virtual ~ILandscapeHeightmapFileFormat() {}
 };
 
 class ILandscapeWeightmapFileFormat
@@ -167,4 +172,9 @@ public:
 	{
 		checkf(0, TEXT("File type hasn't implemented support for weightmap export - %s"), *FPaths::GetExtension(WeightmapFilename, true));
 	}
+
+	/**
+ 	 * Note: Even though this is an interface class we need a virtual destructor as derived objects are deleted via a pointer to this interface
+	 */
+	virtual ~ILandscapeWeightmapFileFormat() {}
 };
