Merging //UE4/Dev-Main to Dev-VREditor (from CL 3126397) (part 2)
- Added some missing files

diff --git a/Engine/Source/Editor/LandscapeEditor/Private/LandscapeEdMode.cpp b/Engine/Source/Editor/LandscapeEditor/Private/LandscapeEdMode.cpp
--- a/Engine/Source/Editor/LandscapeEditor/Private/LandscapeEdMode.cpp
+++ b/Engine/Source/Editor/LandscapeEditor/Private/LandscapeEdMode.cpp
@@ -129,10 +129,7 @@ void ALandscape::SplitHeightmap(ULandscapeComponent* Comp, bool bMoveToCurrentLe
 		Comp->HeightmapScaleBias = FVector4(1.0f / (float)HeightmapSizeU, 1.0f / (float)HeightmapSizeV, 0.0f, 0.0f);
 		Comp->HeightmapTexture = HeightmapTexture;
 
-		if (Comp->MaterialInstance)
-		{
-			Comp->MaterialInstance->SetTextureParameterValueEditorOnly(FName(TEXT("Heightmap")), HeightmapTexture);
-		}
+		Comp->UpdateMaterialInstances();
 
 		for (int32 i = 0; i < HeightmapTextureMipData.Num(); i++)
 		{
@@ -180,7 +177,7 @@ FEdModeLandscape::FEdModeLandscape()
 	, SplinesTool(nullptr)
 	, LandscapeRenderAddCollision(nullptr)
 	, CachedLandscapeMaterial(nullptr)
-	, bToolActive(false)
+	, ToolActiveViewport(nullptr)
 	, bIsPainting(false)
 {
 	GLayerDebugColorMaterial = LandscapeTool::CreateMaterialInstance(LoadObject<UMaterial>(nullptr, TEXT("/Engine/EditorLandscapeResources/LayerVisMaterial.LayerVisMaterial")));
@@ -320,7 +317,7 @@ void FEdModeLandscape::Enter()
 	ALandscapeProxy* SelectedLandscape = GEditor->GetSelectedActors()->GetTop<ALandscapeProxy>();
 	if (SelectedLandscape)
 	{
-		CurrentToolTarget.LandscapeInfo = SelectedLandscape->GetLandscapeInfo(false);
+		CurrentToolTarget.LandscapeInfo = SelectedLandscape->GetLandscapeInfo();
 		GEditor->SelectNone(false, true);
 		GEditor->SelectActor(SelectedLandscape, true, false);
 	}
@@ -586,7 +583,10 @@ void FEdModeLandscape::OnVRHoverUpdate(FEditorViewportClient& ViewportClient, UV
 				{
 					if (bIsPainting && !(CurrentToolTarget.TargetType == ELandscapeToolTargetType::Weightmap && CurrentToolTarget.LayerInfo == NULL))
 					{
-						bToolActive = CurrentTool->BeginTool(&ViewportClient, CurrentToolTarget, HitLocation, Interactor);
+						if( CurrentTool->BeginTool(&ViewportClient, CurrentToolTarget, HitLocation, Interactor) )
+						{
+							ToolActiveViewport = ViewportClient.Viewport;
+						}
 					}
 
 					if (CurrentBrush)
@@ -615,10 +615,10 @@ void FEdModeLandscape::OnVRAction(FEditorViewportClient& ViewportClient, UViewpo
 			// Begin landscape brush
 			if (Action.Event == IE_Pressed && !VRInteractor->IsHoveringOverUI())
 			{
-				if (bToolActive)
+				if (ViewportClient.Viewport != nullptr && ViewportClient.Viewport == ToolActiveViewport)
 				{
 					CurrentTool->EndTool(&ViewportClient);
-					bToolActive = false;
+					ToolActiveViewport = nullptr;
 				}
 
 				if (CurrentTool && (CurrentTool->GetSupportedTargetTypes() == ELandscapeToolTargetTypeMask::NA || CurrentToolTarget.TargetType != ELandscapeToolTargetType::Invalid))
@@ -631,7 +631,10 @@ void FEdModeLandscape::OnVRAction(FEditorViewportClient& ViewportClient, UViewpo
 						{
 							if (!(CurrentToolTarget.TargetType == ELandscapeToolTargetType::Weightmap && CurrentToolTarget.LayerInfo == NULL))
 							{
-								bToolActive = CurrentTool->BeginTool(&ViewportClient, CurrentToolTarget, HitLocation, Interactor);
+								if( CurrentTool->BeginTool( &ViewportClient, CurrentToolTarget, HitLocation, Interactor ) )
+								{
+									ToolActiveViewport = ViewportClient.Viewport;
+								}
 							}
 
 							bIsPainting = true;
@@ -645,10 +648,10 @@ void FEdModeLandscape::OnVRAction(FEditorViewportClient& ViewportClient, UViewpo
 			// End landscape brush
 			else if (Action.Event == IE_Released)
 			{
-				if (CurrentTool && bToolActive)
+				if (CurrentTool && ViewportClient.Viewport != nullptr && ViewportClient.Viewport == ToolActiveViewport)
 				{
 					CurrentTool->EndTool(&ViewportClient);
-					bToolActive = false;
+					ToolActiveViewport = nullptr;
 				}
 				bIsPainting = false;
 			}
@@ -669,7 +672,7 @@ void FEdModeLandscape::Tick(FEditorViewportClient* ViewportClient, float DeltaTi
 
 	FViewport* const Viewport = ViewportClient->Viewport;
 
-	if (bToolActive && ensure(CurrentTool))
+	if (ToolActiveViewport && ToolActiveViewport == Viewport && ensure(CurrentTool))
 	{
 		// Require Ctrl or not as per user preference
 		const ELandscapeFoliageEditorControlType LandscapeEditorControlType = GetDefault<ULevelEditorViewportSettings>()->LandscapeEditorControlType;
@@ -679,7 +682,7 @@ void FEdModeLandscape::Tick(FEditorViewportClient* ViewportClient, float DeltaTi
 		{
 			CurrentTool->EndTool(ViewportClient);
 			Viewport->CaptureMouse(false);
-			bToolActive = false;
+			ToolActiveViewport = nullptr;
 		}
 	}
 
@@ -727,17 +730,20 @@ void FEdModeLandscape::Tick(FEditorViewportClient* ViewportClient, float DeltaTi
 /** FEdMode: Called when the mouse is moved over the viewport */
 bool FEdModeLandscape::MouseMove(FEditorViewportClient* ViewportClient, FViewport* Viewport, int32 MouseX, int32 MouseY)
 {
-	if (bToolActive && ensure(CurrentTool))
+	// due to mouse capture this should only ever be called on the active viewport
+	// if it ever gets called on another viewport the mouse has been released without us picking it up
+	if (ToolActiveViewport && ensure(CurrentTool))
 	{
 		// Require Ctrl or not as per user preference
 		const ELandscapeFoliageEditorControlType LandscapeEditorControlType = GetDefault<ULevelEditorViewportSettings>()->LandscapeEditorControlType;
 
-		if (!Viewport->KeyState(EKeys::LeftMouseButton) ||
+		if (ToolActiveViewport != Viewport ||
+			!Viewport->KeyState(EKeys::LeftMouseButton) ||
 			(LandscapeEditorControlType == ELandscapeFoliageEditorControlType::RequireCtrl && !IsCtrlDown(Viewport)))
 		{
 			CurrentTool->EndTool(ViewportClient);
 			Viewport->CaptureMouse(false);
-			bToolActive = false;
+			ToolActiveViewport = nullptr;
 		}
 	}
 
@@ -761,7 +767,7 @@ bool FEdModeLandscape::MouseMove(FEditorViewportClient* ViewportClient, FViewpor
 bool FEdModeLandscape::DisallowMouseDeltaTracking() const
 {
 	// We never want to use the mouse delta tracker while painting
-	return bToolActive;
+	return (ToolActiveViewport != nullptr);
 }
 
 /**
@@ -1353,11 +1359,11 @@ bool FEdModeLandscape::InputKey(FEditorViewportClient* ViewportClient, FViewport
 		if (Key == EKeys::LeftMouseButton && Event == IE_Pressed)
 		{
 			// When debugging it's possible to miss the "mouse released" event, if we get a "mouse pressed" event when we think it's already pressed then treat it as release first
-			if (bToolActive)
+			if (ToolActiveViewport)
 			{
 				CurrentTool->EndTool(ViewportClient);
 				Viewport->CaptureMouse(false);
-				bToolActive = false;
+				ToolActiveViewport = nullptr;
 			}
 
 			// Only activate tool if we're not already moving the camera and we're not trying to drag a transform widget
@@ -1388,9 +1394,14 @@ bool FEdModeLandscape::InputKey(FEditorViewportClient* ViewportClient, FViewport
 						else
 						{
 							Viewport->CaptureMouse(true);
-							bToolActive = CurrentTool->BeginTool(ViewportClient, CurrentToolTarget, HitLocation);
-							if (!bToolActive)
+							bool bToolActive = CurrentTool->BeginTool(ViewportClient, CurrentToolTarget, HitLocation);
+							if (bToolActive)
+							{
+								ToolActiveViewport = Viewport;
+							}
+							else
 							{
+								ToolActiveViewport = nullptr;
 								Viewport->CaptureMouse(false);
 							}
 							ViewportClient->Invalidate(false, false);
@@ -1405,13 +1416,13 @@ bool FEdModeLandscape::InputKey(FEditorViewportClient* ViewportClient, FViewport
 		if (Key == EKeys::LeftMouseButton ||
 			(LandscapeEditorControlType == ELandscapeFoliageEditorControlType::RequireCtrl && (Key == EKeys::LeftControl || Key == EKeys::RightControl)))
 		{
-			if (Event == IE_Released && CurrentTool && bToolActive)
+			if (Event == IE_Released && CurrentTool && ToolActiveViewport)
 			{
 				//Set the cursor position to that of the slate cursor so it wont snap back
 				Viewport->SetPreCaptureMousePosFromSlateCursor();
 				CurrentTool->EndTool(ViewportClient);
 				Viewport->CaptureMouse(false);
-				bToolActive = false;
+				ToolActiveViewport = nullptr;
 				return true;
 			}
 		}
@@ -1468,14 +1479,14 @@ bool FEdModeLandscape::InputKey(FEditorViewportClient* ViewportClient, FViewport
 			return true;
 		}
 
-		// Prev tool 
+		// Prev tool
 		if (Event == IE_Pressed && Key == EKeys::Comma)
 		{
-			if (CurrentTool && bToolActive)
+			if (CurrentTool && ToolActiveViewport)
 			{
 				CurrentTool->EndTool(ViewportClient);
 				Viewport->CaptureMouse(false);
-				bToolActive = false;
+				ToolActiveViewport = nullptr;
 			}
 
 			int32 OldToolIndex = CurrentToolMode->ValidTools.Find(CurrentTool->GetToolName());
@@ -1485,14 +1496,14 @@ bool FEdModeLandscape::InputKey(FEditorViewportClient* ViewportClient, FViewport
 			return true;
 		}
 
-		// Next tool 
+		// Next tool
 		if (Event == IE_Pressed && Key == EKeys::Period)
 		{
-			if (CurrentTool && bToolActive)
+			if (CurrentTool && ToolActiveViewport)
 			{
 				CurrentTool->EndTool(ViewportClient);
 				Viewport->CaptureMouse(false);
-				bToolActive = false;
+				ToolActiveViewport = nullptr;
 			}
 
 			int32 OldToolIndex = CurrentToolMode->ValidTools.Find(CurrentTool->GetToolName());
@@ -2528,171 +2539,377 @@ void FEdModeLandscape::ReimportData(const FLandscapeTargetListInfo& TargetInfo)
 void FEdModeLandscape::ImportData(const FLandscapeTargetListInfo& TargetInfo, const FString& Filename)
 {
 	ULandscapeInfo* LandscapeInfo = TargetInfo.LandscapeInfo.Get();
-	if (LandscapeInfo != NULL)
+	int32 MinX, MinY, MaxX, MaxY;
+	if (LandscapeInfo && LandscapeInfo->GetLandscapeExtent(MinX, MinY, MaxX, MaxY))
 	{
-		TArray<uint8> Data;
-		int32 MinX, MinY, MaxX, MaxY;
-		if (FFileHelper::LoadFileToArray(Data, *Filename) &&
-			LandscapeInfo->GetLandscapeExtent(MinX, MinY, MaxX, MaxY))
+		const FLandscapeFileResolution LandscapeResolution = {(uint32)(1 + MaxX - MinX), (uint32)(1 + MaxY - MinY)};
+
+		ILandscapeEditorModule& LandscapeEditorModule = FModuleManager::GetModuleChecked<ILandscapeEditorModule>("LandscapeEditor");
+
+		if (TargetInfo.TargetType == ELandscapeToolTargetType::Heightmap)
 		{
-			const int32 SizeX = (1 + MaxX - MinX);
-			const int32 SizeY = (1 + MaxY - MinY);
+			const ILandscapeHeightmapFileFormat* HeightmapFormat = LandscapeEditorModule.GetHeightmapFormatByExtension(*FPaths::GetExtension(Filename, true));
+
+			if (!HeightmapFormat)
+			{
+				FMessageDialog::Open(EAppMsgType::Ok, NSLOCTEXT("LandscapeEditor.NewLandscape", "Import_UnknownFileType", "File type not recognised"));
+
+				return;
+			}
 
-			IImageWrapperPtr ImageWrapper = nullptr;
-			const TArray<uint8>* RawData = nullptr; // Pointer to actual used data, actual store could be in Data or ImageWrapper
+			FLandscapeFileResolution ImportResolution = {0, 0};
 
-			if (Filename.EndsWith(".png"))
+			const FLandscapeHeightmapInfo HeightmapInfo = HeightmapFormat->Validate(*Filename);
+
+			// display error message if there is one, and abort the import
+			if (HeightmapInfo.ResultCode == ELandscapeImportResult::Error)
 			{
-				IImageWrapperModule& ImageWrapperModule = FModuleManager::LoadModuleChecked<IImageWrapperModule>("ImageWrapper");
-				ImageWrapper = ImageWrapperModule.CreateImageWrapper(EImageFormat::PNG);
+				FMessageDialog::Open(EAppMsgType::Ok, HeightmapInfo.ErrorMessage);
+
+				return;
+			}
 
-				if (ImageWrapper->SetCompressed(Data.GetData(), Data.Num()))
+			// if the file is a raw format with multiple possibly resolutions, only attempt import if one matches the current landscape
+			if (HeightmapInfo.PossibleResolutions.Num() > 1)
+			{
+				if (!HeightmapInfo.PossibleResolutions.Contains(LandscapeResolution))
 				{
-					const int32 ImportSizeX = ImageWrapper->GetWidth();
-					const int32 ImportSizeY = ImageWrapper->GetHeight();
+					FFormatNamedArguments Args;
+					Args.Add(TEXT("LandscapeSizeX"), LandscapeResolution.Width);
+					Args.Add(TEXT("LandscapeSizeY"), LandscapeResolution.Height);
 
-					if (ImportSizeX != SizeX || ImportSizeY != SizeY)
-					{
-						FFormatNamedArguments Args;
-						Args.Add(TEXT("FileSizeX"), ImportSizeX);
-						Args.Add(TEXT("FileSizeY"), ImportSizeY);
-						Args.Add(TEXT("LandscapeSizeX"), SizeX);
-						Args.Add(TEXT("LandscapeSizeY"), SizeY);
+					FMessageDialog::Open(EAppMsgType::Ok,
+						FText::Format(NSLOCTEXT("LandscapeEditor.NewLandscape", "Import_HeightmapSizeMismatchRaw", "The heightmap file does not match the current Landscape extent ({LandscapeSizeX}\u00D7{LandscapeSizeY}), and its exact resolution could not be determined"), Args));
 
-						auto Result = FMessageDialog::Open(EAppMsgType::OkCancel,
-							FText::Format(NSLOCTEXT("LandscapeEditor.NewLandscape", "Import_SizeMismatch", "File's size ({FileSizeX}\u00D7{FileSizeY}) does not match with current Landscape extent ({LandscapeSizeX}\u00D7{LandscapeSizeY}), if you continue it will be padded/clipped to fit"), Args));
+					return;
+				}
+				else
+				{
+					ImportResolution = LandscapeResolution;
+				}
+			}
 
-						if (Result != EAppReturnType::Ok)
-						{
-							return;
-						}
-					}
+			// display warning message if there is one and allow user to cancel
+			if (HeightmapInfo.ResultCode == ELandscapeImportResult::Warning)
+			{
+				auto Result = FMessageDialog::Open(EAppMsgType::OkCancel, HeightmapInfo.ErrorMessage);
 
-					if (TargetInfo.TargetType == ELandscapeToolTargetType::Heightmap)
-					{
-						if (ImageWrapper->GetFormat() != ERGBFormat::Gray)
-						{
-							EAppReturnType::Type Result = FMessageDialog::Open(EAppMsgType::OkCancel, NSLOCTEXT("LandscapeEditor.NewLandscape", "Import_HeightmapFileColorPng", "The Heightmap file appears to be a color png, grayscale is expected. The import *can* continue, but the result may not be what you expect..."));
-							if (Result != EAppReturnType::Ok)
-							{
-								return;
-							}
-						}
-						else if (ImageWrapper->GetBitDepth() != 16)
-						{
-							EAppReturnType::Type Result = FMessageDialog::Open(EAppMsgType::OkCancel, NSLOCTEXT("LandscapeEditor.NewLandscape", "Import_HeightmapFileLowBitDepth", "The Heightmap file appears to be an 8-bit png, 16-bit is preferred. The import *can* continue, but the result may be lower quality than desired."));
-							if (Result != EAppReturnType::Ok)
-							{
-								return;
-							}
-						}
+				if (Result != EAppReturnType::Ok)
+				{
+					return;
+				}
+			}
 
-						if (ImageWrapper->GetBitDepth() <= 8)
-						{
-							ImageWrapper->GetRaw(ERGBFormat::Gray, 8, RawData);
-							Data.Reset();
-							Data.SetNumUninitialized(RawData->Num() * sizeof(uint16));
-							uint16* DataPtr = (uint16*)Data.GetData();
-							for (int32 i = 0; i < RawData->Num(); i++)
-							{
-								DataPtr[i] = (*RawData)[i] * 0x101; // Expand to 16-bit
-							}
-							RawData = &Data;
-						}
-						else
-						{
-							ImageWrapper->GetRaw(ERGBFormat::Gray, 16, RawData);
-						}
-					}
-					else
-					{
-						if (ImageWrapper->GetFormat() != ERGBFormat::Gray)
-						{
-							EAppReturnType::Type Result = FMessageDialog::Open(EAppMsgType::OkCancel, NSLOCTEXT("LandscapeEditor.NewLandscape", "Import_LayerColorPng", "The Layer file appears to be a color png, grayscale is expected. The import *can* continue, but the result may not be what you expect..."));
-							if (Result != EAppReturnType::Ok)
-							{
-								return;
-							}
-						}
+			// if the file is a format with resolution information, warn the user if the resolution doesn't match the current landscape
+			// unlike for raw this is only a warning as we can pad/clip the data if we know what resolution it is
+			if (HeightmapInfo.PossibleResolutions.Num() == 1)
+			{
+				ImportResolution = HeightmapInfo.PossibleResolutions[0];
+				if (ImportResolution != LandscapeResolution)
+				{
+					FFormatNamedArguments Args;
+					Args.Add(TEXT("FileSizeX"), ImportResolution.Width);
+					Args.Add(TEXT("FileSizeY"), ImportResolution.Height);
+					Args.Add(TEXT("LandscapeSizeX"), LandscapeResolution.Width);
+					Args.Add(TEXT("LandscapeSizeY"), LandscapeResolution.Height);
 
-						ImageWrapper->GetRaw(ERGBFormat::Gray, 8, RawData);
-					}
+					auto Result = FMessageDialog::Open(EAppMsgType::OkCancel,
+						FText::Format(NSLOCTEXT("LandscapeEditor.NewLandscape", "Import_HeightmapSizeMismatch", "The heightmap file's size ({FileSizeX}\u00D7{FileSizeY}) does not match the current Landscape extent ({LandscapeSizeX}\u00D7{LandscapeSizeY}), if you continue it will be padded/clipped to fit"), Args));
 
-					if (ImportSizeX != SizeX || ImportSizeY != SizeY)
+					if (Result != EAppReturnType::Ok)
 					{
-						// Cloned from FLandscapeEditorDetailCustomization_NewLandscape.OnCreateButtonClicked
-						// so that reimports behave the same as the initial import :)
+						return;
+					}
+				}
+			}
 
-						const int32 OffsetX = (SizeX - ImportSizeX) / 2;
-						const int32 OffsetY = (SizeY - ImportSizeY) / 2;
+			FLandscapeHeightmapImportData ImportData = HeightmapFormat->Import(*Filename, ImportResolution);
 
-						if (TargetInfo.TargetType == ELandscapeToolTargetType::Heightmap)
-						{
-							// have to use a TempData because RawData could be pointing to Data
-							TArray<uint8> TempData;
-							TempData.SetNumUninitialized(SizeX * SizeY * sizeof(uint16));
+			if (ImportData.ResultCode == ELandscapeImportResult::Error)
+			{
+				FMessageDialog::Open(EAppMsgType::Ok, ImportData.ErrorMessage);
 
-							LandscapeEditorUtils::ExpandData<uint16>((uint16*)TempData.GetData(), (const uint16*)RawData->GetData(),
-								0, 0, ImportSizeX - 1, ImportSizeY - 1,
-								-OffsetX, -OffsetY, SizeX - OffsetX - 1, SizeY - OffsetY - 1);
+				return;
+			}
 
-							Data = MoveTemp(TempData);
-							RawData = &Data;
-						}
-						else
-						{
-							Data = LandscapeEditorUtils::ExpandData<uint8>(*RawData,
-								0, 0, ImportSizeX - 1, ImportSizeY - 1,
-								-OffsetX, -OffsetY, SizeX - OffsetX - 1, SizeY - OffsetY - 1);
-							RawData = &Data;
-						}
-					}
-				}
-				else
-				{
-					FMessageDialog::Open(EAppMsgType::Ok,
-						FText::Format(LOCTEXT("ImportData_CorruptPngFile", "Import of {0} failed (corrupt png?)"), FText::FromString(Filename)));
-					return;
-				}
+			TArray<uint16> Data;
+			if (ImportResolution != LandscapeResolution)
+			{
+				// Cloned from FLandscapeEditorDetailCustomization_NewLandscape.OnCreateButtonClicked
+				// so that reimports behave the same as the initial import :)
+
+				const int32 OffsetX = (int32)(LandscapeResolution.Width - ImportResolution.Width) / 2;
+				const int32 OffsetY = (int32)(LandscapeResolution.Height - ImportResolution.Height) / 2;
+
+				Data.SetNumUninitialized(LandscapeResolution.Width * LandscapeResolution.Height * sizeof(uint16));
+
+				LandscapeEditorUtils::ExpandData<uint16>(Data.GetData(), ImportData.Data.GetData(),
+					0, 0, ImportResolution.Width - 1, ImportResolution.Height - 1,
+					-OffsetX, -OffsetY, LandscapeResolution.Width - OffsetX - 1, LandscapeResolution.Height - OffsetY - 1);
 			}
 			else
 			{
-				// I would love to deprecate the raw/r8/r16 support for landscape, r16 doesn't even handle endianness issues...
-				RawData = &Data;
+				Data = MoveTemp(ImportData.Data);
 			}
 
-			if (TargetInfo.TargetType == ELandscapeToolTargetType::Heightmap)
+			FScopedTransaction Transaction(LOCTEXT("Undo_ImportHeightmap", "Importing Landscape Heightmap"));
+
+			FHeightmapAccessor<false> HeightmapAccessor(LandscapeInfo);
+			HeightmapAccessor.SetData(MinX, MinY, MaxX, MaxY, Data.GetData());
+		}
+		else
+		{
+			const ILandscapeWeightmapFileFormat* WeightmapFormat = LandscapeEditorModule.GetWeightmapFormatByExtension(*FPaths::GetExtension(Filename, true));
+
+			if (!WeightmapFormat)
 			{
-				if (RawData->Num() == SizeX * SizeY * sizeof(uint16))
+				FMessageDialog::Open(EAppMsgType::Ok, NSLOCTEXT("LandscapeEditor.NewLandscape", "Import_UnknownFileType", "File type not recognised"));
+
+				return;
+			}
+
+			FLandscapeFileResolution ImportResolution = {0, 0};
+
+			const FLandscapeWeightmapInfo WeightmapInfo = WeightmapFormat->Validate(*Filename, TargetInfo.LayerName);
+
+			// display error message if there is one, and abort the import
+			if (WeightmapInfo.ResultCode == ELandscapeImportResult::Error)
+			{
+				FMessageDialog::Open(EAppMsgType::Ok, WeightmapInfo.ErrorMessage);
+
+				return;
+			}
+
+			// if the file is a raw format with multiple possibly resolutions, only attempt import if one matches the current landscape
+			if (WeightmapInfo.PossibleResolutions.Num() > 1)
+			{
+				if (!WeightmapInfo.PossibleResolutions.Contains(LandscapeResolution))
 				{
-					FScopedTransaction Transaction(LOCTEXT("Undo_ImportHeightmap", "Importing Landscape Heightmap"));
+					FFormatNamedArguments Args;
+					Args.Add(TEXT("LandscapeSizeX"), LandscapeResolution.Width);
+					Args.Add(TEXT("LandscapeSizeY"), LandscapeResolution.Height);
+
+					FMessageDialog::Open(EAppMsgType::Ok,
+						FText::Format(NSLOCTEXT("LandscapeEditor.NewLandscape", "Import_LayerSizeMismatch_ResNotDetermined", "The layer file does not match the current Landscape extent ({LandscapeSizeX}\u00D7{LandscapeSizeY}), and its exact resolution could not be determined"), Args));
 
-					FHeightmapAccessor<false> HeightmapAccessor(LandscapeInfo);
-					HeightmapAccessor.SetData(MinX, MinY, MaxX, MaxY, (const uint16*)RawData->GetData());
+					return;
 				}
 				else
 				{
-					FMessageDialog::Open(EAppMsgType::Ok,
-						FText::Format(NSLOCTEXT("UnrealEd", "LandscapeReImport_BadFileSize", "{0}'s filesize does not match with current Landscape extent"), FText::FromString(Filename)));
+					ImportResolution = LandscapeResolution;
 				}
 			}
-			else
+
+			// display warning message if there is one and allow user to cancel
+			if (WeightmapInfo.ResultCode == ELandscapeImportResult::Warning)
 			{
-				if (RawData->Num() == SizeX * SizeY)
+				auto Result = FMessageDialog::Open(EAppMsgType::OkCancel, WeightmapInfo.ErrorMessage);
+
+				if (Result != EAppReturnType::Ok)
 				{
-					FScopedTransaction Transaction(LOCTEXT("Undo_ImportWeightmap", "Importing Landscape Layer"));
+					return;
+				}
+			}
 
-					FAlphamapAccessor<false, false> AlphamapAccessor(LandscapeInfo, TargetInfo.LayerInfoObj.Get());
-					AlphamapAccessor.SetData(MinX, MinY, MaxX, MaxY, RawData->GetData(), ELandscapeLayerPaintingRestriction::None);
+			// if the file is a format with resolution information, warn the user if the resolution doesn't match the current landscape
+			// unlike for raw this is only a warning as we can pad/clip the data if we know what resolution it is
+			if (WeightmapInfo.PossibleResolutions.Num() == 1)
+			{
+				ImportResolution = WeightmapInfo.PossibleResolutions[0];
+				if (ImportResolution != LandscapeResolution)
+				{
+					FFormatNamedArguments Args;
+					Args.Add(TEXT("FileSizeX"), ImportResolution.Width);
+					Args.Add(TEXT("FileSizeY"), ImportResolution.Height);
+					Args.Add(TEXT("LandscapeSizeX"), LandscapeResolution.Width);
+					Args.Add(TEXT("LandscapeSizeY"), LandscapeResolution.Height);
+
+					auto Result = FMessageDialog::Open(EAppMsgType::OkCancel,
+						FText::Format(NSLOCTEXT("LandscapeEditor.NewLandscape", "Import_LayerSizeMismatch_WillClamp", "The layer file's size ({FileSizeX}\u00D7{FileSizeY}) does not match the current Landscape extent ({LandscapeSizeX}\u00D7{LandscapeSizeY}), if you continue it will be padded/clipped to fit"), Args));
+
+					if (Result != EAppReturnType::Ok)
+					{
+						return;
+					}
 				}
-				else
+			}
+
+			FLandscapeWeightmapImportData ImportData = WeightmapFormat->Import(*Filename, TargetInfo.LayerName, ImportResolution);
+
+			if (ImportData.ResultCode == ELandscapeImportResult::Error)
+			{
+				FMessageDialog::Open(EAppMsgType::Ok, ImportData.ErrorMessage);
+
+				return;
+			}
+
+			TArray<uint8> Data;
+			if (ImportResolution != LandscapeResolution)
+			{
+				// Cloned from FLandscapeEditorDetailCustomization_NewLandscape.OnCreateButtonClicked
+				// so that reimports behave the same as the initial import :)
+
+				const int32 OffsetX = (int32)(LandscapeResolution.Width - ImportResolution.Width) / 2;
+				const int32 OffsetY = (int32)(LandscapeResolution.Height - ImportResolution.Height) / 2;
+
+				Data.SetNumUninitialized(LandscapeResolution.Width * LandscapeResolution.Height * sizeof(uint8));
+
+				LandscapeEditorUtils::ExpandData<uint8>(Data.GetData(), ImportData.Data.GetData(),
+					0, 0, ImportResolution.Width - 1, ImportResolution.Height - 1,
+					-OffsetX, -OffsetY, LandscapeResolution.Width - OffsetX - 1, LandscapeResolution.Height - OffsetY - 1);
+			}
+			else
+			{
+				Data = MoveTemp(ImportData.Data);
+			}
+
+			FScopedTransaction Transaction(LOCTEXT("Undo_ImportWeightmap", "Importing Landscape Layer"));
+
+			FAlphamapAccessor<false, false> AlphamapAccessor(LandscapeInfo, TargetInfo.LayerInfoObj.Get());
+			AlphamapAccessor.SetData(MinX, MinY, MaxX, MaxY, Data.GetData(), ELandscapeLayerPaintingRestriction::None);
+		}
+	}
+}
+
+void FEdModeLandscape::DeleteLandscapeComponents(ULandscapeInfo* LandscapeInfo, TSet<ULandscapeComponent*> ComponentsToDelete)
+{
+	LandscapeInfo->Modify();
+	ALandscapeProxy* Proxy = LandscapeInfo->GetLandscapeProxy();
+	Proxy->Modify();
+
+	for (ULandscapeComponent* Component : ComponentsToDelete)
+	{
+		Component->Modify();
+		ULandscapeHeightfieldCollisionComponent* CollisionComp = Component->CollisionComponent.Get();
+		if (CollisionComp)
+		{
+			CollisionComp->Modify();
+		}
+	}
+
+	int32 ComponentSizeVerts = LandscapeInfo->ComponentNumSubsections * (LandscapeInfo->SubsectionSizeQuads + 1);
+	int32 NeedHeightmapSize = 1 << FMath::CeilLogTwo(ComponentSizeVerts);
+
+	TSet<ULandscapeComponent*> HeightmapUpdateComponents;
+	// Need to split all the component which share Heightmap with selected components
+	// Search neighbor only
+	for (ULandscapeComponent* Component : ComponentsToDelete)
+	{
+		int32 SearchX = Component->HeightmapTexture->Source.GetSizeX() / NeedHeightmapSize;
+		int32 SearchY = Component->HeightmapTexture->Source.GetSizeY() / NeedHeightmapSize;
+		FIntPoint ComponentBase = Component->GetSectionBase() / Component->ComponentSizeQuads;
+
+		for (int32 Y = 0; Y < SearchY; ++Y)
+		{
+			for (int32 X = 0; X < SearchX; ++X)
+			{
+				// Search for four directions...
+				for (int32 Dir = 0; Dir < 4; ++Dir)
 				{
-					FMessageDialog::Open(EAppMsgType::Ok,
-						FText::Format(NSLOCTEXT("UnrealEd", "LandscapeReImport_BadFileSize", "{0}'s filesize does not match with current Landscape extent"), FText::FromString(Filename)));
+					int32 XDir = (Dir >> 1) ? 1 : -1;
+					int32 YDir = (Dir % 2) ? 1 : -1;
+					ULandscapeComponent* Neighbor = LandscapeInfo->XYtoComponentMap.FindRef(ComponentBase + FIntPoint(XDir*X, YDir*Y));
+					if (Neighbor && Neighbor->HeightmapTexture == Component->HeightmapTexture && !HeightmapUpdateComponents.Contains(Neighbor))
+					{
+						Neighbor->Modify();
+						HeightmapUpdateComponents.Add(Neighbor);
+					}
 				}
 			}
 		}
 	}
+
+	// Changing Heightmap format for selected components
+	for (ULandscapeComponent* Component : HeightmapUpdateComponents)
+	{
+		ALandscape::SplitHeightmap(Component, false);
+	}
+
+	// Remove attached foliage
+	for (ULandscapeComponent* Component : ComponentsToDelete)
+	{
+		ULandscapeHeightfieldCollisionComponent* CollisionComp = Component->CollisionComponent.Get();
+		if (CollisionComp)
+		{
+			AInstancedFoliageActor::DeleteInstancesForComponent(Proxy->GetWorld(), CollisionComp);
+		}
+	}
+
+	// Check which ones are need for height map change
+	for (ULandscapeComponent* Component : ComponentsToDelete)
+	{
+		// Reset neighbors LOD information
+		FIntPoint ComponentBase = Component->GetSectionBase() / Component->ComponentSizeQuads;
+		FIntPoint NeighborKeys[8] =
+		{
+			ComponentBase + FIntPoint(-1, -1),
+			ComponentBase + FIntPoint(+0, -1),
+			ComponentBase + FIntPoint(+1, -1),
+			ComponentBase + FIntPoint(-1, +0),
+			ComponentBase + FIntPoint(+1, +0),
+			ComponentBase + FIntPoint(-1, +1),
+			ComponentBase + FIntPoint(+0, +1),
+			ComponentBase + FIntPoint(+1, +1)
+		};
+
+		for (const FIntPoint& NeighborKey : NeighborKeys)
+		{
+			ULandscapeComponent* NeighborComp = LandscapeInfo->XYtoComponentMap.FindRef(NeighborKey);
+			if (NeighborComp && !ComponentsToDelete.Contains(NeighborComp))
+			{
+				NeighborComp->Modify();
+				NeighborComp->InvalidateLightingCache();
+
+				// is this really needed? It can happen multiple times per component!
+				FComponentReregisterContext ReregisterContext(NeighborComp);
+			}
+		}
+
+		// Remove Selected Region in deleted Component
+		for (int32 Y = 0; Y < Component->ComponentSizeQuads; ++Y)
+		{
+			for (int32 X = 0; X < Component->ComponentSizeQuads; ++X)
+			{
+				LandscapeInfo->SelectedRegion.Remove(FIntPoint(X, Y) + Component->GetSectionBase());
+			}
+		}
+
+		if (Component->HeightmapTexture)
+		{
+			Component->HeightmapTexture->SetFlags(RF_Transactional);
+			Component->HeightmapTexture->Modify();
+			Component->HeightmapTexture->MarkPackageDirty();
+			Component->HeightmapTexture->ClearFlags(RF_Standalone); // Remove when there is no reference for this Heightmap...
+		}
+
+		for (int32 i = 0; i < Component->WeightmapTextures.Num(); ++i)
+		{
+			Component->WeightmapTextures[i]->SetFlags(RF_Transactional);
+			Component->WeightmapTextures[i]->Modify();
+			Component->WeightmapTextures[i]->MarkPackageDirty();
+			Component->WeightmapTextures[i]->ClearFlags(RF_Standalone);
+		}
+
+		if (Component->XYOffsetmapTexture)
+		{
+			Component->XYOffsetmapTexture->SetFlags(RF_Transactional);
+			Component->XYOffsetmapTexture->Modify();
+			Component->XYOffsetmapTexture->MarkPackageDirty();
+			Component->XYOffsetmapTexture->ClearFlags(RF_Standalone);
+		}
+
+		ULandscapeHeightfieldCollisionComponent* CollisionComp = Component->CollisionComponent.Get();
+		if (CollisionComp)
+		{
+			CollisionComp->DestroyComponent();
+		}
+		Component->DestroyComponent();
+	}
+
+	// Remove Selection
+	LandscapeInfo->ClearSelectedRegion(true);
+	//EdMode->SetMaskEnable(Landscape->SelectedRegion.Num());
+	GEngine->BroadcastLevelActorListChanged();
 }
 
 ALandscape* FEdModeLandscape::ChangeComponentSetting(int32 NumComponentsX, int32 NumComponentsY, int32 NumSubsections, int32 SubsectionSizeQuads, bool bResample)
@@ -2723,6 +2940,7 @@ ALandscape* FEdModeLandscape::ChangeComponentSetting(int32 NumComponentsX, int32
 			TArray<uint16> HeightData;
 			TArray<FLandscapeImportLayerInfo> ImportLayerInfos;
 			FVector LandscapeOffset = FVector::ZeroVector;
+			FIntPoint LandscapeOffsetQuads = FIntPoint::ZeroValue;
 			float LandscapeScaleFactor = 1.0f;
 
 			int32 NewMinX, NewMinY, NewMaxX, NewMaxY;
@@ -2802,6 +3020,7 @@ ALandscape* FEdModeLandscape::ChangeComponentSetting(int32 NumComponentsX, int32
 
 				// offset landscape to component boundary
 				LandscapeOffset = FVector(NewMinX, NewMinY, 0) * OldLandscapeProxy->GetActorScale();
+				LandscapeOffsetQuads = FIntPoint(NewMinX, NewMinY);
 				NewMinX = 0;
 				NewMinY = 0;
 				NewMaxX = NewVertsX - 1;
@@ -2846,6 +3065,8 @@ ALandscape* FEdModeLandscape::ChangeComponentSetting(int32 NumComponentsX, int32
 			Landscape->bUsedForNavigation = OldLandscapeProxy->bUsedForNavigation;
 			Landscape->MaxPaintedLayersPerComponent = OldLandscapeProxy->MaxPaintedLayersPerComponent;
 
+			Landscape->CreateLandscapeInfo();
+
 			// Clone landscape splines
 			TLazyObjectPtr<ALandscape> OldLandscapeActor = LandscapeInfo->LandscapeActor;
 			if (OldLandscapeActor.IsValid() && OldLandscapeActor->SplineComponent != NULL)
@@ -2881,10 +3102,43 @@ ALandscape* FEdModeLandscape::ChangeComponentSetting(int32 NumComponentsX, int32
 						}
 					}
 				}
+
+				// delete any components that were deleted in the original
+				TSet<ULandscapeComponent*> ComponentsToDelete;
+				for (const TPair<FIntPoint, ULandscapeComponent*>& Entry : NewLandscapeInfo->XYtoComponentMap)
+				{
+					if (!LandscapeInfo->XYtoComponentMap.Contains(Entry.Key))
+					{
+						ComponentsToDelete.Add(Entry.Value);
+					}
+				}
+				if (ComponentsToDelete.Num() > 0)
+				{
+					DeleteLandscapeComponents(NewLandscapeInfo, ComponentsToDelete);
+				}
 			}
 			else
 			{
 				// TODO: remap foliage when not resampling (i.e. when there isn't a 1:1 mapping between old and new component)
+
+				// delete any components that are in areas that were entirely deleted in the original
+				ULandscapeInfo* NewLandscapeInfo = Landscape->GetLandscapeInfo();	
+				TSet<ULandscapeComponent*> ComponentsToDelete;
+				for (const TPair<FIntPoint, ULandscapeComponent*>& Entry : NewLandscapeInfo->XYtoComponentMap)
+				{
+					float OldX = Entry.Key.X * NewComponentSizeQuads + LandscapeOffsetQuads.X;
+					float OldY = Entry.Key.Y * NewComponentSizeQuads + LandscapeOffsetQuads.Y;
+					TSet<ULandscapeComponent*> OverlapComponents;
+					LandscapeInfo->GetComponentsInRegion(OldX, OldY, OldX + NewComponentSizeQuads, OldY + NewComponentSizeQuads, OverlapComponents, false);
+					if (OverlapComponents.Num() == 0)
+					{
+						ComponentsToDelete.Add(Entry.Value);
+					}
+				}
+				if (ComponentsToDelete.Num() > 0)
+				{
+					DeleteLandscapeComponents(NewLandscapeInfo, ComponentsToDelete);
+				}
 			}
 
 			// Delete the old Landscape and all its proxies
diff --git a/Engine/Source/Editor/LandscapeEditor/Private/LandscapeEdMode.h b/Engine/Source/Editor/LandscapeEditor/Private/LandscapeEdMode.h
--- a/Engine/Source/Editor/LandscapeEditor/Private/LandscapeEdMode.h
+++ b/Engine/Source/Editor/LandscapeEditor/Private/LandscapeEdMode.h
@@ -438,8 +438,12 @@ public:
 	void ReimportData(const FLandscapeTargetListInfo& TargetInfo);
 	void ImportData(const FLandscapeTargetListInfo& TargetInfo, const FString& Filename);
 
+	/** Resample landscape to a different resolution or change the component size */
 	ALandscape* ChangeComponentSetting(int32 NumComponentsX, int32 NumComponentsY, int32 InNumSubsections, int32 InSubsectionSizeQuads, bool bResample);
 
+	/** Delete the specified landscape components */
+	void DeleteLandscapeComponents(ULandscapeInfo* LandscapeInfo, TSet<ULandscapeComponent*> ComponentsToDelete);
+
 	TArray<FLandscapeToolMode> LandscapeToolModes;
 	TArray<TUniquePtr<FLandscapeTool>> LandscapeTools;
 	TArray<FLandscapeBrushSet> LandscapeBrushSets;
@@ -460,7 +464,7 @@ private:
 
 	UMaterialInterface* CachedLandscapeMaterial;
 
-	bool bToolActive;
+	const FViewport* ToolActiveViewport;
 
 	FDelegateHandle OnWorldChangeDelegateHandle;
 	FDelegateHandle OnMaterialCompilationFinishedDelegateHandle;
diff --git a/Engine/Source/Editor/LandscapeEditor/Private/LandscapeEdModeSplineTools.cpp b/Engine/Source/Editor/LandscapeEditor/Private/LandscapeEdModeSplineTools.cpp
--- a/Engine/Source/Editor/LandscapeEditor/Private/LandscapeEdModeSplineTools.cpp
+++ b/Engine/Source/Editor/LandscapeEditor/Private/LandscapeEdModeSplineTools.cpp
@@ -632,7 +632,7 @@ public:
 			{
 				if (!ToLandscape)
 				{
-					ULandscapeInfo* ProxyLandscapeInfo = FromProxy->GetLandscapeInfo(false);
+					ULandscapeInfo* ProxyLandscapeInfo = FromProxy->GetLandscapeInfo();
 					check(ProxyLandscapeInfo);
 					ToLandscape = ProxyLandscapeInfo->GetCurrentLevelLandscapeProxy(true);
 					if (!ToLandscape)
@@ -722,7 +722,7 @@ public:
 			{
 				if (!ToLandscape)
 				{
-					ULandscapeInfo* ProxyLandscapeInfo = FromProxy->GetLandscapeInfo(false);
+					ULandscapeInfo* ProxyLandscapeInfo = FromProxy->GetLandscapeInfo();
 					check(ProxyLandscapeInfo);
 					ToLandscape = ProxyLandscapeInfo->GetCurrentLevelLandscapeProxy(true);
 					if (!ToLandscape)
