Copying //UE4/Dev-VREditor to //UE4/Dev-Main (Source: //UE4/Dev-VREditor @ 3149834)

==========================
MAJOR FEATURES + CHANGES
==========================

Change 3012664 on 2016/06/14 by [REDACTED_PERSON]

	VR Editor:
	- First extensibility stage for interaction
	- Fixed world space rotation with gizmo handle

Change 3053756 on 2016/07/18 by [REDACTED_PERSON]

	VREditor: Fixed brush preview present while in foliage mode and hovering at UI

Change 3055507 on 2016/07/19 by [REDACTED_PERSON]

	VREditor: Small clean for a foliage function signature

Change 3055772 on 2016/07/19 by [REDACTED_PERSON]

	VREditor: Fixed mesh paint brush not showing brush preview while in normal editor.

Change 3055785 on 2016/07/19 by [REDACTED_PERSON]

	Changing User Defined Struct Editor to vertical layout, tab collapsed for new users

Change 3057466 on 2016/07/20 by [REDACTED_PERSON]

	VR Editor: Deleted unused/invalid asset files (UE-33318)

Change 3060541 on 2016/07/21 by [REDACTED_PERSON]

	Mesh Paint: Fixed various bugs
	- Fixed brush preview not rendered for lasers/mouse when not actively painting (UE-33554)
	- Fixed 'full press' over actors preventing UI from being clickable (UE-33550)
	- Fixed brush cursor displayed when hovering over UI (including selection bar/close button) (UE-33551)
	- Fixed VR transform gizmo getting in the way of everything while painting (it is now hidden while in mesh paint mode)
	- Fixed not being able to interact with UIs after messing around with mesh paint (UE-33621)

Change 3061537 on 2016/07/22 by [REDACTED_PERSON]

	VREditor: Color Picker in VR implemented with delegate overrides for creation and destruction. Transforms still need adjustment.

Change 3061873 on 2016/07/22 by [REDACTED_PERSON]

	VREditor: Fixed foliage select tool not properly select foliage instanced mesh

Change 3063094 on 2016/07/25 by [REDACTED_PERSON]

	VR Editor :
	- Fixed bug that Laser extends beyond UI when hovered over Selection Bar or Close Button #jira UE-33552
	- Fixed crash when Closing Editor (Alt F4 while in VR mode) #jira UE-32509
	- Fixed crash when enabling VR Editor in editor preferences without a HMD connected
	- Fixed bug if you "slowly press" over UI selection bars or close buttons, nothing happens #jira UE-33553
	- Avatar code refactor to its own actor class #jira UETOOL-812

Change 3064489 on 2016/07/26 by [REDACTED_PERSON]

	VREditor : - Fix crash when starting VREditor and then changing levels #jira UE-33766

Change 3068328 on 2016/07/28 by [REDACTED_PERSON]

	VR Editor: Landscape editing in VR (from Trung Le)

Change 3071688 on 2016/08/01 by [REDACTED_PERSON]

	[INTEGRATE] VREditor: Don't teleport viewport camera when entering 'forced VR mode' (from CL 3071686)

Change 3093142 on 2016/08/18 by [REDACTED_PERSON]

	[INTEGRATE] Viewport Interaction: Fixed transform gizmo actor being destroyed and re-created every frame when no actor was selected (from CL 3092429)
	- Happened because every frame we would call GetTransformGizmoActor() which would create it on demand, then later in the frame we'd destroy it (UpdateTranformGizmo) after seeing no objects were selected
	- This hurt performance as thousands of mesh instances would be created over time (slower GC and object iteration)
	- This needs to be merged to the VREditor branch. (MikeF will do that later.)

Change 3109859 on 2016/09/01 by [REDACTED_PERSON]

	VR Editor: - VREditor extensibility make sure all worlds have world interaction #jira UETOOL-892
	- FVREditorMode made into a UObject
	- Added EditorWorldManager that holds all worlds with a ViewportWorldInteraction and VREditorMode
	- ViewportWorldInteraction is ticked from UEditorEngine instead of using FTickableEditorObject
	- Fixed transform gizmo being visible on start at the viewport location
	- Fixed Crash when double clicking a level in Content Browser to open it #jira UE-32541
	- Fixed AvatarActor log warning when going into VREditor mode #jira UE-34434

Change 3110060 on 2016/09/01 by [REDACTED_PERSON]

	VR Editor: Added teleport without hitting objects

Change 3125335 on 2016/09/14 by [REDACTED_PERSON]

	VR Editor: Fixed transform gizmo always being visible, even when nothing is selected (UE-35244)
	- This was a regression introduced with some recent changes to mesh paint mode

Change 3131122 on 2016/09/19 by [REDACTED_PERSON]

	VR Editor: - Added new feature to go back to the default  world to meters when scaling the user and pressing the touchpad #jira UE-35266
	- Added world to meters to the SavedEditorState and removed it from the VREditorModeManager
	- Made a get function for the SavedEditorState
	- Changed EditorViewport in UViewportWorldInteraction to a raw pointer instead of TOptional
	- Added by Mike: Option to skip scaling user this frame in ViewportWorldInteraction

Change 3132051 on 2016/09/20 by [REDACTED_PERSON]

	VR Editor: Fix pointer trace goes through Docking anchor if UI panel is behind #jira UE-29391

Change 3133790 on 2016/09/21 by [REDACTED_PERSON]

	VR Editor: Fix Oculus VR Headset A,B,X, and Y buttons not lighting properly to correct buttons #jira UE-30950

Change 3133824 on 2016/09/21 by [REDACTED_PERSON]

	VR Editor: Disabled VR Editor when simulating in editor is on

Change 3137557 on 2016/09/23 by [REDACTED_PERSON]

	VR Editor: Moved UI related classes to the UI subfolder and VREditorTransformGizmo to Gizmo subfolder. Added subfolders to PrivateIncludePaths in Build.cs.

Change 3137637 on 2016/09/23 by [REDACTED_PERSON]

	VR Editor: Updated all copyrights to 2016

Change 3139784 on 2016/09/26 by [REDACTED_PERSON]

	VR Editor: Fixed blurry toolbar icon when "Use Small Icons" preference was enabled (UE-32208)
	- Also, added a higher resolution reference image for this icon (not used yet)

Change 3140420 on 2016/09/26 by [REDACTED_PERSON]

	- Fixed error caused by AExtendingTheEditorGameMode trying to derive from AGameModeBase instead of AGameMode
	- Fixed error caused by AExtendingTheEditorCharacter not including ExtendingTheEditorProjectile.h

Change 3143346 on 2016/09/28 by [REDACTED_PERSON]

	Added a checkbox for disabling startup VR Editor tutorials to the Experimental Editor Settings. On by default.

Change 3144154 on 2016/09/28 by [REDACTED_PERSON]

	Fix for OnMouseButtonUp not resetting SButton Brush state from Pressed to Hovered.

Change 3149820 on 2016/10/04 by [REDACTED_PERSON]

	VREditor: Fixed compilation editor with non-editor targets (UE-36784)

Change 3149834 on 2016/10/04 by [REDACTED_PERSON]

	VREditor: Fix editor startup crash (UE-36783)
	- Recent changes to VREditor caused "LevelEditor" module to become loaded earlier than ever before, which caused various registration problems leading to crashes
	- Changed VREditor to not load "LevelEditor" until VR mode was actually enabled

Change 3150226 on 2016/10/04 by [REDACTED_PERSON]

	VREditor: Fixed a few static analysis warnings

diff --git a/Engine/Source/Editor/LandscapeEditor/Private/LandscapeEdModeComponentTools.cpp b/Engine/Source/Editor/LandscapeEditor/Private/LandscapeEdModeComponentTools.cpp
--- a/Engine/Source/Editor/LandscapeEditor/Private/LandscapeEdModeComponentTools.cpp
+++ b/Engine/Source/Editor/LandscapeEditor/Private/LandscapeEdModeComponentTools.cpp
@@ -14,100 +14,155 @@
 #include "ComponentReregisterContext.h"
 #include "PhysicalMaterials/PhysicalMaterial.h"
 #include "Materials/MaterialExpressionLandscapeVisibilityMask.h"
-#include "Algo/Copy.h"
 
 #define LOCTEXT_NAMESPACE "Landscape"
 
-//
-// FLandscapeToolSelect
-//
 class FLandscapeToolStrokeSelect : public FLandscapeToolStrokeBase
 {
 	bool bInitializedComponentInvert;
-	bool bInvert;
-	bool bNeedsSelectionUpdate;
+	bool bComponentInvert;
 
 public:
 	FLandscapeToolStrokeSelect(FEdModeLandscape* InEdMode, FEditorViewportClient* InViewportClient, const FLandscapeToolTarget& InTarget)
 		: FLandscapeToolStrokeBase(InEdMode, InViewportClient, InTarget)
 		, bInitializedComponentInvert(false)
-		, bNeedsSelectionUpdate(false)
 		, Cache(InTarget)
 	{
 	}
 
-	~FLandscapeToolStrokeSelect()
-	{
-		if (bNeedsSelectionUpdate)
-		{
-			TArray<UObject*> Objects;
-			if (LandscapeInfo)
-			{
-				TSet<ULandscapeComponent*> SelectedComponents = LandscapeInfo->GetSelectedComponents();
-				Objects.Reset(SelectedComponents.Num());
-				Algo::Copy(SelectedComponents, Objects);
-			}
-			FPropertyEditorModule& PropertyModule = FModuleManager::Get().LoadModuleChecked<FPropertyEditorModule>(TEXT("PropertyEditor"));
-			PropertyModule.UpdatePropertyViews(Objects);
-		}
-	}
-
-	void Apply(FEditorViewportClient* ViewportClient, FLandscapeBrush* Brush, const ULandscapeEditorObject* UISettings, const TArray<FLandscapeToolMousePosition>& MousePositions)
+	void Apply(FEditorViewportClient* ViewportClient, FLandscapeBrush* Brush, const ULandscapeEditorObject* UISettings, const TArray<FLandscapeToolInteractorPosition>& InteractorPositions)
 	{
 		if (LandscapeInfo)
 		{
 			LandscapeInfo->Modify();
 
-			// TODO - only retrieve bounds as we don't need the data
-			const FLandscapeBrushData BrushInfo = Brush->ApplyBrush(MousePositions);
-			if (!BrushInfo)
+			// Invert when holding Shift
+			bool bInvert = InteractorPositions[InteractorPositions.Num() - 1].bModifierPressed;
+
+			if (Brush->GetBrushType() == ELandscapeBrushType::Component)
 			{
-				return;
-			}
+				// TODO - only retrieve bounds as we don't need the data
+				const FLandscapeBrushData BrushInfo = Brush->ApplyBrush(InteractorPositions);
+				if (!BrushInfo)
+				{
+					return;
+				}
 
-			int32 X1, Y1, X2, Y2;
-			BrushInfo.GetInclusiveBounds(X1, Y1, X2, Y2);
+				int32 X1, Y1, X2, Y2;
+				BrushInfo.GetInclusiveBounds(X1, Y1, X2, Y2);
 
-			// Shrink bounds by 1,1 to avoid GetComponentsInRegion picking up extra components on all sides due to the overlap between components
-			TSet<ULandscapeComponent*> NewComponents;
-			LandscapeInfo->GetComponentsInRegion(X1 + 1, Y1 + 1, X2 - 1, Y2 - 1, NewComponents);
+				// Shrink bounds by 1,1 to avoid GetComponentsInRegion picking up extra components on all sides due to the overlap between components
+				TSet<ULandscapeComponent*> NewComponents;
+				LandscapeInfo->GetComponentsInRegion(X1 + 1, Y1 + 1, X2 - 1, Y2 - 1, NewComponents);
 
-			if (!bInitializedComponentInvert)
-			{
-				// Get the component under the mouse location. Copied from FLandscapeBrushComponent::ApplyBrush()
-				const float MouseX = MousePositions[0].Position.X;
-				const float MouseY = MousePositions[0].Position.Y;
-				const int32 MouseComponentIndexX = (MouseX >= 0.0f) ? FMath::FloorToInt(MouseX / LandscapeInfo->ComponentSizeQuads) : FMath::CeilToInt(MouseX / LandscapeInfo->ComponentSizeQuads);
-				const int32 MouseComponentIndexY = (MouseY >= 0.0f) ? FMath::FloorToInt(MouseY / LandscapeInfo->ComponentSizeQuads) : FMath::CeilToInt(MouseY / LandscapeInfo->ComponentSizeQuads);
-				ULandscapeComponent* MouseComponent = LandscapeInfo->XYtoComponentMap.FindRef(FIntPoint(MouseComponentIndexX, MouseComponentIndexY));
-
-				if (MouseComponent != nullptr)
+				if (!bInitializedComponentInvert)
 				{
-					bInvert = LandscapeInfo->GetSelectedComponents().Contains(MouseComponent);
+					// Get the component under the mouse location. Copied from FLandscapeBrushComponent::ApplyBrush()
+					const float MouseX = InteractorPositions[0].Position.X;
+					const float MouseY = InteractorPositions[0].Position.Y;
+					const int32 MouseComponentIndexX = (MouseX >= 0.0f) ? FMath::FloorToInt(MouseX / LandscapeInfo->ComponentSizeQuads) : FMath::CeilToInt(MouseX / LandscapeInfo->ComponentSizeQuads);
+					const int32 MouseComponentIndexY = (MouseY >= 0.0f) ? FMath::FloorToInt(MouseY / LandscapeInfo->ComponentSizeQuads) : FMath::CeilToInt(MouseY / LandscapeInfo->ComponentSizeQuads);
+					ULandscapeComponent* MouseComponent = LandscapeInfo->XYtoComponentMap.FindRef(FIntPoint(MouseComponentIndexX, MouseComponentIndexY));
+
+					if (MouseComponent != nullptr)
+					{
+						bComponentInvert = LandscapeInfo->GetSelectedComponents().Contains(MouseComponent);
+					}
+					else
+					{
+						bComponentInvert = false;
+					}
+
+					bInitializedComponentInvert = true;
+				}
+
+				bInvert = bComponentInvert;
+
+				TSet<ULandscapeComponent*> NewSelection;
+				if (bInvert)
+				{
+					NewSelection = LandscapeInfo->GetSelectedComponents().Difference(NewComponents);
 				}
 				else
 				{
-					bInvert = false;
+					NewSelection = LandscapeInfo->GetSelectedComponents().Union(NewComponents);
 				}
 
-				bInitializedComponentInvert = true;
-			}
+				LandscapeInfo->Modify();
+				LandscapeInfo->UpdateSelectedComponents(NewSelection);
 
-			TSet<ULandscapeComponent*> NewSelection;
-			if (bInvert)
-			{
-				NewSelection = LandscapeInfo->GetSelectedComponents().Difference(NewComponents);
+				// Update Details tab with selection
+				TArray<UObject*> Objects;
+				Objects.Reset(NewSelection.Num());
+				for (auto It = NewSelection.CreateConstIterator(); It; ++It)
+				{
+					Objects.Add(*It);
+				}
+				FPropertyEditorModule& PropertyModule = FModuleManager::Get().LoadModuleChecked<FPropertyEditorModule>(TEXT("PropertyEditor"));
+				PropertyModule.UpdatePropertyViews(Objects);
 			}
-			else
+			else // Select various shape regions
 			{
-				NewSelection = LandscapeInfo->GetSelectedComponents().Union(NewComponents);
-			}
+				const FLandscapeBrushData BrushInfo = Brush->ApplyBrush(InteractorPositions);
+				if (!BrushInfo)
+				{
+					return;
+				}
 
-			LandscapeInfo->Modify();
-			LandscapeInfo->UpdateSelectedComponents(NewSelection);
+				int32 X1, Y1, X2, Y2;
+				BrushInfo.GetInclusiveBounds(X1, Y1, X2, Y2);
+
+				// Tablet pressure
+				float Pressure = ViewportClient->Viewport->IsPenActive() ? ViewportClient->Viewport->GetTabletPressure() : 1.0f;
 
-			// Update Details tab with selection
-			bNeedsSelectionUpdate = true;
+				Cache.CacheData(X1, Y1, X2, Y2);
+				TArray<uint8> Data;
+				Cache.GetCachedData(X1, Y1, X2, Y2, Data);
+
+				TSet<ULandscapeComponent*> NewComponents;
+				LandscapeInfo->GetComponentsInRegion(X1, Y1, X2, Y2, NewComponents);
+				LandscapeInfo->UpdateSelectedComponents(NewComponents, false);
+
+				for (int32 Y = BrushInfo.GetBounds().Min.Y; Y < BrushInfo.GetBounds().Max.Y; Y++)
+				{
+					const float* BrushScanline = BrushInfo.GetDataPtr(FIntPoint(0, Y));
+					uint8* DataScanline = Data.GetData() + (Y - Y1) * (X2 - X1 + 1) + (0 - X1);
+
+					for (int32 X = BrushInfo.GetBounds().Min.X; X < BrushInfo.GetBounds().Max.X; X++)
+					{
+						const FIntPoint Key = FIntPoint(X, Y);
+						const float BrushValue = BrushScanline[X];
+
+						if (BrushValue > 0.0f && LandscapeInfo->IsValidPosition(X, Y))
+						{
+							float PaintValue = BrushValue * UISettings->ToolStrength * Pressure;
+							float Value = DataScanline[X] / 255.0f;
+							checkSlow(FMath::IsNearlyEqual(Value, LandscapeInfo->SelectedRegion.FindRef(Key), 1 / 255.0f));
+							if (bInvert)
+							{
+								Value = FMath::Max(Value - PaintValue, 0.0f);
+							}
+							else
+							{
+								Value = FMath::Min(Value + PaintValue, 1.0f);
+							}
+							if (Value > 0.0f)
+							{
+								LandscapeInfo->SelectedRegion.Add(Key, Value);
+							}
+							else
+							{
+								LandscapeInfo->SelectedRegion.Remove(Key);
+							}
+
+							DataScanline[X] = FMath::Clamp<int32>(FMath::RoundToInt(Value * 255), 0, 255);
+						}
+					}
+				}
+
+				Cache.SetCachedData(X1, Y1, X2, Y2, Data);
+				Cache.Flush();
+			}
 		}
 	}
 
@@ -115,11 +170,15 @@ protected:
 	FLandscapeDataCache Cache;
 };
 
-class FLandscapeToolSelect : public FLandscapeToolBase<FLandscapeToolStrokeSelect>
+// 
+// FLandscapeToolSelect
+//
+template<class TStrokeClass>
+class FLandscapeToolSelect : public FLandscapeToolBase<TStrokeClass>
 {
 public:
 	FLandscapeToolSelect(FEdModeLandscape* InEdMode)
-		: FLandscapeToolBase<FLandscapeToolStrokeSelect>(InEdMode)
+		: FLandscapeToolBase<TStrokeClass>(InEdMode)
 	{
 	}
 
@@ -141,16 +200,16 @@ public:
 	{
 	}
 
-	void Apply(FEditorViewportClient* ViewportClient, FLandscapeBrush* Brush, const ULandscapeEditorObject* UISettings, const TArray<FLandscapeToolMousePosition>& MousePositions)
+	void Apply(FEditorViewportClient* ViewportClient, FLandscapeBrush* Brush, const ULandscapeEditorObject* UISettings, const TArray<FLandscapeToolInteractorPosition>& InteractorPositions)
 	{
 		if (LandscapeInfo)
 		{
 			LandscapeInfo->Modify();
 
 			// Invert when holding Shift
-			bool bInvert = MousePositions[MousePositions.Num() - 1].bShiftDown;
+			bool bInvert = InteractorPositions[ InteractorPositions.Num() - 1].bModifierPressed;
 
-			const FLandscapeBrushData BrushInfo = Brush->ApplyBrush(MousePositions);
+			const FLandscapeBrushData BrushInfo = Brush->ApplyBrush(InteractorPositions);
 			if (!BrushInfo)
 			{
 				return;
@@ -216,11 +275,12 @@ protected:
 	FLandscapeDataCache Cache;
 };
 
-class FLandscapeToolMask : public FLandscapeToolBase<FLandscapeToolStrokeMask>
+template<class TStrokeClass>
+class FLandscapeToolMask : public FLandscapeToolSelect<TStrokeClass>
 {
 public:
 	FLandscapeToolMask(FEdModeLandscape* InEdMode)
-		: FLandscapeToolBase<FLandscapeToolStrokeMask>(InEdMode)
+		: FLandscapeToolSelect<TStrokeClass>(InEdMode)
 	{
 	}
 
@@ -228,13 +288,8 @@ public:
 	virtual FText GetDisplayName() override { return NSLOCTEXT("UnrealEd", "LandscapeMode_Mask", "Region Selection"); };
 	virtual void SetEditRenderType() override { GLandscapeEditRenderMode = ELandscapeEditRenderMode::SelectRegion | (GLandscapeEditRenderMode & ELandscapeEditRenderMode::BitMaskForMask); }
 	virtual bool SupportsMask() override { return true; }
-
-	virtual ELandscapeToolType GetToolType() override { return ELandscapeToolType::Mask; }
 };
 
-//
-// FLandscapeToolVisibility
-//
 class FLandscapeToolStrokeVisibility : public FLandscapeToolStrokeBase
 {
 public:
@@ -244,13 +299,13 @@ public:
 	{
 	}
 
-	void Apply(FEditorViewportClient* ViewportClient, FLandscapeBrush* Brush, const ULandscapeEditorObject* UISettings, const TArray<FLandscapeToolMousePosition>& MousePositions)
+	void Apply(FEditorViewportClient* ViewportClient, FLandscapeBrush* Brush, const ULandscapeEditorObject* UISettings, const TArray<FLandscapeToolInteractorPosition>& InteractorPositions)
 	{
 		if (LandscapeInfo)
 		{
 			LandscapeInfo->Modify();
 			// Get list of verts to update
-			FLandscapeBrushData BrushInfo = Brush->ApplyBrush(MousePositions);
+			FLandscapeBrushData BrushInfo = Brush->ApplyBrush(InteractorPositions);
 			if (!BrushInfo)
 			{
 				return;
@@ -260,7 +315,7 @@ public:
 			BrushInfo.GetInclusiveBounds(X1, Y1, X2, Y2);
 
 			// Invert when holding Shift
-			bool bInvert = MousePositions[MousePositions.Num() - 1].bShiftDown;
+			bool bInvert = InteractorPositions[InteractorPositions.Num() - 1].bModifierPressed;
 
 			// Tablet pressure
 			float Pressure = ViewportClient->Viewport->IsPenActive() ? ViewportClient->Viewport->GetTabletPressure() : 1.0f;
@@ -295,6 +350,9 @@ protected:
 	FLandscapeVisCache Cache;
 };
 
+// 
+// FLandscapeToolVisibility
+//
 class FLandscapeToolVisibility : public FLandscapeToolBase<FLandscapeToolStrokeVisibility>
 {
 public:
@@ -303,7 +361,7 @@ public:
 	{
 	}
 
-	virtual bool BeginTool(FEditorViewportClient* ViewportClient, const FLandscapeToolTarget& InTarget, const FVector& InHitLocation) override
+	virtual bool BeginTool(FEditorViewportClient* ViewportClient, const FLandscapeToolTarget& InTarget, const FVector& InHitLocation, const UViewportInteractor* Interactor = nullptr) override
 	{
 		ALandscapeProxy* Proxy = InTarget.LandscapeInfo->GetLandscapeProxy();
 		UMaterialInterface* HoleMaterial = Proxy->GetLandscapeHoleMaterial();
@@ -333,9 +391,6 @@ public:
 	}
 };
 
-//
-// FLandscapeToolMoveToLevel
-//
 class FLandscapeToolStrokeMoveToLevel : public FLandscapeToolStrokeBase
 {
 public:
@@ -344,7 +399,7 @@ public:
 	{
 	}
 
-	void Apply(FEditorViewportClient* ViewportClient, FLandscapeBrush* Brush, const ULandscapeEditorObject* UISettings, const TArray<FLandscapeToolMousePosition>& MousePositions)
+	void Apply(FEditorViewportClient* ViewportClient, FLandscapeBrush* Brush, const ULandscapeEditorObject* UISettings, const TArray<FLandscapeToolInteractorPosition>& InteractorPositions)
 	{
 		ALandscape* Landscape = LandscapeInfo ? LandscapeInfo->LandscapeActor.Get() : nullptr;
 
@@ -383,7 +438,7 @@ public:
 			{
 				// Get list of verts to update
 				// TODO - only retrieve bounds as we don't need the data
-				FLandscapeBrushData BrushInfo = Brush->ApplyBrush(MousePositions);
+				FLandscapeBrushData BrushInfo = Brush->ApplyBrush(InteractorPositions);
 				if (!BrushInfo)
 				{
 					return;
@@ -769,6 +824,9 @@ public:
 	}
 };
 
+// 
+// FLandscapeToolMoveToLevel
+//
 class FLandscapeToolMoveToLevel : public FLandscapeToolBase<FLandscapeToolStrokeMoveToLevel>
 {
 public:
@@ -784,9 +842,7 @@ public:
 	virtual bool SupportsMask() override { return false; }
 };
 
-//
-// FLandscapeToolAddComponent
-//
+
 class FLandscapeToolStrokeAddComponent : public FLandscapeToolStrokeBase
 {
 public:
@@ -804,7 +860,7 @@ public:
 		XYOffsetCache.Flush();
 	}
 
-	virtual void Apply(FEditorViewportClient* ViewportClient, FLandscapeBrush* Brush, const ULandscapeEditorObject* UISettings, const TArray<FLandscapeToolMousePosition>& MousePositions)
+	virtual void Apply(FEditorViewportClient* ViewportClient, FLandscapeBrush* Brush, const ULandscapeEditorObject* UISettings, const TArray<FLandscapeToolInteractorPosition>& InteractorPositions)
 	{
 		ALandscapeProxy* Landscape = LandscapeInfo ? LandscapeInfo->GetCurrentLevelLandscapeProxy(true) : nullptr;
 		if (Landscape && EdMode->LandscapeRenderAddCollision)
@@ -813,7 +869,7 @@ public:
 
 			// Get list of verts to update
 			// TODO - only retrieve bounds as we don't need the data
-			FLandscapeBrushData BrushInfo = Brush->ApplyBrush(MousePositions);
+			FLandscapeBrushData BrushInfo = Brush->ApplyBrush(InteractorPositions);
 			if (!BrushInfo)
 			{
 				return;
@@ -964,6 +1020,9 @@ protected:
 	FLandscapeXYOffsetCache<true> XYOffsetCache;
 };
 
+// 
+// FLandscapeToolAddComponent
+//
 class FLandscapeToolAddComponent : public FLandscapeToolBase<FLandscapeToolStrokeAddComponent>
 {
 public:
@@ -993,9 +1052,6 @@ public:
 	}
 };
 
-//
-// FLandscapeToolDeleteComponent
-//
 class FLandscapeToolStrokeDeleteComponent : public FLandscapeToolStrokeBase
 {
 public:
@@ -1004,7 +1060,7 @@ public:
 	{
 	}
 
-	void Apply(FEditorViewportClient* ViewportClient, FLandscapeBrush* Brush, const ULandscapeEditorObject* UISettings, const TArray<FLandscapeToolMousePosition>& MousePositions)
+	void Apply(FEditorViewportClient* ViewportClient, FLandscapeBrush* Brush, const ULandscapeEditorObject* UISettings, const TArray<FLandscapeToolInteractorPosition>& InteractorPositions)
 	{
 		if (LandscapeInfo)
 		{
@@ -1013,7 +1069,7 @@ public:
 			{
 				// Get list of components to delete from brush
 				// TODO - only retrieve bounds as we don't need the vert data
-				FLandscapeBrushData BrushInfo = Brush->ApplyBrush(MousePositions);
+				FLandscapeBrushData BrushInfo = Brush->ApplyBrush(InteractorPositions);
 				if (!BrushInfo)
 				{
 					return;
@@ -1032,6 +1088,9 @@ public:
 	}
 };
 
+// 
+// FLandscapeToolDeleteComponent
+//
 class FLandscapeToolDeleteComponent : public FLandscapeToolBase<FLandscapeToolStrokeDeleteComponent>
 {
 public:
@@ -1047,9 +1106,6 @@ public:
 	virtual bool SupportsMask() override { return false; }
 };
 
-//
-// FLandscapeToolCopy
-//
 template<class ToolTarget>
 class FLandscapeToolStrokeCopy : public FLandscapeToolStrokeBase
 {
@@ -1068,7 +1124,7 @@ public:
 		float Data;
 	};
 
-	void Apply(FEditorViewportClient* ViewportClient, FLandscapeBrush* Brush, const ULandscapeEditorObject* UISettings, const TArray<FLandscapeToolMousePosition>& MousePositions)
+	void Apply(FEditorViewportClient* ViewportClient, FLandscapeBrush* Brush, const ULandscapeEditorObject* UISettings, const TArray<FLandscapeToolInteractorPosition>& InteractorPositions)
 	{
 		//ULandscapeInfo* LandscapeInfo = EdMode->CurrentToolTarget.LandscapeInfo;
 		ALandscapeGizmoActiveActor* Gizmo = EdMode->CurrentGizmoActor.Get();
@@ -1078,7 +1134,7 @@ public:
 
 			// Get list of verts to update
 			// TODO - only retrieve bounds as we don't need the data
-			FLandscapeBrushData BrushInfo = Brush->ApplyBrush(MousePositions);
+			FLandscapeBrushData BrushInfo = Brush->ApplyBrush(InteractorPositions);
 			if (!BrushInfo)
 			{
 				return;
@@ -1355,6 +1411,9 @@ protected:
 	FLandscapeFullWeightCache WeightCache;
 };
 
+// 
+// FLandscapeToolCopy
+//
 template<class ToolTarget>
 class FLandscapeToolCopy : public FLandscapeToolBase<FLandscapeToolStrokeCopy<ToolTarget>>
 {
@@ -1379,7 +1438,7 @@ public:
 		return ELandscapeToolTargetTypeMask::FromType(ToolTarget::TargetType);
 	}
 
-	virtual bool BeginTool(FEditorViewportClient* ViewportClient, const FLandscapeToolTarget& InTarget, const FVector& InHitLocation) override
+	virtual bool BeginTool(FEditorViewportClient* ViewportClient, const FLandscapeToolTarget& InTarget, const FVector& InHitLocation, const UViewportInteractor* Interactor = nullptr) override
 	{
 		this->EdMode->GizmoBrush->Tick(ViewportClient, 0.1f);
 
@@ -1402,9 +1461,6 @@ protected:
 	FLandscapeBrush* BackupCurrentBrush;
 };
 
-//
-// FLandscapeToolPaste
-//
 template<class ToolTarget>
 class FLandscapeToolStrokePaste : public FLandscapeToolStrokeBase
 {
@@ -1417,7 +1473,7 @@ public:
 	{
 	}
 
-	void Apply(FEditorViewportClient* ViewportClient, FLandscapeBrush* Brush, const ULandscapeEditorObject* UISettings, const TArray<FLandscapeToolMousePosition>& MousePositions)
+	void Apply(FEditorViewportClient* ViewportClient, FLandscapeBrush* Brush, const ULandscapeEditorObject* UISettings, const TArray<FLandscapeToolInteractorPosition>& InteractorPositions)
 	{
 		//ULandscapeInfo* LandscapeInfo = EdMode->CurrentToolTarget.LandscapeInfo;
 		ALandscapeGizmoActiveActor* Gizmo = EdMode->CurrentGizmoActor.Get();
@@ -1457,7 +1513,7 @@ public:
 			//LandscapeInfo->Modify();
 
 			// Get list of verts to update
-			FLandscapeBrushData BrushInfo = Brush->ApplyBrush(MousePositions);
+			FLandscapeBrushData BrushInfo = Brush->ApplyBrush(InteractorPositions);
 			if (!BrushInfo)
 			{
 				return;
@@ -1670,6 +1726,9 @@ protected:
 	FLandscapeFullWeightCache WeightCache;
 };
 
+// 
+// FLandscapeToolPaste
+//
 template<class ToolTarget>
 class FLandscapeToolPaste : public FLandscapeToolBase<FLandscapeToolStrokePaste<ToolTarget>>
 {
@@ -1700,7 +1759,7 @@ public:
 		return ELandscapeToolTargetTypeMask::FromType(ToolTarget::TargetType);
 	}
 
-	virtual bool BeginTool(FEditorViewportClient* ViewportClient, const FLandscapeToolTarget& InTarget, const FVector& InHitLocation) override
+	virtual bool BeginTool(FEditorViewportClient* ViewportClient, const FLandscapeToolTarget& InTarget, const FVector& InHitLocation, const UViewportInteractor* Interactor = nullptr) override
 	{
 		this->EdMode->GizmoBrush->Tick(ViewportClient, 0.1f);
 
@@ -1741,9 +1800,6 @@ protected:
 	FLandscapeBrush* BackupCurrentBrush;
 };
 
-//
-// FLandscapeToolCopyPaste
-//
 template<class ToolTarget>
 class FLandscapeToolCopyPaste : public FLandscapeToolPaste<ToolTarget>
 {
@@ -1816,7 +1872,7 @@ void FEdModeLandscape::PasteDataFromGizmo()
 	}
 }
 
-//
+// 
 // FLandscapeToolNewLandscape
 //
 class FLandscapeToolNewLandscape : public FLandscapeTool
@@ -1849,7 +1905,7 @@ public:
 		EdMode->NewLandscapePreviewMode = ENewLandscapePreviewMode::None;
 	}
 
-	virtual bool BeginTool(FEditorViewportClient* ViewportClient, const FLandscapeToolTarget& Target, const FVector& InHitLocation) override
+	virtual bool BeginTool(FEditorViewportClient* ViewportClient, const FLandscapeToolTarget& Target, const FVector& InHitLocation, const UViewportInteractor* Interactor = nullptr) override
 	{
 		// does nothing
 		return false;
@@ -1868,7 +1924,7 @@ public:
 };
 
 
-//
+// 
 // FLandscapeToolResizeLandscape
 //
 class FLandscapeToolResizeLandscape : public FLandscapeTool
@@ -1913,7 +1969,7 @@ public:
 	{
 	}
 
-	virtual bool BeginTool(FEditorViewportClient* ViewportClient, const FLandscapeToolTarget& Target, const FVector& InHitLocation) override
+	virtual bool BeginTool(FEditorViewportClient* ViewportClient, const FLandscapeToolTarget& Target, const FVector& InHitLocation, const UViewportInteractor* Interactor = nullptr) override
 	{
 		// does nothing
 		return false;
@@ -1949,7 +2005,7 @@ void FEdModeLandscape::InitializeTool_ResizeLandscape()
 
 void FEdModeLandscape::InitializeTool_Select()
 {
-	auto Tool_Select = MakeUnique<FLandscapeToolSelect>(this);
+	auto Tool_Select = MakeUnique<FLandscapeToolSelect<FLandscapeToolStrokeSelect>>(this);
 	Tool_Select->ValidBrushes.Add("BrushSet_Component");
 	LandscapeTools.Add(MoveTemp(Tool_Select));
 }
@@ -1977,7 +2033,7 @@ void FEdModeLandscape::InitializeTool_MoveToLevel()
 
 void FEdModeLandscape::InitializeTool_Mask()
 {
-	auto Tool_Mask = MakeUnique<FLandscapeToolMask>(this);
+	auto Tool_Mask = MakeUnique<FLandscapeToolMask<FLandscapeToolStrokeSelect>>(this);
 	Tool_Mask->ValidBrushes.Add("BrushSet_Circle");
 	Tool_Mask->ValidBrushes.Add("BrushSet_Alpha");
 	Tool_Mask->ValidBrushes.Add("BrushSet_Pattern");
diff --git a/Engine/Source/Editor/VREditor/VREditorAvatarActor.cpp b/Engine/Source/Editor/VREditor/VREditorAvatarActor.cpp
--- a/Engine/Source/Editor/VREditor/VREditorAvatarActor.cpp
+++ b/Engine/Source/Editor/VREditor/VREditorAvatarActor.cpp
@@ -1,4 +1,4 @@
-// Copyright 1998-2015 Epic Games, Inc. All Rights Reserved.
+// Copyright 1998-2016 Epic Games, Inc. All Rights Reserved.
 
 #include "VREditorModule.h"
 #include "VREditorAvatarActor.h"
@@ -34,11 +34,17 @@ AVREditorAvatarActor::AVREditorAvatarActor( const FObjectInitializer& ObjectInit
 	bIsDrawingWorldMovementPostProcess( false ),
 	VRMode( nullptr )
 {
-	if (UNLIKELY(IsRunningDedicatedServer()))
+	if (UNLIKELY(IsRunningDedicatedServer()))   // @todo vreditor: Hack to avoid loading font assets in the cooker on Linux
 	{
 		return;
 	}
 
+	{
+		USceneComponent* SceneRootComponent = CreateDefaultSubobject<USceneComponent>( TEXT( "RootComponent" ) );
+		AddOwnedComponent( SceneRootComponent );
+		SetRootComponent( SceneRootComponent );
+	}
+
 	// Give us a head mesh
 	{
 		HeadMeshComponent = CreateDefaultSubobject<UStaticMeshComponent>( TEXT( "HeadMeshComponent" ) );
@@ -200,23 +206,23 @@ AVREditorAvatarActor::~AVREditorAvatarActor()
 	UserScaleIndicatorText = nullptr;
 }
 
-void AVREditorAvatarActor::Init( FVREditorMode* InVRMode  )
+void AVREditorAvatarActor::Init( UVREditorMode* InVRMode  )
 {
 	VRMode = InVRMode;
 
 	// Set the default color for the progress bar
 	{
 		static FName StaticLaserColorName( "LaserColor" );
-		const FLinearColor FixedProgressbarColor = VRMode->GetColor( FVREditorMode::EColors::WorldDraggingColor_TwoHands );
+		const FLinearColor FixedProgressbarColor = VRMode->GetColor( UVREditorMode::EColors::WorldDraggingColor_TwoHands );
 		FixedUserScaleMID->SetVectorParameterValue( StaticLaserColorName, FixedProgressbarColor );
 		TranslucentFixedUserScaleMID->SetVectorParameterValue( StaticLaserColorName, FixedProgressbarColor );
 
-		const FLinearColor CurrentProgressbarColor = VRMode->GetColor( FVREditorMode::EColors::GreenGizmoColor );
+		const FLinearColor CurrentProgressbarColor = VRMode->GetColor( UVREditorMode::EColors::GreenGizmoColor );
 		CurrentUserScaleMID->SetVectorParameterValue( StaticLaserColorName, CurrentProgressbarColor );
 		TranslucentCurrentUserScaleMID->SetVectorParameterValue( StaticLaserColorName, CurrentProgressbarColor );
 	}
 
-	UserScaleIndicatorText->SetTextRenderColor( VRMode->GetColor( FVREditorMode::EColors::WorldDraggingColor_TwoHands ).ToFColor( false ) );
+	UserScaleIndicatorText->SetTextRenderColor( VRMode->GetColor( UVREditorMode::EColors::WorldDraggingColor_TwoHands ).ToFColor( false ) );
 
 	// Tell the grid to stay relative to the rootcomponent
 	const FAttachmentTransformRules AttachmentTransformRules = FAttachmentTransformRules( EAttachmentRule::KeepRelative, true );
@@ -235,8 +241,8 @@ void AVREditorAvatarActor::TickManually( const float DeltaTime )
 		RoomSpaceTransformWithWorldToMetersScaling.SetScale3D( RoomSpaceTransformWithWorldToMetersScaling.GetScale3D( ) * WorldScaleFactorVector );
 
 		// @todo vreditor urgent: Head disabled until we can fix late frame update issue
-#if 0
-		if ( VRMode->IsActuallyUsingVR() && GEngine->HMDDevice.IsValid( ) && GEngine->HMDDevice->IsHeadTrackingAllowed( ) )
+		check( HeadMeshComponent != nullptr );
+		if( false ) // && VRMode->IsActuallyUsingVR() && GEngine->HMDDevice.IsValid() && GEngine->HMDDevice->IsHeadTrackingAllowed() )
 		{
 			HeadMeshComponent->SetVisibility( true );
 
@@ -247,23 +253,24 @@ void AVREditorAvatarActor::TickManually( const float DeltaTime )
 			HeadMeshComponent->SetRelativeTransform( RoomSpaceTransformWithWorldToMetersScaling );
 		}
 		else
-#endif
 		{
 			HeadMeshComponent->SetVisibility( false );
 		}
 	}
 
 	// Scale the grid so that it stays the same size in the tracking space
+	check( WorldMovementGridMeshComponent != nullptr );
 	WorldMovementGridMeshComponent->SetRelativeScale3D( WorldScaleFactorVector * VREd::GridScaleMultiplier->GetFloat() );
 	WorldMovementGridMeshComponent->SetRelativeLocation( WorldScaleFactorVector * FVector( 0.0f, 0.0f, VREd::GridHeightOffset->GetFloat() ) );
 
-	// Update the user scale indicator
+	// Update the user scale indicator //@todo
 	{
 		UVREditorInteractor* LeftHandInteractor = VRMode->GetHandInteractor( EControllerHand::Left );
 		UVREditorInteractor* RightHandInteractor = VRMode->GetHandInteractor( EControllerHand::Right );
-		if ( ( LeftHandInteractor != nullptr && RightHandInteractor != nullptr ) &&
-			 ( ( LeftHandInteractor->GetDraggingMode( ) == EViewportInteractionDraggingMode::World && RightHandInteractor->GetDraggingMode( ) == EViewportInteractionDraggingMode::AssistingDrag ) ||
-			 ( LeftHandInteractor->GetDraggingMode( ) == EViewportInteractionDraggingMode::AssistingDrag && RightHandInteractor->GetDraggingMode( ) == EViewportInteractionDraggingMode::World ) ) )
+
+		if ( LeftHandInteractor != nullptr && RightHandInteractor != nullptr &&
+			 ( ( LeftHandInteractor->GetDraggingMode() == EViewportInteractionDraggingMode::World && RightHandInteractor->GetDraggingMode() == EViewportInteractionDraggingMode::AssistingDrag ) ||
+			   ( LeftHandInteractor->GetDraggingMode() == EViewportInteractionDraggingMode::AssistingDrag && RightHandInteractor->GetDraggingMode() == EViewportInteractionDraggingMode::World ) ) )
 		{
 			// Setting all components to be visible
 			CurrentScaleProgressMeshComponent->SetVisibility( true );
