Copying //UE4/Dev-Editor to //UE4/Dev-Main (Source: //UE4/Dev-Editor @ 3279756)

diff --git a/Engine/Source/Editor/LandscapeEditor/Private/LandscapeEditorDetailCustomization_TargetLayers.cpp b/Engine/Source/Editor/LandscapeEditor/Private/LandscapeEditorDetailCustomization_TargetLayers.cpp
--- a/Engine/Source/Editor/LandscapeEditor/Private/LandscapeEditorDetailCustomization_TargetLayers.cpp
+++ b/Engine/Source/Editor/LandscapeEditor/Private/LandscapeEditorDetailCustomization_TargetLayers.cpp
@@ -38,6 +38,9 @@
 #include "LandscapeRender.h"
 #include "Materials/MaterialExpressionLandscapeVisibilityMask.h"
 #include "LandscapeEdit.h"
+#include "IDetailGroup.h"
+#include "SBoxPanel.h"
+#include "Private/SlateEditorStyle.h"
 
 #define LOCTEXT_NAMESPACE "LandscapeEditor.TargetLayers"
 
@@ -50,6 +53,11 @@ BEGIN_SLATE_FUNCTION_BUILD_OPTIMIZATION
 void FLandscapeEditorDetailCustomization_TargetLayers::CustomizeDetails(IDetailLayoutBuilder& DetailBuilder)
 {
 	TSharedRef<IPropertyHandle> PropertyHandle_PaintingRestriction = DetailBuilder.GetProperty(GET_MEMBER_NAME_CHECKED(ULandscapeEditorObject, PaintingRestriction));
+	TSharedRef<IPropertyHandle> PropertyHandle_TargetDisplayOrder = DetailBuilder.GetProperty(GET_MEMBER_NAME_CHECKED(ULandscapeEditorObject, TargetDisplayOrder));
+	PropertyHandle_TargetDisplayOrder->MarkHiddenByCustomization();
+
+	TSharedRef<IPropertyHandle> PropertyHandle_TargetShowUnusedLayers = DetailBuilder.GetProperty(GET_MEMBER_NAME_CHECKED(ULandscapeEditorObject, ShowUnusedLayers));
+	PropertyHandle_TargetShowUnusedLayers->MarkHiddenByCustomization();	
 
 	if (!ShouldShowTargetLayers())
 	{
@@ -72,7 +80,7 @@ void FLandscapeEditorDetailCustomization_TargetLayers::CustomizeDetails(IDetailL
 		.ErrorText(LOCTEXT("Visibility_Tip","Note: You must add a \"Landscape Visibility Mask\" node to your material before you can paint visibility."))
 	];
 
-	TargetsCategory.AddCustomBuilder(MakeShareable(new FLandscapeEditorCustomNodeBuilder_TargetLayers(DetailBuilder.GetThumbnailPool().ToSharedRef())));
+	TargetsCategory.AddCustomBuilder(MakeShareable(new FLandscapeEditorCustomNodeBuilder_TargetLayers(DetailBuilder.GetThumbnailPool().ToSharedRef(), PropertyHandle_TargetDisplayOrder, PropertyHandle_TargetShowUnusedLayers)));
 }
 END_SLATE_FUNCTION_BUILD_OPTIMIZATION
 
@@ -151,8 +159,10 @@ FEdModeLandscape* FLandscapeEditorCustomNodeBuilder_TargetLayers::GetEditorMode(
 	return (FEdModeLandscape*)GLevelEditorModeTools().GetActiveMode(FBuiltinEditorModes::EM_Landscape);
 }
 
-FLandscapeEditorCustomNodeBuilder_TargetLayers::FLandscapeEditorCustomNodeBuilder_TargetLayers(TSharedRef<FAssetThumbnailPool> InThumbnailPool)
+FLandscapeEditorCustomNodeBuilder_TargetLayers::FLandscapeEditorCustomNodeBuilder_TargetLayers(TSharedRef<FAssetThumbnailPool> InThumbnailPool, TSharedRef<IPropertyHandle> InTargetDisplayOrderPropertyHandle, TSharedRef<IPropertyHandle> InTargetShowUnusedLayersPropertyHandle)
 	: ThumbnailPool(InThumbnailPool)
+	, TargetDisplayOrderPropertyHandle(InTargetDisplayOrderPropertyHandle)
+	, TargetShowUnusedLayersPropertyHandle(InTargetShowUnusedLayersPropertyHandle)
 {
 }
 
@@ -172,7 +182,245 @@ void FLandscapeEditorCustomNodeBuilder_TargetLayers::SetOnRebuildChildren(FSimpl
 
 void FLandscapeEditorCustomNodeBuilder_TargetLayers::GenerateHeaderRowContent(FDetailWidgetRow& NodeRow)
 {
+	FEdModeLandscape* LandscapeEdMode = GetEditorMode();
+	
+	if (LandscapeEdMode == NULL)
+	{
+		return;	
+	}
+
+	NodeRow.NameWidget
+		[
+			SNew(STextBlock)
+			.Font(IDetailLayoutBuilder::GetDetailFont())
+			.Text(FText::FromString(TEXT("Layers")))
+		];
+
+	if (LandscapeEdMode->CurrentToolMode->SupportedTargetTypes & ELandscapeToolTargetTypeMask::Weightmap)
+	{
+		NodeRow.ValueWidget
+		[
+			SNew(SHorizontalBox)
+			
+			+ SHorizontalBox::Slot()
+			.AutoWidth()
+			.VAlign(VAlign_Center)
+			.Padding(0.0f, 0.0f, 0.0f, 0.0f)
+			[
+				SNew(SComboButton)
+				.ComboButtonStyle(FEditorStyle::Get(), "ToolbarComboButton")
+				.ForegroundColor(FSlateColor::UseForeground())
+				.HasDownArrow(true)
+				.ContentPadding(FMargin(1, 0))
+				.VAlign(VAlign_Center)
+				.HAlign(HAlign_Center)
+				.ToolTipText(LOCTEXT("TargetLayerSortButtonTooltip", "Define how we want to sort the displayed layers"))
+				.OnGetMenuContent(this, &FLandscapeEditorCustomNodeBuilder_TargetLayers::GetTargetLayerDisplayOrderButtonMenuContent)
+				.ButtonContent()
+				[
+					SNew(SHorizontalBox)
+					+ SHorizontalBox::Slot()
+					.AutoWidth()
+					.VAlign(VAlign_Center)
+					[
+						SNew( SOverlay )
+						+SOverlay::Slot()
+						[
+							SNew(SImage)
+							.Image(FEditorStyle::GetBrush("LandscapeEditor.Target_DisplayOrder.Default"))
+						]	
+						+SOverlay::Slot()
+						[
+							SNew(SImage)
+							.Image(this, &FLandscapeEditorCustomNodeBuilder_TargetLayers::GetTargetLayerDisplayOrderBrush)
+						]
+					]
+				]
+			]
+
+			+ SHorizontalBox::Slot()
+			.AutoWidth()
+			.VAlign(VAlign_Center)
+			.Padding(5.0f, 0.0f, 0.0f, 0.0f)
+			[
+				SNew(SComboButton)
+				.ComboButtonStyle(FEditorStyle::Get(), "ToolbarComboButton")
+				.ForegroundColor(FSlateColor::UseForeground())
+				.HasDownArrow(true)
+				.ContentPadding(FMargin(1, 0))
+				.VAlign(VAlign_Center)
+				.HAlign(HAlign_Center)
+				.ToolTipText(LOCTEXT("TargetLayerUnusedLayerButtonTooltip", "Define if we want to display unused layers"))
+				.OnGetMenuContent(this, &FLandscapeEditorCustomNodeBuilder_TargetLayers::GetTargetLayerShowUnusedButtonMenuContent)
+				.ButtonContent()
+				[
+					SNew(SHorizontalBox)
+					+ SHorizontalBox::Slot()
+					.AutoWidth()
+					.VAlign(VAlign_Center)
+					[
+						SNew(SBox)
+						.WidthOverride(16.0f)
+						.HeightOverride(16.0f)
+						[
+							SNew(SImage)
+							.Image(FEditorStyle::GetBrush("GenericViewButton"))
+						]
+					]
+				]
+			]
+		];
+	}
+}
+
+TSharedRef<SWidget> FLandscapeEditorCustomNodeBuilder_TargetLayers::GetTargetLayerDisplayOrderButtonMenuContent()
+{
+	FMenuBuilder MenuBuilder(/*bInShouldCloseWindowAfterMenuSelection=*/true, nullptr, nullptr, /*bCloseSelfOnly=*/ true);
+
+	MenuBuilder.BeginSection("TargetLayerSortType", LOCTEXT("SortTypeHeading", "Sort Type"));
+	{
+		MenuBuilder.AddMenuEntry(
+			LOCTEXT("TargetLayerDisplayOrderDefault", "Default"),
+			LOCTEXT("TargetLayerDisplayOrderDefaultToolTip", "Sort using order defined in the material."),
+			FSlateIcon(),
+			FUIAction(
+				FExecuteAction::CreateSP(this, &FLandscapeEditorCustomNodeBuilder_TargetLayers::SetSelectedDisplayOrder, ELandscapeLayerDisplayMode::Default),
+				FCanExecuteAction(),
+				FIsActionChecked::CreateSP(this, &FLandscapeEditorCustomNodeBuilder_TargetLayers::IsSelectedDisplayOrder, ELandscapeLayerDisplayMode::Default)
+			),
+			NAME_None,
+			EUserInterfaceActionType::RadioButton
+		);
+
+		MenuBuilder.AddMenuEntry(
+			LOCTEXT("TargetLayerDisplayOrderAlphabetical", "Alphabetical"),
+			LOCTEXT("TargetLayerDisplayOrderAlphabeticalToolTip", "Sort using alphabetical order."),
+			FSlateIcon(),
+			FUIAction(
+				FExecuteAction::CreateSP(this, &FLandscapeEditorCustomNodeBuilder_TargetLayers::SetSelectedDisplayOrder, ELandscapeLayerDisplayMode::Alphabetical),
+				FCanExecuteAction(),
+				FIsActionChecked::CreateSP(this, &FLandscapeEditorCustomNodeBuilder_TargetLayers::IsSelectedDisplayOrder, ELandscapeLayerDisplayMode::Alphabetical)
+			),
+			NAME_None,
+			EUserInterfaceActionType::RadioButton
+		);
+
+		MenuBuilder.AddMenuEntry(
+			LOCTEXT("TargetLayerDisplayOrderCustom", "Custom"),
+			LOCTEXT("TargetLayerDisplayOrderCustomToolTip", "This sort options will be set when changing manually display order by dragging layers"),
+			FSlateIcon(),
+			FUIAction(
+				FExecuteAction::CreateSP(this, &FLandscapeEditorCustomNodeBuilder_TargetLayers::SetSelectedDisplayOrder, ELandscapeLayerDisplayMode::UserSpecific),
+				FCanExecuteAction(),
+				FIsActionChecked::CreateSP(this, &FLandscapeEditorCustomNodeBuilder_TargetLayers::IsSelectedDisplayOrder, ELandscapeLayerDisplayMode::UserSpecific)
+			),
+			NAME_None,
+			EUserInterfaceActionType::RadioButton
+		);
+	}
+	MenuBuilder.EndSection();
+
+	return MenuBuilder.MakeWidget();
+}
+
+TSharedRef<SWidget> FLandscapeEditorCustomNodeBuilder_TargetLayers::GetTargetLayerShowUnusedButtonMenuContent()
+{
+	FMenuBuilder MenuBuilder(/*bInShouldCloseWindowAfterMenuSelection=*/true, nullptr, nullptr, /*bCloseSelfOnly=*/ true);
+
+	MenuBuilder.BeginSection("TargetLayerUnusedType", LOCTEXT("UnusedTypeHeading", "Layer Visilibity"));
+	{
+		MenuBuilder.AddMenuEntry(
+			LOCTEXT("TargetLayerShowUnusedLayer", "Show all layers"),
+			LOCTEXT("TargetLayerShowUnusedLayerToolTip", "Show all layers"),
+			FSlateIcon(),
+			FUIAction(
+				FExecuteAction::CreateSP(this, &FLandscapeEditorCustomNodeBuilder_TargetLayers::ShowUnusedLayers, true),
+				FCanExecuteAction(),
+				FIsActionChecked::CreateSP(this, &FLandscapeEditorCustomNodeBuilder_TargetLayers::ShouldShowUnusedLayers, true)
+			),
+			NAME_None,
+			EUserInterfaceActionType::RadioButton
+		);
+
+		MenuBuilder.AddMenuEntry(
+			LOCTEXT("TargetLayerHideUnusedLayer", "Hide unused layers"),
+			LOCTEXT("TargetLayerHideUnusedLayerToolTip", "Only show used layer"),
+			FSlateIcon(),
+			FUIAction(
+				FExecuteAction::CreateSP(this, &FLandscapeEditorCustomNodeBuilder_TargetLayers::ShowUnusedLayers, false),
+				FCanExecuteAction(),
+				FIsActionChecked::CreateSP(this, &FLandscapeEditorCustomNodeBuilder_TargetLayers::ShouldShowUnusedLayers, false)
+			),
+			NAME_None,
+			EUserInterfaceActionType::RadioButton
+		);
+	}
+
+	MenuBuilder.EndSection();
+
+	return MenuBuilder.MakeWidget();
+}
+
+void FLandscapeEditorCustomNodeBuilder_TargetLayers::ShowUnusedLayers(bool Result)
+{
+	TargetShowUnusedLayersPropertyHandle->SetValue(Result);
+}
+
+bool FLandscapeEditorCustomNodeBuilder_TargetLayers::ShouldShowUnusedLayers(bool Result) const
+{
+	FEdModeLandscape* LandscapeEdMode = GetEditorMode();
+	if (LandscapeEdMode != nullptr)
+	{
+		return LandscapeEdMode->UISettings->ShowUnusedLayers == Result;
+	}
+
+	return false;
+}
+
+void FLandscapeEditorCustomNodeBuilder_TargetLayers::SetSelectedDisplayOrder(ELandscapeLayerDisplayMode InDisplayOrder)
+{
+	TargetDisplayOrderPropertyHandle->SetValue((uint8)InDisplayOrder);	
+}
+
+bool FLandscapeEditorCustomNodeBuilder_TargetLayers::IsSelectedDisplayOrder(ELandscapeLayerDisplayMode InDisplayOrder) const
+{
+	FEdModeLandscape* LandscapeEdMode = GetEditorMode();
+	if (LandscapeEdMode != nullptr)
+	{
+		return LandscapeEdMode->UISettings->TargetDisplayOrder == InDisplayOrder;
+	}
+
+	return false;
+}
+
+const FSlateBrush* FLandscapeEditorCustomNodeBuilder_TargetLayers::GetTargetLayerDisplayOrderBrush() const
+{
+	FEdModeLandscape* LandscapeEdMode = GetEditorMode();
+	if (LandscapeEdMode != nullptr)
+	{
+		switch (LandscapeEdMode->UISettings->TargetDisplayOrder)
+		{
+			case ELandscapeLayerDisplayMode::Alphabetical: return FEditorStyle::Get().GetBrush("LandscapeEditor.Target_DisplayOrder.Alphabetical");
+			case ELandscapeLayerDisplayMode::UserSpecific: return FEditorStyle::Get().GetBrush("LandscapeEditor.Target_DisplayOrder.Custom");
+		}
+	}
+
+	return nullptr;
+}
+
+EVisibility FLandscapeEditorCustomNodeBuilder_TargetLayers::ShouldShowLayer(TSharedRef<FLandscapeTargetListInfo> Target) const
+{
+	if (Target->TargetType == ELandscapeToolTargetType::Weightmap)
+	{
+		FEdModeLandscape* LandscapeEdMode = GetEditorMode();
 
+		if (LandscapeEdMode != nullptr)
+		{
+			return LandscapeEdMode->ShouldShowLayer(Target) ? EVisibility::Visible : EVisibility::Collapsed;
+		}
+	}
+
+	return EVisibility::Visible;
 }
 
 BEGIN_SLATE_FUNCTION_BUILD_OPTIMIZATION
@@ -181,39 +429,77 @@ void FLandscapeEditorCustomNodeBuilder_TargetLayers::GenerateChildContent(IDetai
 	FEdModeLandscape* LandscapeEdMode = GetEditorMode();
 	if (LandscapeEdMode != NULL)
 	{
-		const TArray<TSharedRef<FLandscapeTargetListInfo>>& TargetsList = LandscapeEdMode->GetTargetList();
+		const TArray<TSharedRef<FLandscapeTargetListInfo>>& TargetList = LandscapeEdMode->GetTargetList();
+		const TArray<FName>* TargetDisplayOrderList = LandscapeEdMode->GetTargetDisplayOrderList();
+		const TArray<FName>& TargetShownLayerList = LandscapeEdMode->GetTargetShownList();
 
-		for (int32 i = 0; i < TargetsList.Num(); i++)
+		if (TargetDisplayOrderList == nullptr)
 		{
-			const TSharedRef<FLandscapeTargetListInfo>& Target = TargetsList[i];
-			GenerateRow(ChildrenBuilder, Target);
+			return;
+		}
+
+		TSharedPtr<SDragAndDropVerticalBox> TargetLayerList = SNew(SDragAndDropVerticalBox)
+			.OnCanAcceptDrop(this, &FLandscapeEditorCustomNodeBuilder_TargetLayers::HandleCanAcceptDrop)
+			.OnAcceptDrop(this, &FLandscapeEditorCustomNodeBuilder_TargetLayers::HandleAcceptDrop)
+			.OnDragDetected(this, &FLandscapeEditorCustomNodeBuilder_TargetLayers::HandleDragDetected);
+
+		TargetLayerList->SetDropIndicator_Above(*FEditorStyle::GetBrush("LandscapeEditor.TargetList.DropZone.Above"));
+		TargetLayerList->SetDropIndicator_Below(*FEditorStyle::GetBrush("LandscapeEditor.TargetList.DropZone.Below"));
+
+		ChildrenBuilder.AddChildContent(FText::FromString(FString(TEXT("Layers"))))
+			.Visibility(EVisibility::Visible)
+			[
+				TargetLayerList.ToSharedRef()
+			];
+
+		for (int32 i = 0; i < TargetDisplayOrderList->Num(); ++i)
+		{
+			for (const TSharedRef<FLandscapeTargetListInfo>& TargetInfo : TargetList)
+			{
+				if (TargetInfo->LayerName == (*TargetDisplayOrderList)[i] && (TargetInfo->TargetType != ELandscapeToolTargetType::Weightmap || TargetShownLayerList.Find(TargetInfo->LayerName) != INDEX_NONE))
+				{
+					TSharedPtr<SWidget> GeneratedRowWidget = GenerateRow(TargetInfo);
+
+					if (GeneratedRowWidget.IsValid())
+					{
+						TargetLayerList->AddSlot()
+						.AutoHeight()						
+						[
+							GeneratedRowWidget.ToSharedRef()
+						];
+					}
+
+					break;
+				}
+			}
 		}
 	}
 }
 END_SLATE_FUNCTION_BUILD_OPTIMIZATION
 
 BEGIN_SLATE_FUNCTION_BUILD_OPTIMIZATION
-void FLandscapeEditorCustomNodeBuilder_TargetLayers::GenerateRow(IDetailChildrenBuilder& ChildrenBuilder, const TSharedRef<FLandscapeTargetListInfo> Target)
+TSharedPtr<SWidget> FLandscapeEditorCustomNodeBuilder_TargetLayers::GenerateRow(const TSharedRef<FLandscapeTargetListInfo> Target)
 {
+	TSharedPtr<SWidget> RowWidget;
+
 	FEdModeLandscape* LandscapeEdMode = GetEditorMode();
 	if (LandscapeEdMode)
 	{
 		if ((LandscapeEdMode->CurrentTool->GetSupportedTargetTypes() & LandscapeEdMode->CurrentToolMode->SupportedTargetTypes & ELandscapeToolTargetTypeMask::FromType(Target->TargetType)) == 0)
 		{
-			return;
+			return RowWidget;
 		}
 	}
 	
 	if (Target->TargetType != ELandscapeToolTargetType::Weightmap)
 	{
-		ChildrenBuilder.AddChildContent(Target->TargetName)
-		[
-			SNew(SLandscapeEditorSelectableBorder)
+		RowWidget = SNew(SLandscapeEditorSelectableBorder)
 			.Padding(0)
 			.VAlign(VAlign_Center)
 			.OnContextMenuOpening_Static(&FLandscapeEditorCustomNodeBuilder_TargetLayers::OnTargetLayerContextMenuOpening, Target)
 			.OnSelected_Static(&FLandscapeEditorCustomNodeBuilder_TargetLayers::OnTargetSelectionChanged, Target)
-			.IsSelected_Static(&FLandscapeEditorCustomNodeBuilder_TargetLayers::GetTargetLayerIsSelected, Target)
+			.IsSelected_Static(&FLandscapeEditorCustomNodeBuilder_TargetLayers::GetTargetLayerIsSelected, Target)			
+			.Visibility(this, &FLandscapeEditorCustomNodeBuilder_TargetLayers::ShouldShowLayer, Target)
 			[
 				SNew(SHorizontalBox)
 				+ SHorizontalBox::Slot()
@@ -240,23 +526,34 @@ void FLandscapeEditorCustomNodeBuilder_TargetLayers::GenerateRow(IDetailChildren
 						.ShadowOffset(FVector2D::UnitVector)
 					]
 				]
-			]
-		];
+			];
 	}
 	else
 	{
 		static const FSlateColorBrush SolidWhiteBrush = FSlateColorBrush(FColorList::White);
 
-		ChildrenBuilder.AddChildContent(Target->TargetName)
-		[
-			SNew(SLandscapeEditorSelectableBorder)
+		RowWidget = SNew(SLandscapeEditorSelectableBorder)
 			.Padding(0)
 			.VAlign(VAlign_Center)
 			.OnContextMenuOpening_Static(&FLandscapeEditorCustomNodeBuilder_TargetLayers::OnTargetLayerContextMenuOpening, Target)
 			.OnSelected_Static(&FLandscapeEditorCustomNodeBuilder_TargetLayers::OnTargetSelectionChanged, Target)
 			.IsSelected_Static(&FLandscapeEditorCustomNodeBuilder_TargetLayers::GetTargetLayerIsSelected, Target)
-			[
+			.Visibility(this, &FLandscapeEditorCustomNodeBuilder_TargetLayers::ShouldShowLayer, Target)
+			[				
 				SNew(SHorizontalBox)
+
+				+ SHorizontalBox::Slot()
+				.AutoWidth()
+				.VAlign(VAlign_Center)
+				[
+					SNew(SBox)
+					.Padding(FMargin(2.0f, 0.0f, 2.0f, 0.0f))
+					[
+						SNew(SImage)
+						.Image(FCoreStyle::Get().GetBrush("VerticalBoxDragIndicator"))
+					]
+				]
+
 				+ SHorizontalBox::Slot()
 				.AutoWidth()
 				.VAlign(VAlign_Center)
@@ -335,7 +632,6 @@ void FLandscapeEditorCustomNodeBuilder_TargetLayers::GenerateRow(IDetailChildren
 							.OnObjectChanged_Static(&FLandscapeEditorCustomNodeBuilder_TargetLayers::OnTargetLayerSetObject, Target)
 							.OnShouldFilterAsset_Static(&FLandscapeEditorCustomNodeBuilder_TargetLayers::ShouldFilterLayerInfo, Target->LayerName)
 							.AllowClear(false)
-							//.DisplayThumbnail(false)
 						]
 						+ SHorizontalBox::Slot()
 						.AutoWidth()
@@ -348,7 +644,7 @@ void FLandscapeEditorCustomNodeBuilder_TargetLayers::GenerateRow(IDetailChildren
 							.ForegroundColor(FSlateColor::UseForeground())
 							.IsFocusable(false)
 							.ToolTipText(LOCTEXT("Tooltip_Create", "Create Layer Info"))
-							.Visibility_Static(&FLandscapeEditorCustomNodeBuilder_TargetLayers::GetTargetLayerCreateVisibility, Target)
+							.IsEnabled_Static(&FLandscapeEditorCustomNodeBuilder_TargetLayers::GetTargetLayerCreateEnabled, Target)
 							.OnGetMenuContent_Static(&FLandscapeEditorCustomNodeBuilder_TargetLayers::OnGetTargetLayerCreateMenu, Target)
 							.ButtonContent()
 							[
@@ -446,12 +742,95 @@ void FLandscapeEditorCustomNodeBuilder_TargetLayers::GenerateRow(IDetailChildren
 						]
 					]
 				]
-			]
-		];
+			];
 	}
+
+	return RowWidget;
 }
 END_SLATE_FUNCTION_BUILD_OPTIMIZATION
 
+FReply FLandscapeEditorCustomNodeBuilder_TargetLayers::HandleDragDetected(const FGeometry& MyGeometry, const FPointerEvent& MouseEvent, int32 SlotIndex, SVerticalBox::FSlot* Slot)
+{
+	FEdModeLandscape* LandscapeEdMode = GetEditorMode();
+
+	if (LandscapeEdMode != nullptr)
+	{
+		const TArray<FName>& TargetShownList = LandscapeEdMode->GetTargetShownList();
+
+		if (TargetShownList.IsValidIndex(SlotIndex))
+		{
+			const TArray<FName>* TargetDisplayOrderList = LandscapeEdMode->GetTargetDisplayOrderList();
+
+			if (TargetDisplayOrderList != nullptr)
+			{
+				int32 DisplayOrderLayerIndex = TargetDisplayOrderList->Find(LandscapeEdMode->UISettings->ShowUnusedLayers ? TargetShownList[SlotIndex + LandscapeEdMode->GetTargetLayerStartingIndex()] : TargetShownList[SlotIndex]);
+
+				if (TargetDisplayOrderList->IsValidIndex(DisplayOrderLayerIndex))
+				{
+					const TArray<TSharedRef<FLandscapeTargetListInfo>>& TargetList = LandscapeEdMode->GetTargetList();
+
+					for (const TSharedRef<FLandscapeTargetListInfo>& TargetInfo : TargetList)
+					{
+						if (TargetInfo->LayerName == (*TargetDisplayOrderList)[DisplayOrderLayerIndex])
+						{
+							return FReply::Handled().BeginDragDrop(FTargetLayerDragDropOp::New(SlotIndex, Slot, GenerateRow(TargetInfo)));
+						}
+					}
+				}
+			}
+		}
+	}
+
+	return FReply::Unhandled();
+}
+
+TOptional<SDragAndDropVerticalBox::EItemDropZone> FLandscapeEditorCustomNodeBuilder_TargetLayers::HandleCanAcceptDrop(const FDragDropEvent& DragDropEvent, SDragAndDropVerticalBox::EItemDropZone DropZone, SVerticalBox::FSlot* Slot)
+{
+	TSharedPtr<FTargetLayerDragDropOp> DragDropOperation = DragDropEvent.GetOperationAs<FTargetLayerDragDropOp>();
+
+	if (DragDropOperation.IsValid())
+	{
+		return DropZone;
+	}
+
+	return TOptional<SDragAndDropVerticalBox::EItemDropZone>();
+}
+
+FReply FLandscapeEditorCustomNodeBuilder_TargetLayers::HandleAcceptDrop(FDragDropEvent const& DragDropEvent, SDragAndDropVerticalBox::EItemDropZone DropZone, int32 SlotIndex, SVerticalBox::FSlot* Slot)
+{
+	TSharedPtr<FTargetLayerDragDropOp> DragDropOperation = DragDropEvent.GetOperationAs<FTargetLayerDragDropOp>();
+
+	if (DragDropOperation.IsValid())
+	{
+		FEdModeLandscape* LandscapeEdMode = GetEditorMode();
+
+		if (LandscapeEdMode != nullptr)
+		{
+			const TArray<FName>& TargetShownList = LandscapeEdMode->GetTargetShownList();
+
+			if (TargetShownList.IsValidIndex(DragDropOperation->SlotIndexBeingDragged) && TargetShownList.IsValidIndex(SlotIndex))
+			{
+				const TArray<FName>* TargetDisplayOrderList = LandscapeEdMode->GetTargetDisplayOrderList();
+
+				if (TargetDisplayOrderList != nullptr)
+				{
+					int32 StartingLayerIndex = TargetDisplayOrderList->Find(LandscapeEdMode->UISettings->ShowUnusedLayers ? TargetShownList[DragDropOperation->SlotIndexBeingDragged + LandscapeEdMode->GetTargetLayerStartingIndex()] : TargetShownList[DragDropOperation->SlotIndexBeingDragged]);
+					int32 DestinationLayerIndex = TargetDisplayOrderList->Find(LandscapeEdMode->UISettings->ShowUnusedLayers ? TargetShownList[SlotIndex + LandscapeEdMode->GetTargetLayerStartingIndex()] : TargetShownList[SlotIndex]);
+
+					if (StartingLayerIndex != INDEX_NONE && DestinationLayerIndex != INDEX_NONE)
+					{
+						LandscapeEdMode->MoveTargetLayerDisplayOrder(StartingLayerIndex, DestinationLayerIndex);
+
+						return FReply::Handled();
+					}
+				}
+			}
+		}
+	}
+
+	return FReply::Unhandled();
+}
+
 bool FLandscapeEditorCustomNodeBuilder_TargetLayers::GetTargetLayerIsSelected(const TSharedRef<FLandscapeTargetListInfo> Target)
 {
 	FEdModeLandscape* LandscapeEdMode = GetEditorMode();
@@ -676,6 +1055,12 @@ void FLandscapeEditorCustomNodeBuilder_TargetLayers::OnFillLayer(const TSharedRe
 	}
 }
 
+void FLandscapeEditorCustomNodeBuilder_TargetLayers::FillEmptyLayers(ULandscapeInfo* LandscapeInfo, ULandscapeLayerInfoObject* LandscapeInfoObject)
+{
+	FLandscapeEditDataInterface LandscapeEdit(LandscapeInfo);
+	LandscapeEdit.FillEmptyLayers(LandscapeInfoObject);
+}
+
 
 void FLandscapeEditorCustomNodeBuilder_TargetLayers::OnClearLayer(const TSharedRef<FLandscapeTargetListInfo> Target)
 {
@@ -752,6 +1137,8 @@ void FLandscapeEditorCustomNodeBuilder_TargetLayers::OnTargetLayerSetObject(cons
 				}
 				LandscapeEdMode->UpdateTargetList();
 			}
+
+			FillEmptyLayers(LandscapeInfo, SelectedLayerInfo);
 		}
 		else
 		{
@@ -770,14 +1157,14 @@ EVisibility FLandscapeEditorCustomNodeBuilder_TargetLayers::GetTargetLayerInfoSe
 	return EVisibility::Collapsed;
 }
 
-EVisibility FLandscapeEditorCustomNodeBuilder_TargetLayers::GetTargetLayerCreateVisibility(const TSharedRef<FLandscapeTargetListInfo> Target)
+bool FLandscapeEditorCustomNodeBuilder_TargetLayers::GetTargetLayerCreateEnabled(const TSharedRef<FLandscapeTargetListInfo> Target)
 {
 	if (!Target->LayerInfoObj.IsValid())
 	{
-		return EVisibility::Visible;
+		return true;
 	}
 
-	return EVisibility::Collapsed;
+	return false;
 }
 
 EVisibility FLandscapeEditorCustomNodeBuilder_TargetLayers::GetTargetLayerMakePublicVisibility(const TSharedRef<FLandscapeTargetListInfo> Target)
@@ -882,6 +1269,8 @@ void FLandscapeEditorCustomNodeBuilder_TargetLayers::OnTargetLayerCreateClicked(
 			GEditor->SyncBrowserToObjects(Objects);
 			
 			LandscapeEdMode->TargetsListUpdated.Broadcast();
+
+			FillEmptyLayers(LandscapeInfo, LayerInfo);
 		}
 	}
 }
@@ -1084,4 +1473,33 @@ const FSlateBrush* SLandscapeEditorSelectableBorder::GetBorder() const
 	}
 }
 
+TSharedRef<FTargetLayerDragDropOp> FTargetLayerDragDropOp::New(int32 InSlotIndexBeingDragged, SVerticalBox::FSlot* InSlotBeingDragged, TSharedPtr<SWidget> WidgetToShow)
+{
+	TSharedRef<FTargetLayerDragDropOp> Operation = MakeShareable(new FTargetLayerDragDropOp);
+
+	Operation->MouseCursor = EMouseCursor::GrabHandClosed;
+	Operation->SlotIndexBeingDragged = InSlotIndexBeingDragged;
+	Operation->SlotBeingDragged = InSlotBeingDragged;
+	Operation->WidgetToShow = WidgetToShow;
+
+	Operation->Construct();
+
+	return Operation;
+}
+
+FTargetLayerDragDropOp::~FTargetLayerDragDropOp()
+{
+}
+
+TSharedPtr<SWidget> FTargetLayerDragDropOp::GetDefaultDecorator() const
+{
+	return SNew(SBorder)
+			.BorderImage(FEditorStyle::GetBrush("ContentBrowser.AssetDragDropTooltipBackground"))
+			.Content()
+			[
+				WidgetToShow.ToSharedRef()
+			];
+		
+}
+
 #undef LOCTEXT_NAMESPACE
diff --git a/Engine/Source/Editor/LandscapeEditor/Private/LandscapeEditorDetailCustomization_TargetLayers.h b/Engine/Source/Editor/LandscapeEditor/Private/LandscapeEditorDetailCustomization_TargetLayers.h
--- a/Engine/Source/Editor/LandscapeEditor/Private/LandscapeEditorDetailCustomization_TargetLayers.h
+++ b/Engine/Source/Editor/LandscapeEditor/Private/LandscapeEditorDetailCustomization_TargetLayers.h
@@ -21,6 +21,7 @@
 class FDetailWidgetRow;
 class IDetailChildrenBuilder;
 class IDetailLayoutBuilder;
+class SDragAndDropVerticalBox;
 
 /**
  * Slate widgets customizer for the target layers list in the Landscape Editor
@@ -43,10 +44,10 @@ protected:
 	static EVisibility GetVisibility_VisibilityTip();
 };
 
-class FLandscapeEditorCustomNodeBuilder_TargetLayers : public IDetailCustomNodeBuilder
+class FLandscapeEditorCustomNodeBuilder_TargetLayers : public IDetailCustomNodeBuilder, public TSharedFromThis<FLandscapeEditorCustomNodeBuilder_TargetLayers>
 {
 public:
-	FLandscapeEditorCustomNodeBuilder_TargetLayers(TSharedRef<FAssetThumbnailPool> ThumbnailPool);
+	FLandscapeEditorCustomNodeBuilder_TargetLayers(TSharedRef<FAssetThumbnailPool> ThumbnailPool, TSharedRef<IPropertyHandle> InTargetDisplayOrderPropertyHandle, TSharedRef<IPropertyHandle> InTargetShowUnusedLayersPropertyHandle);
 	~FLandscapeEditorCustomNodeBuilder_TargetLayers();
 
 	virtual void SetOnRebuildChildren( FSimpleDelegate InOnRegenerateChildren ) override;
@@ -59,10 +60,12 @@ public:
 
 protected:
 	TSharedRef<FAssetThumbnailPool> ThumbnailPool;
+	TSharedRef<IPropertyHandle> TargetDisplayOrderPropertyHandle;
+	TSharedRef<IPropertyHandle> TargetShowUnusedLayersPropertyHandle;
 
 	static class FEdModeLandscape* GetEditorMode();
 
-	void GenerateRow(IDetailChildrenBuilder& ChildrenBuilder, const TSharedRef<FLandscapeTargetListInfo> Target);
+	TSharedPtr<SWidget> GenerateRow(const TSharedRef<FLandscapeTargetListInfo> Target);
 
 	static bool GetTargetLayerIsSelected(const TSharedRef<FLandscapeTargetListInfo> Target);
 	static void OnTargetSelectionChanged(const TSharedRef<FLandscapeTargetListInfo> Target);
@@ -71,11 +74,12 @@ protected:
 	static void OnImportLayer(const TSharedRef<FLandscapeTargetListInfo> Target);
 	static void OnReimportLayer(const TSharedRef<FLandscapeTargetListInfo> Target);
 	static void OnFillLayer(const TSharedRef<FLandscapeTargetListInfo> Target);
+	static void FillEmptyLayers(ULandscapeInfo* LandscapeInfo, ULandscapeLayerInfoObject* LandscapeInfoObject);
 	static void OnClearLayer(const TSharedRef<FLandscapeTargetListInfo> Target);
 	static bool ShouldFilterLayerInfo(const class FAssetData& AssetData, FName LayerName);
 	static void OnTargetLayerSetObject(const FAssetData& AssetData, const TSharedRef<FLandscapeTargetListInfo> Target);
 	static EVisibility GetTargetLayerInfoSelectorVisibility(const TSharedRef<FLandscapeTargetListInfo> Target);
-	static EVisibility GetTargetLayerCreateVisibility(const TSharedRef<FLandscapeTargetListInfo> Target);
+	static bool GetTargetLayerCreateEnabled(const TSharedRef<FLandscapeTargetListInfo> Target);
 	static EVisibility GetTargetLayerMakePublicVisibility(const TSharedRef<FLandscapeTargetListInfo> Target);
 	static EVisibility GetTargetLayerDeleteVisibility(const TSharedRef<FLandscapeTargetListInfo> Target);
 	static TSharedRef<SWidget> OnGetTargetLayerCreateMenu(const TSharedRef<FLandscapeTargetListInfo> Target);
@@ -88,6 +92,21 @@ protected:
 	static EVisibility GetDebugModeColorChannelVisibility(const TSharedRef<FLandscapeTargetListInfo> Target);
 	static ECheckBoxState DebugModeColorChannelIsChecked(const TSharedRef<FLandscapeTargetListInfo> Target, int32 Channel);
 	static void OnDebugModeColorChannelChanged(ECheckBoxState NewCheckedState, const TSharedRef<FLandscapeTargetListInfo> Target, int32 Channel);
+
+	// Drag/Drop handling
+	FReply HandleDragDetected(const FGeometry& MyGeometry, const FPointerEvent& MouseEvent, int32 SlotIndex, SVerticalBox::FSlot* Slot);
+	TOptional<SDragAndDropVerticalBox::EItemDropZone> HandleCanAcceptDrop(const FDragDropEvent& DragDropEvent, SDragAndDropVerticalBox::EItemDropZone DropZone, SVerticalBox::FSlot* Slot);
+	FReply HandleAcceptDrop(FDragDropEvent const& DragDropEvent, SDragAndDropVerticalBox::EItemDropZone DropZone, int32 SlotIndex, SVerticalBox::FSlot* Slot);
+
+	const FSlateBrush* GetTargetLayerDisplayOrderBrush() const;
+	TSharedRef<SWidget> GetTargetLayerDisplayOrderButtonMenuContent();
+	void SetSelectedDisplayOrder(ELandscapeLayerDisplayMode InDisplayOrder);
+	bool IsSelectedDisplayOrder(ELandscapeLayerDisplayMode InDisplayOrder) const;
+
+	TSharedRef<SWidget> GetTargetLayerShowUnusedButtonMenuContent();
+	void ShowUnusedLayers(bool Result);
+	bool ShouldShowUnusedLayers(bool Result) const;
+	EVisibility ShouldShowLayer(TSharedRef<FLandscapeTargetListInfo> Target) const;
 };
 
 class SLandscapeEditorSelectableBorder : public SBorder
@@ -120,3 +139,18 @@ protected:
 	FSimpleDelegate OnSelected;
 	TAttribute<bool> IsSelected;
 };
+
+class FTargetLayerDragDropOp : public FDragAndDropVerticalBoxOp
+{
+public:
+	DRAG_DROP_OPERATOR_TYPE(FTargetLayerDragDropOp, FDragAndDropVerticalBoxOp)
+
+	TSharedPtr<SWidget> WidgetToShow;
+
+	static TSharedRef<FTargetLayerDragDropOp> New(int32 InSlotIndexBeingDragged, SVerticalBox::FSlot* InSlotBeingDragged, TSharedPtr<SWidget> InWidgetToShow);
+
+public:
+	virtual ~FTargetLayerDragDropOp();
+
+	virtual TSharedPtr<SWidget> GetDefaultDecorator() const override;
+};
\ No newline at end of file
diff --git a/Engine/Source/Editor/LandscapeEditor/Private/LandscapeEditorObject.cpp b/Engine/Source/Editor/LandscapeEditor/Private/LandscapeEditorObject.cpp
--- a/Engine/Source/Editor/LandscapeEditor/Private/LandscapeEditorObject.cpp
+++ b/Engine/Source/Editor/LandscapeEditor/Private/LandscapeEditorObject.cpp
@@ -98,6 +98,8 @@ ULandscapeEditorObject::ULandscapeEditorObject(const FObjectInitializer& ObjectI
 	, AlphaTextureSizeY(1)
 
 	, BrushComponentSize(1)
+	, TargetDisplayOrder(ELandscapeLayerDisplayMode::Default)
+	, ShowUnusedLayers(true)
 {
 	// Structure to hold one-time initialization
 	struct FConstructorStatics
@@ -164,6 +166,18 @@ void ULandscapeEditorObject::PostEditChangeProperty(FPropertyChangedEvent& Prope
 	{
 		UpdateComponentLayerWhitelist();
 	}
+
+	if (PropertyChangedEvent.MemberProperty == nullptr ||
+		PropertyChangedEvent.MemberProperty->GetFName() == GET_MEMBER_NAME_CHECKED(ULandscapeEditorObject, TargetDisplayOrder))
+	{
+		UpdateTargetLayerDisplayOrder();
+	}
+
+	if (PropertyChangedEvent.MemberProperty == nullptr ||
+		PropertyChangedEvent.MemberProperty->GetFName() == GET_MEMBER_NAME_CHECKED(ULandscapeEditorObject, ShowUnusedLayers))
+	{
+		UpdateShowUnusedLayers();
+	}
 }
 
 /** Load UI settings from ini file */
@@ -274,6 +288,8 @@ void ULandscapeEditorObject::Load()
 	GConfig->GetBool(TEXT("LandscapeEdit"), TEXT("bApplyToAllTargets"), InbApplyToAllTargets, GEditorPerProjectIni);
 	bApplyToAllTargets = InbApplyToAllTargets;
 
+	GConfig->GetBool(TEXT("LandscapeEdit"), TEXT("ShowUnusedLayers"), ShowUnusedLayers, GEditorPerProjectIni);
+
 	// Set EditRenderMode
 	SetbUseSelectedRegion(bUseSelectedRegion);
 	SetbUseNegativeMask(bUseNegativeMask);
@@ -368,6 +384,8 @@ void ULandscapeEditorObject::Save()
 	GConfig->SetString(TEXT("LandscapeEdit"), TEXT("NewLandscapeMaterialName"), *NewLandscapeMaterialName, GEditorPerProjectIni);
 
 	GConfig->SetInt(TEXT("LandscapeEdit"), TEXT("ImportLandscape_AlphamapType"), (uint8)ImportLandscape_AlphamapType, GEditorPerProjectIni);
+
+	GConfig->SetBool(TEXT("LandscapeEdit"), TEXT("ShowUnusedLayers"), ShowUnusedLayers, GEditorPerProjectIni);
 }
 
 // Region
@@ -645,3 +663,20 @@ void ULandscapeEditorObject::UpdateComponentLayerWhitelist()
 		ParentMode->CurrentToolTarget.LandscapeInfo->UpdateComponentLayerWhitelist();
 	}
 }
+
+void ULandscapeEditorObject::UpdateTargetLayerDisplayOrder()
+{
+	if (ParentMode != nullptr)
+	{
+		ParentMode->UpdateTargetLayerDisplayOrder(TargetDisplayOrder);
+	}
+}
+
+void ULandscapeEditorObject::UpdateShowUnusedLayers()
+{
+	if (ParentMode != nullptr)
+	{
+		ParentMode->UpdateShownLayerList();
+	}
+}
+
diff --git a/Engine/Source/Editor/LandscapeEditor/Private/SLandscapeEditor.cpp b/Engine/Source/Editor/LandscapeEditor/Private/SLandscapeEditor.cpp
--- a/Engine/Source/Editor/LandscapeEditor/Private/SLandscapeEditor.cpp
+++ b/Engine/Source/Editor/LandscapeEditor/Private/SLandscapeEditor.cpp
@@ -165,6 +165,11 @@ void FLandscapeToolKit::NotifyBrushChanged()
 	LandscapeEditorWidgets->NotifyBrushChanged();
 }
 
+void FLandscapeToolKit::RefreshDetailPanel()
+{
+	LandscapeEditorWidgets->RefreshDetailPanel();
+}
+
 void FLandscapeToolKit::OnChangeMode(FName ModeName)
 {
 	FEdModeLandscape* LandscapeEdMode = GetEditorMode();
@@ -507,15 +512,15 @@ bool SLandscapeEditor::GetIsPropertyVisible(const FPropertyAndParent& PropertyAn
 
 void SLandscapeEditor::NotifyToolChanged()
 {
-	FEdModeLandscape* LandscapeEdMode = GetEditorMode();
-	if (LandscapeEdMode)
-	{
-		// Refresh details panel
-		DetailsPanel->SetObject(LandscapeEdMode->UISettings, true);
-	}
+	RefreshDetailPanel();
 }
 
 void SLandscapeEditor::NotifyBrushChanged()
+{
+	RefreshDetailPanel();
+}
+
+void SLandscapeEditor::RefreshDetailPanel()
 {
 	FEdModeLandscape* LandscapeEdMode = GetEditorMode();
 	if (LandscapeEdMode)
diff --git a/Engine/Source/Editor/LandscapeEditor/Public/LandscapeEditorObject.h b/Engine/Source/Editor/LandscapeEditor/Public/LandscapeEditorObject.h
--- a/Engine/Source/Editor/LandscapeEditor/Public/LandscapeEditorObject.h
+++ b/Engine/Source/Editor/LandscapeEditor/Public/LandscapeEditorObject.h
@@ -290,6 +290,13 @@ class ULandscapeEditorObject : public UObject
 	UPROPERTY(Category = "Tool Settings", EditAnywhere, NonTransactional, AdvancedDisplay, meta = (DisplayName = "Show Preview Grid", ShowForTools = "Flatten", ShowForTargetTypes = "Heightmap", EditCondition = "bUseFlattenTarget", HideEditConditionToggle, UIMin = "-32768", UIMax = "32768"))
 	bool bShowFlattenTargetPreview;
 
+	// Whether the Eye Dropper mode is activated
+	UPROPERTY(NonTransactional, Transient)
+	bool bFlattenEyeDropperModeActivated;
+
+	UPROPERTY(NonTransactional, Transient)
+	float FlattenEyeDropperModeDesiredTarget;
+
 	// Ramp Tool:
 
 	// Width of ramp
@@ -580,6 +587,13 @@ public:
 	UPROPERTY(Category="Target Layers", EditAnywhere, NonTransactional, meta=(ShowForTargetTypes="Weightmap,Visibility"))
 	ELandscapeLayerPaintingRestriction PaintingRestriction;
 
+	// Display order of the targets
+	UPROPERTY(Category = "Target Layers", EditAnywhere)
+	ELandscapeLayerDisplayMode TargetDisplayOrder;
+
+	UPROPERTY(Category = "Target Layers", EditAnywhere)
+	bool ShowUnusedLayers;	
+
 #if WITH_EDITOR
 	//~ Begin UObject Interface
 	void PostEditChangeProperty(FPropertyChangedEvent& PropertyChangedEvent) override;
@@ -662,4 +676,7 @@ public:
 	{
 		ParentMode = LandscapeParent;
 	}
+
+	void UpdateTargetLayerDisplayOrder();
+	void UpdateShowUnusedLayers();
 };
diff --git a/Engine/Source/Editor/LandscapeEditor/Public/LandscapeToolInterface.h b/Engine/Source/Editor/LandscapeEditor/Public/LandscapeToolInterface.h
--- a/Engine/Source/Editor/LandscapeEditor/Public/LandscapeToolInterface.h
+++ b/Engine/Source/Editor/LandscapeEditor/Public/LandscapeToolInterface.h
@@ -235,6 +235,8 @@ public:
 	virtual bool HandleClick(HHitProxy* HitProxy, const FViewportClick& Click) { return false; }
 	virtual bool InputKey(FEditorViewportClient* InViewportClient, FViewport* InViewport, FKey InKey, EInputEvent InEvent) { return false; }
 	virtual bool InputDelta(FEditorViewportClient* InViewportClient, FViewport* InViewport, FVector& InDrag, FRotator& InRot, FVector& InScale) { return false; }
+	virtual bool GetCursor(EMouseCursor::Type& OutCursor) const { return false;  }
+
 	FLandscapeTool() : PreviousBrushIndex(-1) {}
 	virtual ~FLandscapeTool() {}
 	virtual const TCHAR* GetToolName() = 0;
@@ -251,6 +253,9 @@ public:
 	virtual FMatrix GetWidgetRotation() const { return FMatrix::Identity; }
 	virtual bool DisallowMouseDeltaTracking() const { return false; }
 
+	virtual void SetCanToolBeActivated(bool Value) { }
+	virtual bool CanToolBeActivated() const { return true;  }
+
 	virtual EEditAction::Type GetActionEditDuplicate() { return EEditAction::Skip; }
 	virtual EEditAction::Type GetActionEditDelete() { return EEditAction::Skip; }
 	virtual EEditAction::Type GetActionEditCut() { return EEditAction::Skip; }
diff --git a/Engine/Source/Editor/LevelEditor/Private/LevelEditor.cpp b/Engine/Source/Editor/LevelEditor/Private/LevelEditor.cpp
--- a/Engine/Source/Editor/LevelEditor/Private/LevelEditor.cpp
+++ b/Engine/Source/Editor/LevelEditor/Private/LevelEditor.cpp
@@ -588,7 +588,7 @@ void FLevelEditorModule::BindGlobalLevelEditorCommands()
 	ActionList.MapAction( Commands.NewLevel, FExecuteAction::CreateStatic( &FLevelEditorActionCallbacks::NewLevel ), FCanExecuteAction::CreateStatic( &FLevelEditorActionCallbacks::NewLevel_CanExecute ) );
 	ActionList.MapAction(Commands.OpenLevel, FExecuteAction::CreateStatic(&FLevelEditorActionCallbacks::OpenLevel), FCanExecuteAction::CreateStatic(&FLevelEditorActionCallbacks::OpenLevel_CanExecute));
 	ActionList.MapAction( Commands.Save, FExecuteAction::CreateStatic( &FLevelEditorActionCallbacks::Save ), FCanExecuteAction::CreateStatic( &FLevelEditorActionCallbacks::CanSaveWorld ) );
-	ActionList.MapAction( Commands.SaveAs, FExecuteAction::CreateStatic( &FLevelEditorActionCallbacks::SaveAs ), FCanExecuteAction::CreateStatic( &FLevelEditorActionCallbacks::CanSaveWorld ) );
+	ActionList.MapAction( Commands.SaveAs, FExecuteAction::CreateStatic( &FLevelEditorActionCallbacks::SaveCurrentAs ), FCanExecuteAction::CreateStatic( &FLevelEditorActionCallbacks::CanSaveWorld ) );
 	ActionList.MapAction( Commands.SaveAllLevels, FExecuteAction::CreateStatic( &FLevelEditorActionCallbacks::SaveAllLevels ), FCanExecuteAction::CreateStatic( &FLevelEditorActionCallbacks::CanSaveWorld ) );
 	ActionList.MapAction( Commands.ToggleFavorite, FExecuteAction::CreateStatic( &FLevelEditorActionCallbacks::ToggleFavorite ), FCanExecuteAction::CreateStatic( &FLevelEditorActionCallbacks::ToggleFavorite_CanExecute ), FIsActionChecked::CreateStatic( &FLevelEditorActionCallbacks::ToggleFavorite_IsChecked ) );
 
diff --git a/Engine/Source/Editor/LevelEditor/Private/LevelEditorActions.cpp b/Engine/Source/Editor/LevelEditor/Private/LevelEditorActions.cpp
--- a/Engine/Source/Editor/LevelEditor/Private/LevelEditorActions.cpp
+++ b/Engine/Source/Editor/LevelEditor/Private/LevelEditorActions.cpp
@@ -84,6 +84,9 @@
 #include "MaterialShaderQualitySettings.h"
 #include "IVREditorModule.h"
 #include "ComponentRecreateRenderStateContext.h"
+#include "Engine/LevelStreaming.h"
+#include "Engine/LevelStreamingKismet.h"
+#include "EditorLevelUtils.h"
 
 DEFINE_LOG_CATEGORY_STATIC(LevelEditorActions, Log, All);
 
@@ -383,13 +386,49 @@ void FLevelEditorActionCallbacks::Save()
 	FEditorFileUtils::SaveCurrentLevel();
 }
 
-void FLevelEditorActionCallbacks::SaveAs()
+void FLevelEditorActionCallbacks::SaveCurrentAs()
 {
+	UWorld* World = GetWorld();
+	ULevel* CurrentLevel = World->GetCurrentLevel();
+	
+	UClass* CurrentStreamingLevelClass = ULevelStreamingKismet::StaticClass();
+
+	for (ULevelStreaming* StreamingLevel : World->StreamingLevels)
+	{
+		if (StreamingLevel->GetLoadedLevel() == CurrentLevel)
+		{
+			CurrentStreamingLevelClass = StreamingLevel->GetClass();
+			break;
+		}
+	}
+
+	
+	const bool bSavedPersistentLevelAs = CurrentLevel == World->PersistentLevel;
 	FString SavedFilename;
-	bool bSaved = FEditorFileUtils::SaveLevelAs( GetWorld()->PersistentLevel, &SavedFilename );
-	if (bSaved)
+	bool bSaved = FEditorFileUtils::SaveLevelAs(CurrentLevel, &SavedFilename);
+	if(bSaved)
 	{
-		FEditorFileUtils::LoadMap(SavedFilename);
+		if (bSavedPersistentLevelAs)
+		{
+			FEditorFileUtils::LoadMap(SavedFilename);
+		}
+		else
+		{
+			// Remove the level we just saved over
+			EditorLevelUtils::RemoveLevelFromWorld(CurrentLevel);
+
+			// Add the new level we just saved as to the plevel
+			FString PackageName;
+			if (FPackageName::TryConvertFilenameToLongPackageName(SavedFilename, PackageName))
+			{
+				ULevel* Level = EditorLevelUtils::AddLevelToWorld(World, *PackageName, CurrentStreamingLevelClass);
+
+				// Make the level we just added current because the expectation is that the new level replaces the existing current level
+				EditorLevelUtils::MakeLevelCurrent(Level);
+			}
+
+			FEditorDelegates::RefreshLevelBrowser.Broadcast();
+		}
 	}
 }
 
diff --git a/Engine/Source/Editor/LevelEditor/Private/LevelEditorContextMenu.h b/Engine/Source/Editor/LevelEditor/Private/LevelEditorContextMenu.h
--- a/Engine/Source/Editor/LevelEditor/Private/LevelEditorContextMenu.h
+++ b/Engine/Source/Editor/LevelEditor/Private/LevelEditorContextMenu.h
@@ -1,9 +1,5 @@
 // Copyright 1998-2017 Epic Games, Inc. All Rights Reserved.
 
-
-#ifndef __LevelViewportContextMenu_h__
-#define __LevelViewportContextMenu_h__
-
 #pragma once
 
 #include "CoreMinimal.h"
@@ -74,5 +70,3 @@ private:
 	 */
 	static void BuildGroupMenu( FMenuBuilder& MenuBuilder, const struct FSelectedActorInfo& SelectedActorInfo );
 };
-
-#endif		// __LevelViewportContextMenu_h__
diff --git a/Engine/Source/Editor/LevelEditor/Private/LevelEditorToolBar.h b/Engine/Source/Editor/LevelEditor/Private/LevelEditorToolBar.h
--- a/Engine/Source/Editor/LevelEditor/Private/LevelEditorToolBar.h
+++ b/Engine/Source/Editor/LevelEditor/Private/LevelEditorToolBar.h
@@ -1,9 +1,5 @@
 // Copyright 1998-2017 Epic Games, Inc. All Rights Reserved.
 
-
-#ifndef __LevelEditorToolBar_h__
-#define __LevelEditorToolBar_h__
-
 #pragma once
 
 
@@ -85,7 +81,3 @@ protected:
 	 */
 	static void OnOpenSubLevelBlueprint( ULevel* InLevel );
 };
-
-
-
-#endif		// __LevelEditorToolBar_h__
diff --git a/Engine/Source/Editor/LevelEditor/Private/LevelViewportLayout2x2.h b/Engine/Source/Editor/LevelEditor/Private/LevelViewportLayout2x2.h
--- a/Engine/Source/Editor/LevelEditor/Private/LevelViewportLayout2x2.h
+++ b/Engine/Source/Editor/LevelEditor/Private/LevelViewportLayout2x2.h
@@ -1,9 +1,5 @@
 // Copyright 1998-2017 Epic Games, Inc. All Rights Reserved.
 
-
-#ifndef __LevelViewportLayout2x2_h__
-#define __LevelViewportLayout2x2_h__
-
 #pragma once
 
 #include "CoreMinimal.h"
@@ -35,5 +31,3 @@ private:
 	/** The splitter widget */
 	TSharedPtr< class SSplitter2x2 > SplitterWidget;
 };
-
-#endif
diff --git a/Engine/Source/Editor/LevelEditor/Private/LevelViewportLayoutFourPanes.h b/Engine/Source/Editor/LevelEditor/Private/LevelViewportLayoutFourPanes.h
--- a/Engine/Source/Editor/LevelEditor/Private/LevelViewportLayoutFourPanes.h
+++ b/Engine/Source/Editor/LevelEditor/Private/LevelViewportLayoutFourPanes.h
@@ -1,9 +1,5 @@
 // Copyright 1998-2017 Epic Games, Inc. All Rights Reserved.
 
-
-#ifndef __LevelViewportLayoutFourPanes_h__
-#define __LevelViewportLayoutFourPanes_h__
-
 #pragma once
 
 #include "CoreMinimal.h"
@@ -95,5 +91,3 @@ public:
 		TSharedRef<SWidget> Viewport0, TSharedRef<SWidget> Viewport1, TSharedRef<SWidget> Viewport2, TSharedRef<SWidget> Viewport3,
 		float PrimarySplitterPercentage, float SecondarySplitterPercentage0, float SecondarySplitterPercentage1) override;
 };
-
-#endif
diff --git a/Engine/Source/Editor/LevelEditor/Private/LevelViewportLayoutOnePane.h b/Engine/Source/Editor/LevelEditor/Private/LevelViewportLayoutOnePane.h
--- a/Engine/Source/Editor/LevelEditor/Private/LevelViewportLayoutOnePane.h
+++ b/Engine/Source/Editor/LevelEditor/Private/LevelViewportLayoutOnePane.h
@@ -1,9 +1,5 @@
 // Copyright 1998-2017 Epic Games, Inc. All Rights Reserved.
 
-
-#ifndef __LevelViewportLayoutOnePane_h__
-#define __LevelViewportLayoutOnePane_h__
-
 #pragma once
 
 #include "CoreMinimal.h"
@@ -35,5 +31,3 @@ protected:
 	/** The viewport widget parent box */
 	TSharedPtr< SHorizontalBox > ViewportBox;
 };
-
-#endif
diff --git a/Engine/Source/Editor/LevelEditor/Private/LevelViewportLayoutThreePanes.h b/Engine/Source/Editor/LevelEditor/Private/LevelViewportLayoutThreePanes.h
--- a/Engine/Source/Editor/LevelEditor/Private/LevelViewportLayoutThreePanes.h
+++ b/Engine/Source/Editor/LevelEditor/Private/LevelViewportLayoutThreePanes.h
@@ -1,9 +1,5 @@
 // Copyright 1998-2017 Epic Games, Inc. All Rights Reserved.
 
-
-#ifndef __LevelViewportLayoutThreePanes_h__
-#define __LevelViewportLayoutThreePanes_h__
-
 #pragma once
 
 #include "CoreMinimal.h"
@@ -94,5 +90,3 @@ public:
 		const TSharedRef<SWidget>& ViewportKey0, const TSharedRef<SWidget>& ViewportKey1, const TSharedRef<SWidget>& ViewportKey2,
 		float PrimarySplitterPercentage, float SecondarySplitterPercentage) override;
 };
-
-#endif
diff --git a/Engine/Source/Editor/LevelEditor/Private/LevelViewportLayoutTwoPanes.h b/Engine/Source/Editor/LevelEditor/Private/LevelViewportLayoutTwoPanes.h
--- a/Engine/Source/Editor/LevelEditor/Private/LevelViewportLayoutTwoPanes.h
+++ b/Engine/Source/Editor/LevelEditor/Private/LevelViewportLayoutTwoPanes.h
@@ -1,9 +1,5 @@
 // Copyright 1998-2017 Epic Games, Inc. All Rights Reserved.
 
-
-#ifndef __LevelViewportLayoutTwoPanes_h__
-#define __LevelViewportLayoutTwoPanes_h__
-
 #pragma once
 
 #include "CoreMinimal.h"
@@ -52,5 +48,3 @@ class FLevelViewportLayoutTwoPanesHoriz : public TLevelViewportLayoutTwoPanes<EO
 public:
 	virtual const FName& GetLayoutTypeName() const override { return LevelViewportConfigurationNames::TwoPanesHoriz; }
 };
-
-#endif
diff --git a/Engine/Source/Editor/LevelEditor/Private/SLevelEditor.h b/Engine/Source/Editor/LevelEditor/Private/SLevelEditor.h
--- a/Engine/Source/Editor/LevelEditor/Private/SLevelEditor.h
+++ b/Engine/Source/Editor/LevelEditor/Private/SLevelEditor.h
@@ -1,9 +1,5 @@
 // Copyright 1998-2017 Epic Games, Inc. All Rights Reserved.
 
-
-#ifndef __SLevelEditor_h__
-#define __SLevelEditor_h__
-
 #pragma once
 
 
@@ -247,7 +243,3 @@ private:
 	/** Attached sequencer asset editor */
 	TWeakPtr<IAssetEditorInstance> SequencerAssetEditor;
 };
-
-
-
-#endif	// __SLevelEditor_h__
diff --git a/Engine/Source/Editor/LevelEditor/Private/SLevelViewport.cpp b/Engine/Source/Editor/LevelEditor/Private/SLevelViewport.cpp
--- a/Engine/Source/Editor/LevelEditor/Private/SLevelViewport.cpp
+++ b/Engine/Source/Editor/LevelEditor/Private/SLevelViewport.cpp
@@ -3479,16 +3479,20 @@ void SLevelViewport::StartPlayInEditorSession(UGameViewportClient* PlayClient, c
 	// Register the new viewport widget with Slate for viewport specific message routing.
 	FSlateApplication::Get().RegisterGameViewport(ViewportWidget.ToSharedRef() );
 
+	ULevelEditorPlaySettings const* EditorPlayInSettings = GetDefault<ULevelEditorPlaySettings>();
+	check(EditorPlayInSettings);
+
 	// Kick off a quick transition effect (border graphics)
 	ViewTransitionType = EViewTransition::StartingPlayInEditor;
 	ViewTransitionAnim = FCurveSequence( 0.0f, 1.5f, ECurveEaseFunction::CubicOut );
 	bViewTransitionAnimPending = true;
-	GEditor->PlayEditorSound( TEXT( "/Engine/EditorSounds/GamePreview/StartPlayInEditor_Cue.StartPlayInEditor_Cue" ) );
+	if (EditorPlayInSettings->EnablePIEEnterAndExitSounds)
+	{
+		GEditor->PlayEditorSound(TEXT("/Engine/EditorSounds/GamePreview/StartPlayInEditor_Cue.StartPlayInEditor_Cue"));
+	}
 
 	bPIEHasFocus = ActiveViewport->HasMouseCapture();
 
-	ULevelEditorPlaySettings const* EditorPlayInSettings = GetDefault<ULevelEditorPlaySettings>();
-	check(EditorPlayInSettings);
 
 	if(EditorPlayInSettings->ShowMouseControlLabel)
 	{
@@ -3690,7 +3694,11 @@ void SLevelViewport::EndPlayInEditorSession()
 	ViewTransitionType = EViewTransition::ReturningToEditor;
 	ViewTransitionAnim = FCurveSequence( 0.0f, 1.5f, ECurveEaseFunction::CubicOut );
 	bViewTransitionAnimPending = true;
-	GEditor->PlayEditorSound( TEXT( "/Engine/EditorSounds/GamePreview/EndPlayInEditor_Cue.EndPlayInEditor_Cue" ) );
+
+	if (GetDefault<ULevelEditorPlaySettings>()->EnablePIEEnterAndExitSounds)
+	{
+		GEditor->PlayEditorSound( TEXT( "/Engine/EditorSounds/GamePreview/EndPlayInEditor_Cue.EndPlayInEditor_Cue" ) );
+	}
 
 	GEngine->BroadcastLevelActorListChanged();
 }
@@ -3765,7 +3773,11 @@ void SLevelViewport::SwapViewportsForPlayInEditor()
 	ViewTransitionType = EViewTransition::StartingPlayInEditor;
 	ViewTransitionAnim = FCurveSequence( 0.0f, 1.5f, ECurveEaseFunction::CubicOut );
 	bViewTransitionAnimPending = true;
-	GEditor->PlayEditorSound( TEXT( "/Engine/EditorSounds/GamePreview/EjectFromPlayer_Cue.EjectFromPlayer_Cue" ) );
+
+	if (EditorPlayInSettings->EnablePIEEnterAndExitSounds)
+	{
+		GEditor->PlayEditorSound( TEXT( "/Engine/EditorSounds/GamePreview/EjectFromPlayer_Cue.EjectFromPlayer_Cue" ) );
+	}
 }
 
 
@@ -3775,7 +3787,10 @@ void SLevelViewport::OnSimulateSessionStarted()
 	ViewTransitionType = EViewTransition::StartingSimulate;
 	ViewTransitionAnim = FCurveSequence( 0.0f, 1.5f, ECurveEaseFunction::CubicOut );
 	bViewTransitionAnimPending = true;
-	GEditor->PlayEditorSound( TEXT( "/Engine/EditorSounds/GamePreview/StartSimulate_Cue.StartSimulate_Cue" ) );
+	if (GetDefault<ULevelEditorPlaySettings>()->EnablePIEEnterAndExitSounds)
+	{
+		GEditor->PlayEditorSound( TEXT( "/Engine/EditorSounds/GamePreview/StartSimulate_Cue.StartSimulate_Cue" ) );
+	}
 	 
 	// Make sure the viewport's hit proxies are invalidated.  If not done, clicking in the viewport could select an editor world actor
 	ActiveViewport->InvalidateHitProxy();
@@ -3788,7 +3803,10 @@ void SLevelViewport::OnSimulateSessionFinished()
 	ViewTransitionType = EViewTransition::ReturningToEditor;
 	ViewTransitionAnim = FCurveSequence( 0.0f, 1.5f, ECurveEaseFunction::CubicOut );
 	bViewTransitionAnimPending = true;
-	GEditor->PlayEditorSound( TEXT( "/Engine/EditorSounds/GamePreview/EndSimulate_Cue.EndSimulate_Cue" ) );
+	if (GetDefault<ULevelEditorPlaySettings>()->EnablePIEEnterAndExitSounds)
+	{
+		GEditor->PlayEditorSound( TEXT( "/Engine/EditorSounds/GamePreview/EndSimulate_Cue.EndSimulate_Cue" ) );
+	}
 
 	// Make sure the viewport's hit proxies are invalidated.  If not done, clicking in the viewport could select a pie world actor
 	ActiveViewport->InvalidateHitProxy();
diff --git a/Engine/Source/Editor/LevelEditor/Public/DlgDeltaTransform.h b/Engine/Source/Editor/LevelEditor/Public/DlgDeltaTransform.h
--- a/Engine/Source/Editor/LevelEditor/Public/DlgDeltaTransform.h
+++ b/Engine/Source/Editor/LevelEditor/Public/DlgDeltaTransform.h
@@ -4,8 +4,7 @@
 	DlgDeltaTransform.h: UnrealEd dialog for moving assets.
 =============================================================================*/
 
-#ifndef __DLGDELTATRANSFORM_H__
-#define __DLGDELTATRANSFORM_H__
+#pragma once
 
 #include "CoreMinimal.h"
 
@@ -54,5 +53,3 @@ private:
 	/** Cached pointer to the message box held within the window */
 	TSharedPtr<class SDlgDeltaTransform> DeltaTransformWidget;
 };
-
-#endif // __DLGDELTATRANSFORM_H__
diff --git a/Engine/Source/Editor/LevelEditor/Public/LevelEditor.h b/Engine/Source/Editor/LevelEditor/Public/LevelEditor.h
--- a/Engine/Source/Editor/LevelEditor/Public/LevelEditor.h
+++ b/Engine/Source/Editor/LevelEditor/Public/LevelEditor.h
@@ -1,7 +1,6 @@
 // Copyright 1998-2017 Epic Games, Inc. All Rights Reserved.
 
-#ifndef __LevelEditor_h__
-#define __LevelEditor_h__
+#pragma once
 
 #include "CoreMinimal.h"
 #include "HAL/IConsoleManager.h"
@@ -397,5 +396,3 @@ private:
 	/** Array of delegates that are used to check if the specified objects should be editable on the details panel */
 	TArray<FAreObjectsEditable> AreObjectsEditableDelegates;
 };
-
-#endif // __LevelEditor_h__
diff --git a/Engine/Source/Editor/LevelEditor/Public/LevelEditorActions.h b/Engine/Source/Editor/LevelEditor/Public/LevelEditorActions.h
--- a/Engine/Source/Editor/LevelEditor/Public/LevelEditorActions.h
+++ b/Engine/Source/Editor/LevelEditor/Public/LevelEditorActions.h
@@ -689,7 +689,7 @@ public:
 	static bool CanSaveWorld();
 
 	/** Save the current level as... */
-	static void SaveAs();
+	static void SaveCurrentAs();
 
 	/** Saves the current map */
 	static void Save();
diff --git a/Engine/Source/Editor/LevelEditor/Public/LevelViewportLayout.h b/Engine/Source/Editor/LevelEditor/Public/LevelViewportLayout.h
--- a/Engine/Source/Editor/LevelEditor/Public/LevelViewportLayout.h
+++ b/Engine/Source/Editor/LevelEditor/Public/LevelViewportLayout.h
@@ -1,9 +1,5 @@
 // Copyright 1998-2017 Epic Games, Inc. All Rights Reserved.
 
-
-#ifndef __LevelViewportLayout_h__
-#define __LevelViewportLayout_h__
-
 #pragma once
 
 #include "CoreMinimal.h"
@@ -350,5 +346,3 @@ protected:
 	};
 	TArray<FMaximizeViewportCommand> DeferredMaximizeCommands;
 };
-
-#endif
diff --git a/Engine/Source/Editor/LevelEditor/Public/LevelViewportTabContent.h b/Engine/Source/Editor/LevelEditor/Public/LevelViewportTabContent.h
--- a/Engine/Source/Editor/LevelEditor/Public/LevelViewportTabContent.h
+++ b/Engine/Source/Editor/LevelEditor/Public/LevelViewportTabContent.h
@@ -1,9 +1,5 @@
 // Copyright 1998-2017 Epic Games, Inc. All Rights Reserved.
 
-
-#ifndef __LevelViewportTabContent_h__
-#define __LevelViewportTabContent_h__
-
 #pragma once
 
 #include "CoreMinimal.h"
@@ -85,5 +81,3 @@ private:
 	/** Current layout */
 	TSharedPtr< class FLevelViewportLayout > ActiveLevelViewportLayout;
 };
-
-#endif
diff --git a/Engine/Source/Editor/LevelEditor/Public/LightmapResRatioAdjust.h b/Engine/Source/Editor/LevelEditor/Public/LightmapResRatioAdjust.h
--- a/Engine/Source/Editor/LevelEditor/Public/LightmapResRatioAdjust.h
+++ b/Engine/Source/Editor/LevelEditor/Public/LightmapResRatioAdjust.h
@@ -4,9 +4,6 @@
 	LightmapResRatioAdjust.h: Lightmap Resolution Ratio Adjustment helper
 ================================================================================*/
 
-#ifndef __LightmapResRatioAdjust_h__
-#define __LightmapResRatioAdjust_h__
-
 #pragma once
 
 #include "CoreMinimal.h"
@@ -95,5 +92,3 @@ public:
 	/** If true, then only do selected primitives in the level(s) */
 	bool bSelectedObjectsOnly;
 };
-
-#endif	// __LightmapResRatioAdjust_h__
diff --git a/Engine/Source/Editor/MainFrame/Private/Frame/MainFrameActions.h b/Engine/Source/Editor/MainFrame/Private/Frame/MainFrameActions.h
--- a/Engine/Source/Editor/MainFrame/Private/Frame/MainFrameActions.h
+++ b/Engine/Source/Editor/MainFrame/Private/Frame/MainFrameActions.h
@@ -1,9 +1,5 @@
 // Copyright 1998-2017 Epic Games, Inc. All Rights Reserved.
 
-
-#ifndef __MainFrameActions_h__
-#define __MainFrameActions_h__
-
 #pragma once
 
 #include "CoreMinimal.h"
@@ -266,6 +262,3 @@ private:
 	/** The notification in place while we choose packages to check in */
 	static TWeakPtr<SNotificationItem> ChoosePackagesToCheckInNotification;
 };
-
-
-#endif		// __MainFrameActions_h__
diff --git a/Engine/Source/Editor/MaterialEditor/Public/MaterialEditorModule.h b/Engine/Source/Editor/MaterialEditor/Public/MaterialEditorModule.h
--- a/Engine/Source/Editor/MaterialEditor/Public/MaterialEditorModule.h
+++ b/Engine/Source/Editor/MaterialEditor/Public/MaterialEditorModule.h
@@ -1,7 +1,6 @@
 // Copyright 1998-2017 Epic Games, Inc. All Rights Reserved.
 
-#ifndef __MaterialEditorModule_h__
-#define __MaterialEditorModule_h__
+#pragma once
 
 #include "CoreMinimal.h"
 #include "Modules/ModuleInterface.h"
@@ -68,5 +67,3 @@ private:
 	FMaterialFunctionEditorOpenedEvent MaterialFunctionEditorOpenedEvent;
 	FMaterialInstanceEditorOpenedEvent MaterialInstanceEditorOpenedEvent;
 };
-
-#endif // __MaterialEditorModule_h__
diff --git a/Engine/Source/Editor/Matinee/Public/MatineeModule.h b/Engine/Source/Editor/Matinee/Public/MatineeModule.h
--- a/Engine/Source/Editor/Matinee/Public/MatineeModule.h
+++ b/Engine/Source/Editor/Matinee/Public/MatineeModule.h
@@ -1,7 +1,6 @@
 // Copyright 1998-2017 Epic Games, Inc. All Rights Reserved.
 
-#ifndef __MatineeModule_h__
-#define __MatineeModule_h__
+#pragma once
 
 #include "CoreMinimal.h"
 #include "Modules/ModuleInterface.h"
@@ -44,6 +43,3 @@ public:
 	DECLARE_EVENT(IMainFrameModule, FMatineeEditorOpenedEvent);
 	virtual FMatineeEditorOpenedEvent& OnMatineeEditorOpened() = 0;
 };
-
-
-#endif // __MatineeModule_h__
diff --git a/Engine/Source/Editor/MeshPaint/Private/SMeshPaint.cpp b/Engine/Source/Editor/MeshPaint/Private/SMeshPaint.cpp
--- a/Engine/Source/Editor/MeshPaint/Private/SMeshPaint.cpp
+++ b/Engine/Source/Editor/MeshPaint/Private/SMeshPaint.cpp
@@ -2130,8 +2130,8 @@ void SMeshPaint::Construct(const FArguments& InArgs, TSharedRef<FMeshPaintToolKi
 							.VAlign(VAlign_Center)
 							[
 								SNew(STextBlock)
-								.Text(LOCTEXT("MeshPaint_VertexColorLODSelectionLabel", "Allow Paint On All LOD"))
-								.ToolTipText(LOCTEXT("MeshPaint_VertexColorLODSelectionLabelTooltip", "When uncheck, the painting on the base LOD will be propagate automatically to all other LODs when exiting the mode or changing the selection."))
+								.Text(LOCTEXT("MeshPaint_VertexColorLODSelectionLabel", "Allow Paint On All LODs"))
+								.ToolTipText(LOCTEXT("MeshPaint_VertexColorLODSelectionLabelTooltip", "When checked, any vertices painted on the base LOD will propagate to all other LODs."))
 							]
 							+SHorizontalBox::Slot()
 							.AutoWidth()
diff --git a/Engine/Source/Editor/MovieSceneCaptureDialog/Private/MovieSceneCaptureDialogModule.cpp b/Engine/Source/Editor/MovieSceneCaptureDialog/Private/MovieSceneCaptureDialogModule.cpp
--- a/Engine/Source/Editor/MovieSceneCaptureDialog/Private/MovieSceneCaptureDialogModule.cpp
+++ b/Engine/Source/Editor/MovieSceneCaptureDialog/Private/MovieSceneCaptureDialogModule.cpp
@@ -473,7 +473,7 @@ private:
 		PlayInEditorSettings->GameGetsMouseControl = false;
 		PlayInEditorSettings->ShowMouseControlLabel = false;
 		PlayInEditorSettings->ViewportGetsHMDControl = false;
-		PlayInEditorSettings->EnableSound = false;
+		PlayInEditorSettings->EnableGameSound = false;
 		PlayInEditorSettings->bOnlyLoadVisibleLevelsInPIE = false;
 		PlayInEditorSettings->bPreferToStreamLevelsInPIE = false;
 		PlayInEditorSettings->PIEAlwaysOnTop = false;
diff --git a/Engine/Source/Editor/MovieSceneTools/Private/Sections/BoolPropertySection.cpp b/Engine/Source/Editor/MovieSceneTools/Private/Sections/BoolPropertySection.cpp
--- a/Engine/Source/Editor/MovieSceneTools/Private/Sections/BoolPropertySection.cpp
+++ b/Engine/Source/Editor/MovieSceneTools/Private/Sections/BoolPropertySection.cpp
@@ -55,7 +55,7 @@ int32 FBoolPropertySection::OnPaintSection( FSequencerSectionPainter& Painter )
 
 	int32 LayerId = Painter.PaintSectionBackground();
 
-	const ESlateDrawEffect::Type DrawEffects = Painter.bParentEnabled
+	const ESlateDrawEffect DrawEffects = Painter.bParentEnabled
 		? ESlateDrawEffect::None
 		: ESlateDrawEffect::DisabledEffect;
 
diff --git a/Engine/Source/Editor/MovieSceneTools/Private/Sections/ColorPropertySection.cpp b/Engine/Source/Editor/MovieSceneTools/Private/Sections/ColorPropertySection.cpp
--- a/Engine/Source/Editor/MovieSceneTools/Private/Sections/ColorPropertySection.cpp
+++ b/Engine/Source/Editor/MovieSceneTools/Private/Sections/ColorPropertySection.cpp
@@ -42,7 +42,7 @@ int32 FColorPropertySection::OnPaintSection( FSequencerSectionPainter& Painter )
 {
 	int32 LayerId = Painter.PaintSectionBackground();
 
-	const ESlateDrawEffect::Type DrawEffects = Painter.bParentEnabled ? ESlateDrawEffect::None : ESlateDrawEffect::DisabledEffect;
+	const ESlateDrawEffect DrawEffects = Painter.bParentEnabled ? ESlateDrawEffect::None : ESlateDrawEffect::DisabledEffect;
 
 	const UMovieSceneColorSection* ColorSection = Cast<const UMovieSceneColorSection>( &SectionObject );
 
diff --git a/Engine/Source/Editor/MovieSceneTools/Private/Sections/ThumbnailSection.cpp b/Engine/Source/Editor/MovieSceneTools/Private/Sections/ThumbnailSection.cpp
--- a/Engine/Source/Editor/MovieSceneTools/Private/Sections/ThumbnailSection.cpp
+++ b/Engine/Source/Editor/MovieSceneTools/Private/Sections/ThumbnailSection.cpp
@@ -228,7 +228,7 @@ int32 FThumbnailSection::OnPaintSection( FSequencerSectionPainter& InPainter ) c
 
 	static const float SectionThumbnailPadding = 4.f;
 
-	const ESlateDrawEffect::Type DrawEffects = InPainter.bParentEnabled ? ESlateDrawEffect::None : ESlateDrawEffect::DisabledEffect;
+	const ESlateDrawEffect DrawEffects = InPainter.bParentEnabled ? ESlateDrawEffect::None : ESlateDrawEffect::DisabledEffect;
 
 	int32 LayerId = InPainter.LayerId;
 
diff --git a/Engine/Source/Editor/MovieSceneTools/Private/TrackEditors/FadeTrackEditor.cpp b/Engine/Source/Editor/MovieSceneTools/Private/TrackEditors/FadeTrackEditor.cpp
--- a/Engine/Source/Editor/MovieSceneTools/Private/TrackEditors/FadeTrackEditor.cpp
+++ b/Engine/Source/Editor/MovieSceneTools/Private/TrackEditors/FadeTrackEditor.cpp
@@ -31,7 +31,7 @@ public:
 	{
 		int32 LayerId = Painter.PaintSectionBackground();
 
-		const ESlateDrawEffect::Type DrawEffects = Painter.bParentEnabled ? ESlateDrawEffect::None : ESlateDrawEffect::DisabledEffect;
+		const ESlateDrawEffect DrawEffects = Painter.bParentEnabled ? ESlateDrawEffect::None : ESlateDrawEffect::DisabledEffect;
 
 		const UMovieSceneFadeSection* FadeSection = Cast<const UMovieSceneFadeSection>( &SectionObject );
 
diff --git a/Engine/Source/Editor/MovieSceneTools/Private/TrackEditors/ParticleTrackEditor.cpp b/Engine/Source/Editor/MovieSceneTools/Private/TrackEditors/ParticleTrackEditor.cpp
--- a/Engine/Source/Editor/MovieSceneTools/Private/TrackEditors/ParticleTrackEditor.cpp
+++ b/Engine/Source/Editor/MovieSceneTools/Private/TrackEditors/ParticleTrackEditor.cpp
@@ -79,7 +79,7 @@ int32 FParticleSection::OnPaintSection( FSequencerSectionPainter& InPainter ) co
 		return InPainter.LayerId + 1;
 	}
 
-	const ESlateDrawEffect::Type DrawEffects = InPainter.bParentEnabled ? ESlateDrawEffect::None : ESlateDrawEffect::DisabledEffect;
+	const ESlateDrawEffect DrawEffects = InPainter.bParentEnabled ? ESlateDrawEffect::None : ESlateDrawEffect::DisabledEffect;
 	UMovieSceneParticleSection* AnimSection = Cast<UMovieSceneParticleSection>( &Section );
 	const FTimeToPixel& TimeToPixelConverter = InPainter.GetTimeConverter();
 
diff --git a/Engine/Source/Editor/MovieSceneTools/Private/TrackEditors/SubTrackEditor.cpp b/Engine/Source/Editor/MovieSceneTools/Private/TrackEditors/SubTrackEditor.cpp
--- a/Engine/Source/Editor/MovieSceneTools/Private/TrackEditors/SubTrackEditor.cpp
+++ b/Engine/Source/Editor/MovieSceneTools/Private/TrackEditors/SubTrackEditor.cpp
@@ -121,7 +121,7 @@ public:
 		}
 
 		const float DrawScale = InPainter.SectionGeometry.Size.X / SectionSize;
-		const ESlateDrawEffect::Type DrawEffects = InPainter.bParentEnabled
+		const ESlateDrawEffect DrawEffects = InPainter.bParentEnabled
 			? ESlateDrawEffect::None
 			: ESlateDrawEffect::DisabledEffect;
 
diff --git a/Engine/Source/Editor/NiagaraEditor/Private/NiagaraEditorStyle.cpp b/Engine/Source/Editor/NiagaraEditor/Private/NiagaraEditorStyle.cpp
--- a/Engine/Source/Editor/NiagaraEditor/Private/NiagaraEditorStyle.cpp
+++ b/Engine/Source/Editor/NiagaraEditor/Private/NiagaraEditorStyle.cpp
@@ -37,7 +37,8 @@ FName FNiagaraEditorStyle::GetStyleSetName()
 #define BORDER_BRUSH( RelativePath, ... ) FSlateBorderBrush( Style->RootToContentDir( RelativePath, TEXT(".png") ), __VA_ARGS__ )
 #define TTF_FONT( RelativePath, ... ) FSlateFontInfo( Style->RootToContentDir( RelativePath, TEXT(".ttf") ), __VA_ARGS__ )
 #define OTF_FONT( RelativePath, ... ) FSlateFontInfo( Style->RootToContentDir( RelativePath, TEXT(".otf") ), __VA_ARGS__ )
-#define TTF_CORE_FONT( RelativePath, ... ) FSlateFontInfo( FPaths::EngineContentDir()  / "Slate"/ RelativePath + TEXT(".ttf"), __VA_ARGS__ )
+#define TTF_CORE_FONT( RelativePath, ... ) FSlateFontInfo( FPaths::EngineContentDir()  / "Slate" / RelativePath + TEXT(".ttf"), __VA_ARGS__ )
+#define BOX_CORE_BRUSH( RelativePath, ... ) FSlateBoxBrush( FPaths::EngineContentDir() / "Editor/Slate" / RelativePath + TEXT(".png"), __VA_ARGS__ )
 
 const FVector2D Icon16x16(16.0f, 16.0f);
 const FVector2D Icon20x20(20.0f, 20.0f);
@@ -46,7 +47,7 @@ const FVector2D Icon40x40(40.0f, 40.0f);
 TSharedRef< FSlateStyleSet > FNiagaraEditorStyle::Create()
 {
 	const FTextBlockStyle NormalText = FEditorStyle::GetWidgetStyle<FTextBlockStyle>("NormalText");
-	const FEditableTextBoxStyle NormalEditableTextBox = FEditorStyle::GetWidgetStyle<FEditableTextBoxStyle>("NormalEditableTextBox");
+	const FEditableTextBoxStyle NormalEditableTextBox = FCoreStyle::Get().GetWidgetStyle<FEditableTextBoxStyle>("NormalEditableTextBox");
 	const FSpinBoxStyle NormalSpinBox = FEditorStyle::GetWidgetStyle<FSpinBoxStyle>("SpinBox");
 
 	TSharedRef< FSlateStyleSet > Style = MakeShareable(new FSlateStyleSet("NiagaraEditorStyle"));
@@ -106,9 +107,9 @@ TSharedRef< FSlateStyleSet > FNiagaraEditorStyle::Create()
 	Style->Set("Niagara.CompileStatus.Good",    new IMAGE_BRUSH("Icons/CompileStatus_Good", Icon40x40));
 	Style->Set("Niagara.CompileStatus.Warning", new IMAGE_BRUSH("Icons/CompileStatus_Warning", Icon40x40));
 	Style->Set("Niagara.Asset.ReimportAsset.Needed", new IMAGE_BRUSH("Icons/icon_Reimport_Needed_40x", Icon40x40));
-	Style->Set("Niagara.Asset.ReimportAsset.Default",		new IMAGE_BRUSH("Icons/icon_Reimport_40x", Icon40x40));
+	Style->Set("Niagara.Asset.ReimportAsset.Default", new IMAGE_BRUSH("Icons/icon_Reimport_40x", Icon40x40));
 	
-	Style->Set("NiagaraEditor.MaterialWarningBorder", new BOX_BRUSH("Common/GroupBorderLight", FMargin(4.0f / 16.0f)));
+	Style->Set("NiagaraEditor.MaterialWarningBorder", new BOX_CORE_BRUSH("Common/GroupBorderLight", FMargin(4.0f / 16.0f)));
 
 	return Style;
 }
@@ -119,6 +120,7 @@ TSharedRef< FSlateStyleSet > FNiagaraEditorStyle::Create()
 #undef TTF_FONT
 #undef OTF_FONT
 #undef TTF_CORE_FONT
+#undef BOX_CORE_BRUSH
 
 void FNiagaraEditorStyle::ReloadTextures()
 {
diff --git a/Engine/Source/Editor/PListEditor/Public/PListEditor.h b/Engine/Source/Editor/PListEditor/Public/PListEditor.h
--- a/Engine/Source/Editor/PListEditor/Public/PListEditor.h
+++ b/Engine/Source/Editor/PListEditor/Public/PListEditor.h
@@ -2,9 +2,6 @@
 
 #pragma once
 
-#ifndef __PListEditor_h__
-#define __PListEditor_h__
-
 #include "CoreMinimal.h"
 #include "Modules/ModuleInterface.h"
 
@@ -29,5 +26,3 @@ private:
 
 
 };
-
-#endif
diff --git a/Engine/Source/Editor/PackagesDialog/Private/SPackagesDialog.h b/Engine/Source/Editor/PackagesDialog/Private/SPackagesDialog.h
--- a/Engine/Source/Editor/PackagesDialog/Private/SPackagesDialog.h
+++ b/Engine/Source/Editor/PackagesDialog/Private/SPackagesDialog.h
@@ -1,8 +1,6 @@
 // Copyright 1998-2017 Epic Games, Inc. All Rights Reserved.
 
-
-#ifndef __SPackagesDialog_h__
-#define __SPackagesDialog_h__
+#pragma once
 
 #include "CoreMinimal.h"
 #include "Misc/Attribute.h"
@@ -18,10 +16,6 @@
 #include "ISourceControlProvider.h"
 #include "ISourceControlModule.h"
 
-#ifdef _MSC_VER
-	#pragma once
-#endif
-
 class SCheckBox;
 
 /**
@@ -579,6 +573,3 @@ private:
 	/** The item associated with this row of data */
 	TSharedPtr< FPackageItem > Item;
 };
-
-
-#endif	// __SPackagesDialog_h__
diff --git a/Engine/Source/Editor/PackagesDialog/Public/PackagesDialog.h b/Engine/Source/Editor/PackagesDialog/Public/PackagesDialog.h
--- a/Engine/Source/Editor/PackagesDialog/Public/PackagesDialog.h
+++ b/Engine/Source/Editor/PackagesDialog/Public/PackagesDialog.h
@@ -1,17 +1,11 @@
 // Copyright 1998-2017 Epic Games, Inc. All Rights Reserved.
 
-
-#ifndef __PackagesDialog_h__
-#define __PackagesDialog_h__
+#pragma once
 
 #include "CoreMinimal.h"
 #include "Misc/Attribute.h"
 #include "Modules/ModuleInterface.h"
 
-#ifdef _MSC_VER
-	#pragma once
-#endif
-
 class SPackagesDialog;
 class SWindow;
 enum class ECheckBoxState : uint8;
@@ -164,5 +158,3 @@ private:
 	TArray<UPackage*> UncheckedPackages;
 	TArray<UPackage*> UndeterminedPackages;
 };
-
-#endif	// __PackagesDialog_h__
diff --git a/Engine/Source/Editor/Persona/Private/AnimViewportLODCommands.h b/Engine/Source/Editor/Persona/Private/AnimViewportLODCommands.h
--- a/Engine/Source/Editor/Persona/Private/AnimViewportLODCommands.h
+++ b/Engine/Source/Editor/Persona/Private/AnimViewportLODCommands.h
@@ -1,8 +1,6 @@
 // Copyright 1998-2017 Epic Games, Inc. All Rights Reserved.
 
-
-#ifndef __AnimViewportLODCommands_h_
-#define __AnimViewportLODCommands_h_
+#pragma once
 
 #include "CoreMinimal.h"
 #include "Framework/Commands/Commands.h"
@@ -35,6 +33,3 @@ public:
 	/** Registers our commands with the binding system */
 	virtual void RegisterCommands() override;
 };
-
-
-#endif //__AnimViewportShowCommands_h_
diff --git a/Engine/Source/Editor/Persona/Private/AnimViewportMenuCommands.h b/Engine/Source/Editor/Persona/Private/AnimViewportMenuCommands.h
--- a/Engine/Source/Editor/Persona/Private/AnimViewportMenuCommands.h
+++ b/Engine/Source/Editor/Persona/Private/AnimViewportMenuCommands.h
@@ -1,8 +1,6 @@
 // Copyright 1998-2017 Epic Games, Inc. All Rights Reserved.
 
-
-#ifndef __AnimViewportMenuCommands_h_
-#define __AnimViewportMenuCommands_h_
+#pragma once
 
 #include "CoreMinimal.h"
 #include "Framework/Commands/Commands.h"
@@ -49,6 +47,3 @@ public:
 	/** Registers our commands with the binding system */
 	virtual void RegisterCommands() override;
 };
-
-
-#endif //__AnimViewportMenuCommands_h_
diff --git a/Engine/Source/Editor/Persona/Private/AnimViewportShowCommands.h b/Engine/Source/Editor/Persona/Private/AnimViewportShowCommands.h
--- a/Engine/Source/Editor/Persona/Private/AnimViewportShowCommands.h
+++ b/Engine/Source/Editor/Persona/Private/AnimViewportShowCommands.h
@@ -1,8 +1,6 @@
 // Copyright 1998-2017 Epic Games, Inc. All Rights Reserved.
 
-
-#ifndef __AnimViewportShowCommands_h_
-#define __AnimViewportShowCommands_h_
+#pragma once
 
 #include "CoreMinimal.h"
 #include "EngineDefines.h"
@@ -152,6 +150,3 @@ public:
 	/** Registers our commands with the binding system */
 	virtual void RegisterCommands() override;
 };
-
-
-#endif //__AnimViewportShowCommands_h_
diff --git a/Engine/Source/Editor/Persona/Private/SAnimMontageScrubPanel.h b/Engine/Source/Editor/Persona/Private/SAnimMontageScrubPanel.h
--- a/Engine/Source/Editor/Persona/Private/SAnimMontageScrubPanel.h
+++ b/Engine/Source/Editor/Persona/Private/SAnimMontageScrubPanel.h
@@ -1,7 +1,6 @@
 // Copyright 1998-2017 Epic Games, Inc. All Rights Reserved.
 
-#ifndef __SAnimMontageScrubPanel_h__
-#define __SAnimMontageScrubPanel_h__
+#pragma once
 
 #include "CoreMinimal.h"
 #include "Input/Reply.h"
@@ -56,4 +55,3 @@ protected:
 private:
 	TWeakPtr<SMontageEditor> MontageEditor;
 };
-#endif // __SAnimMontageScrubPanel_h__
diff --git a/Engine/Source/Editor/Persona/Private/SAnimationScrubPanel.h b/Engine/Source/Editor/Persona/Private/SAnimationScrubPanel.h
--- a/Engine/Source/Editor/Persona/Private/SAnimationScrubPanel.h
+++ b/Engine/Source/Editor/Persona/Private/SAnimationScrubPanel.h
@@ -1,7 +1,6 @@
 // Copyright 1998-2017 Epic Games, Inc. All Rights Reserved.
 
-#ifndef __SAnimationScrubPanel_h__
-#define __SAnimationScrubPanel_h__
+#pragma once
 
 #include "CoreMinimal.h"
 #include "Widgets/DeclarativeSyntaxSupport.h"
@@ -113,4 +112,3 @@ protected:
 
 	bool GetDisplayDrag() const;
 };
-#endif // __SAnimationScrubPanel_h__
diff --git a/Engine/Source/Editor/Persona/Public/BoneSelectionWidget.h b/Engine/Source/Editor/Persona/Public/BoneSelectionWidget.h
--- a/Engine/Source/Editor/Persona/Public/BoneSelectionWidget.h
+++ b/Engine/Source/Editor/Persona/Public/BoneSelectionWidget.h
@@ -1,8 +1,5 @@
 // Copyright 1998-2017 Epic Games, Inc. All Rights Reserved.
 
-#ifndef __BoneSelectionWidget_h__
-#define __BoneSelectionWidget_h__
-
 #pragma once
 
 #include "CoreMinimal.h"
@@ -136,5 +133,3 @@ private:
 	// Cache supplied tooltip
 	FText SuppliedToolTip;
 };
-
-#endif		//__BoneSelectionWidget_h__
diff --git a/Engine/Source/Editor/PropertyEditor/Private/DetailPropertyRow.cpp b/Engine/Source/Editor/PropertyEditor/Private/DetailPropertyRow.cpp
--- a/Engine/Source/Editor/PropertyEditor/Private/DetailPropertyRow.cpp
+++ b/Engine/Source/Editor/PropertyEditor/Private/DetailPropertyRow.cpp
@@ -401,6 +401,11 @@ void FDetailPropertyRow::MakeNameOrKeyWidget( FDetailWidgetRow& Row, const TShar
 	{
 		Slot.Padding(0.0f, 0.0f, 20.0f, 0.0f);
 	}
+	else if (InCustomRow.IsValid())
+	{
+		//Allow custom name slot to fill all the area. If the user add a SHorizontalBox with left and right align slot
+		Slot.FillWidth(1.0f);
+	}
 	else
 	{
 		Slot.AutoWidth();
diff --git a/Engine/Source/Editor/PropertyEditor/Private/PropertyHandleImpl.h b/Engine/Source/Editor/PropertyEditor/Private/PropertyHandleImpl.h
--- a/Engine/Source/Editor/PropertyEditor/Private/PropertyHandleImpl.h
+++ b/Engine/Source/Editor/PropertyEditor/Private/PropertyHandleImpl.h
@@ -666,7 +666,7 @@ public:
 	virtual FPropertyAccess::Result GetNumElements(uint32& OutNumElements) override;
 	virtual void SetOnNumElementsChanged(FSimpleDelegate& InOnNumElementsChanged) override;
 	virtual bool HasDocumentation() override { return true; }
-	virtual FString GetDocumentationLink() override { return FString("Engine/UI/LevelEditor/Details/Properties/Set/"); }	// @todo: needs a better documentation page
+	virtual FString GetDocumentationLink() override { return FString("Engine/UI/LevelEditor/Details/Properties/Set/"); }
 	virtual FString GetDocumentationExcerptName() override { return FString("Sets"); }
 private:
 	/**
diff --git a/Engine/Source/Editor/SceneOutliner/Private/SOutlinerTreeView.cpp b/Engine/Source/Editor/SceneOutliner/Private/SOutlinerTreeView.cpp
--- a/Engine/Source/Editor/SceneOutliner/Private/SOutlinerTreeView.cpp
+++ b/Engine/Source/Editor/SceneOutliner/Private/SOutlinerTreeView.cpp
@@ -332,7 +332,7 @@ namespace SceneOutliner
 
 		// See if a disabled effect should be used
 		bool bEnabled = ShouldBeEnabled( bParentEnabled );
-		ESlateDrawEffect::Type DrawEffects = (bEnabled) ? ESlateDrawEffect::None : ESlateDrawEffect::DisabledEffect;
+		ESlateDrawEffect DrawEffects = (bEnabled) ? ESlateDrawEffect::None : ESlateDrawEffect::DisabledEffect;
 
 		const double CurrentTime = FSlateApplication::Get().GetCurrentTime();
 
diff --git a/Engine/Source/Editor/Sequencer/Private/SSequencerSection.cpp b/Engine/Source/Editor/Sequencer/Private/SSequencerSection.cpp
--- a/Engine/Source/Editor/Sequencer/Private/SSequencerSection.cpp
+++ b/Engine/Source/Editor/Sequencer/Private/SSequencerSection.cpp
@@ -31,7 +31,7 @@ struct FSequencerSectionPainterImpl : FSequencerSectionPainter
 
 	virtual int32 PaintSectionBackground(const FLinearColor& Tint) override
 	{
-		const ESlateDrawEffect::Type DrawEffects = bParentEnabled
+		const ESlateDrawEffect DrawEffects = bParentEnabled
 			? ESlateDrawEffect::None
 			: ESlateDrawEffect::DisabledEffect;
 
@@ -345,7 +345,7 @@ int32 SSequencerSection::OnPaint( const FPaintArgs& Args, const FGeometry& Allot
 	
 	const bool bEnabled = bParentEnabled && SectionObject.IsActive();
 	const bool bLocked = SectionObject.IsLocked();
-	const ESlateDrawEffect::Type DrawEffects = bEnabled ? ESlateDrawEffect::None : ESlateDrawEffect::DisabledEffect;
+	const ESlateDrawEffect DrawEffects = bEnabled ? ESlateDrawEffect::None : ESlateDrawEffect::DisabledEffect;
 
 	FGeometry SectionGeometry = MakeSectionGeometryWithoutHandles( AllottedGeometry, SectionInterface );
 
@@ -492,7 +492,7 @@ void SSequencerSection::PaintKeys( FSequencerSectionPainter& InPainter, const FW
 	const FSlateBrush* SquareKeyBrush = FEditorStyle::GetBrush(SquareKeyBrushName);
 	const FSlateBrush* TriangleKeyBrush = FEditorStyle::GetBrush(TriangleKeyBrushName);
 
-	const ESlateDrawEffect::Type DrawEffects = InPainter.bParentEnabled
+	const ESlateDrawEffect DrawEffects = InPainter.bParentEnabled
 		? ESlateDrawEffect::None
 		: ESlateDrawEffect::DisabledEffect;
 
@@ -753,7 +753,7 @@ void SSequencerSection::PaintKeys( FSequencerSectionPainter& InPainter, const FW
 }
 
 
-void SSequencerSection::DrawSectionHandles( const FGeometry& AllottedGeometry, const FSlateRect& MyClippingRect, FSlateWindowElementList& OutDrawElements, int32 LayerId, ESlateDrawEffect::Type DrawEffects, FLinearColor SelectionColor, const ISequencerHotspot* Hotspot ) const
+void SSequencerSection::DrawSectionHandles( const FGeometry& AllottedGeometry, const FSlateRect& MyClippingRect, FSlateWindowElementList& OutDrawElements, int32 LayerId, ESlateDrawEffect DrawEffects, FLinearColor SelectionColor, const ISequencerHotspot* Hotspot ) const
 {
 	UMovieSceneSection* Section = SectionInterface->GetSectionObject();
 	if (!Section || Section->IsInfinite())
diff --git a/Engine/Source/Editor/Sequencer/Private/SSequencerSection.h b/Engine/Source/Editor/Sequencer/Private/SSequencerSection.h
--- a/Engine/Source/Editor/Sequencer/Private/SSequencerSection.h
+++ b/Engine/Source/Editor/Sequencer/Private/SSequencerSection.h
@@ -91,7 +91,7 @@ private:
 	/**
 	 * Draw the section resize handles.
 	 */
-	void DrawSectionHandles( const FGeometry& AllottedGeometry, const FSlateRect& MyClippingRect, FSlateWindowElementList& OutDrawElements, int32 LayerId, ESlateDrawEffect::Type DrawEffects, FLinearColor SelectionColor, const ISequencerHotspot* Hotspot ) const;
+	void DrawSectionHandles( const FGeometry& AllottedGeometry, const FSlateRect& MyClippingRect, FSlateWindowElementList& OutDrawElements, int32 LayerId, ESlateDrawEffect DrawEffects, FLinearColor SelectionColor, const ISequencerHotspot* Hotspot ) const;
 
 	/** @return the sequencer interface */
 	FSequencer& GetSequencer() const;
diff --git a/Engine/Source/Editor/Sequencer/Private/SequencerTimeSliderController.cpp b/Engine/Source/Editor/Sequencer/Private/SequencerTimeSliderController.cpp
--- a/Engine/Source/Editor/Sequencer/Private/SequencerTimeSliderController.cpp
+++ b/Engine/Source/Editor/Sequencer/Private/SequencerTimeSliderController.cpp
@@ -127,7 +127,7 @@ struct FDrawTickArgs
 	/** Start layer for elements */
 	int32 StartLayer;
 	/** Draw effects to apply */
-	ESlateDrawEffect::Type DrawEffects;
+	ESlateDrawEffect DrawEffects;
 	/** Whether or not to only draw major ticks */
 	bool bOnlyDrawMajorTicks;
 	/** Whether or not to mirror labels */
@@ -249,7 +249,7 @@ void FSequencerTimeSliderController::DrawTicks( FSlateWindowElementList& OutDraw
 int32 FSequencerTimeSliderController::OnPaintTimeSlider( bool bMirrorLabels, const FGeometry& AllottedGeometry, const FSlateRect& MyClippingRect, FSlateWindowElementList& OutDrawElements, int32 LayerId, const FWidgetStyle& InWidgetStyle, bool bParentEnabled ) const
 {
 	const bool bEnabled = bParentEnabled;
-	const ESlateDrawEffect::Type DrawEffects = bEnabled ? ESlateDrawEffect::None : ESlateDrawEffect::DisabledEffect;
+	const ESlateDrawEffect DrawEffects = bEnabled ? ESlateDrawEffect::None : ESlateDrawEffect::DisabledEffect;
 
 	TRange<float> LocalViewRange = TimeSliderArgs.ViewRange.Get();
 	const float LocalViewRangeMin = LocalViewRange.GetLowerBoundValue();
@@ -969,7 +969,7 @@ FCursorReply FSequencerTimeSliderController::OnCursorQuery( TSharedRef<const SWi
 
 int32 FSequencerTimeSliderController::OnPaintSectionView( const FGeometry& AllottedGeometry, const FSlateRect& MyClippingRect, FSlateWindowElementList& OutDrawElements, int32 LayerId, bool bEnabled, const FPaintSectionAreaViewArgs& Args ) const
 {
-	const ESlateDrawEffect::Type DrawEffects = bEnabled ? ESlateDrawEffect::None : ESlateDrawEffect::DisabledEffect;
+	const ESlateDrawEffect DrawEffects = bEnabled ? ESlateDrawEffect::None : ESlateDrawEffect::DisabledEffect;
 
 	TRange<float> LocalViewRange = TimeSliderArgs.ViewRange.Get();
 	float LocalScrubPosition = TimeSliderArgs.ScrubPosition.Get();
diff --git a/Engine/Source/Editor/SkeletonEditor/Private/SkeletonTreeCommands.h b/Engine/Source/Editor/SkeletonEditor/Private/SkeletonTreeCommands.h
--- a/Engine/Source/Editor/SkeletonEditor/Private/SkeletonTreeCommands.h
+++ b/Engine/Source/Editor/SkeletonEditor/Private/SkeletonTreeCommands.h
@@ -1,7 +1,6 @@
 // Copyright 1998-2017 Epic Games, Inc. All Rights Reserved.
 
-#ifndef __SkeletonTreeCommands_h_
-#define __SkeletonTreeCommands_h_
+#pragma once
 
 #include "CoreMinimal.h"
 #include "EditorStyleSet.h"
@@ -97,5 +96,3 @@ public:
 	/** Focus the camera on the current selection */
 	TSharedPtr< FUICommandInfo > FocusCamera;
 };
-
-#endif // __SkeletonTreeCommands_h_
diff --git a/Engine/Source/Editor/StaticMeshEditor/Private/SStaticMeshEditorViewport.cpp b/Engine/Source/Editor/StaticMeshEditor/Private/SStaticMeshEditorViewport.cpp
--- a/Engine/Source/Editor/StaticMeshEditor/Private/SStaticMeshEditorViewport.cpp
+++ b/Engine/Source/Editor/StaticMeshEditor/Private/SStaticMeshEditorViewport.cpp
@@ -256,7 +256,7 @@ void SStaticMeshEditorViewport::SetPreviewMesh(UStaticMesh* InStaticMesh)
 	EditorViewportClient->SetPreviewMesh(InStaticMesh, PreviewMeshComponent);
 }
 
-void SStaticMeshEditorViewport::UpdatePreviewMesh(UStaticMesh* InStaticMesh)
+void SStaticMeshEditorViewport::UpdatePreviewMesh(UStaticMesh* InStaticMesh, bool bResetCamera/*= true*/)
 {
 	{
 		const int32 SocketedComponentCount = SocketPreviewMeshComponents.Num();
@@ -300,7 +300,7 @@ void SStaticMeshEditorViewport::UpdatePreviewMesh(UStaticMesh* InStaticMesh)
 		}
 	}
 
-	EditorViewportClient->SetPreviewMesh(InStaticMesh, PreviewMeshComponent);
+	EditorViewportClient->SetPreviewMesh(InStaticMesh, PreviewMeshComponent, bResetCamera);
 }
 
 bool SStaticMeshEditorViewport::IsVisible() const
diff --git a/Engine/Source/Editor/StaticMeshEditor/Private/SStaticMeshEditorViewport.h b/Engine/Source/Editor/StaticMeshEditor/Private/SStaticMeshEditorViewport.h
--- a/Engine/Source/Editor/StaticMeshEditor/Private/SStaticMeshEditorViewport.h
+++ b/Engine/Source/Editor/StaticMeshEditor/Private/SStaticMeshEditorViewport.h
@@ -70,7 +70,7 @@ public:
 	 *
 	 *	@param	InStaticMesh		The static mesh being viewed in the editor.
 	 */
-	void UpdatePreviewMesh(UStaticMesh* InStaticMesh);
+	void UpdatePreviewMesh(UStaticMesh* InStaticMesh, bool bResetCamera=true);
 
 	/** Retrieves the selected edge set. */
 	TSet< int32 >& GetSelectedEdges();
diff --git a/Engine/Source/Editor/StaticMeshEditor/Private/StaticMeshEditor.cpp b/Engine/Source/Editor/StaticMeshEditor/Private/StaticMeshEditor.cpp
--- a/Engine/Source/Editor/StaticMeshEditor/Private/StaticMeshEditor.cpp
+++ b/Engine/Source/Editor/StaticMeshEditor/Private/StaticMeshEditor.cpp
@@ -1509,7 +1509,7 @@ UStaticMesh* FStaticMeshEditor::GetFirstSelectedStaticMeshInContentBrowser() con
 	return NULL;
 }
 
-void FStaticMeshEditor::SetEditorMesh(UStaticMesh* InStaticMesh)
+void FStaticMeshEditor::SetEditorMesh(UStaticMesh* InStaticMesh, bool bResetCamera/*=true*/)
 {
 	ClearSelectedPrims();
 
@@ -1569,7 +1569,7 @@ void FStaticMeshEditor::SetEditorMesh(UStaticMesh* InStaticMesh)
 	// Set the details view.
 	StaticMeshDetailsView->SetObject(StaticMesh);
 
-	Viewport->UpdatePreviewMesh(StaticMesh);
+	Viewport->UpdatePreviewMesh(StaticMesh, bResetCamera);
 	Viewport->RefreshViewport();
 }
 
@@ -1867,7 +1867,9 @@ void FStaticMeshEditor::OnObjectReimported(UObject* InObject)
 	// Make sure we are using the object that is being reimported, otherwise a lot of needless work could occur.
 	if(StaticMesh == InObject)
 	{
-		SetEditorMesh(Cast<UStaticMesh>(InObject));
+		//When we re-import we want to avoid moving the camera in the staticmesh editor
+		bool bResetCamera = false;
+		SetEditorMesh(Cast<UStaticMesh>(InObject), bResetCamera);
 
 		if (SocketManager.IsValid())
 		{
diff --git a/Engine/Source/Editor/StaticMeshEditor/Private/StaticMeshEditor.h b/Engine/Source/Editor/StaticMeshEditor/Private/StaticMeshEditor.h
--- a/Engine/Source/Editor/StaticMeshEditor/Private/StaticMeshEditor.h
+++ b/Engine/Source/Editor/StaticMeshEditor/Private/StaticMeshEditor.h
@@ -164,7 +164,7 @@ private:
 	 *
 	 *	@param	InStaticMesh		The static mesh to use for the editor.
 	 */
-	void SetEditorMesh(UStaticMesh* InStaticMesh);
+	void SetEditorMesh(UStaticMesh* InStaticMesh, bool bResetCamera=true);
 
 	/** Helper function for generating K-DOP collision geometry. */
 	void GenerateKDop(const FVector* Directions, uint32 NumDirections);
diff --git a/Engine/Source/Editor/StaticMeshEditor/Private/StaticMeshEditorViewportClient.cpp b/Engine/Source/Editor/StaticMeshEditor/Private/StaticMeshEditorViewportClient.cpp
--- a/Engine/Source/Editor/StaticMeshEditor/Private/StaticMeshEditorViewportClient.cpp
+++ b/Engine/Source/Editor/StaticMeshEditor/Private/StaticMeshEditorViewportClient.cpp
@@ -1182,10 +1182,16 @@ void FStaticMeshEditorViewportClient::PerspectiveCameraMoved()
 {
 	FEditorViewportClient::PerspectiveCameraMoved();
 
+	// If in the process of transitioning to a new location, don't update the orbit camera position.
+	// On the final update of the transition, we will get here with IsPlaying()==false, and the editor camera position will
+	// be correctly updated.
+	if (GetViewTransform().IsPlaying())
+	{
+		return;
+	}
+
 	// The static mesh editor saves the camera position in terms of an orbit camera, so ensure 
 	// that orbit mode is enabled before we store the current transform information.
-	// ToggleOrbitCamera affects the stored camera location, and information is lost in the transformation,
-	// so remember the current position here so it can be restored afterwards.
 	const bool bWasOrbit = bUsingOrbitCamera;
 	const FVector OldCameraLocation = GetViewLocation();
 	const FRotator OldCameraRotation = GetViewRotation();
@@ -1200,11 +1206,6 @@ void FStaticMeshEditorViewportClient::PerspectiveCameraMoved()
 		);
 
 	ToggleOrbitCamera(bWasOrbit);
-	if (!GetViewTransform().IsPlaying())
-	{
-		SetViewLocation(OldCameraLocation);
-		SetViewRotation(OldCameraRotation);
-	}
 }
 
 void FStaticMeshEditorViewportClient::OnAssetViewerSettingsChanged(const FName& InPropertyName)
@@ -1234,7 +1235,7 @@ void FStaticMeshEditorViewportClient::SetFloorAndEnvironmentVisibility(const boo
 	AdvancedPreviewScene->SetEnvironmentVisibility(bVisible, true);
 }
 
-void FStaticMeshEditorViewportClient::SetPreviewMesh(UStaticMesh* InStaticMesh, UStaticMeshComponent* InStaticMeshComponent)
+void FStaticMeshEditorViewportClient::SetPreviewMesh(UStaticMesh* InStaticMesh, UStaticMeshComponent* InStaticMeshComponent, bool bResetCamera)
 {
 	StaticMesh = InStaticMesh;
 	StaticMeshComponent = InStaticMeshComponent;
@@ -1244,44 +1245,47 @@ void FStaticMeshEditorViewportClient::SetPreviewMesh(UStaticMesh* InStaticMesh,
 		StaticMeshComponent->bDrawMeshCollisionWireframe = bShowCollision;
 		StaticMeshComponent->MarkRenderStateDirty();
 	}
-
-	// If we have a thumbnail transform, we will favor that over the camera position as the user may have customized this for a nice view
-	// If we have neither a custom thumbnail nor a valid camera position, then we'll just use the default thumbnail transform 
-	const USceneThumbnailInfo* const AssetThumbnailInfo = Cast<USceneThumbnailInfo>(StaticMesh->ThumbnailInfo);
-	const USceneThumbnailInfo* const DefaultThumbnailInfo = USceneThumbnailInfo::StaticClass()->GetDefaultObject<USceneThumbnailInfo>();
-
-	// Prefer the asset thumbnail if available
-	const USceneThumbnailInfo* const ThumbnailInfo = (AssetThumbnailInfo) ? AssetThumbnailInfo : DefaultThumbnailInfo;
-	check(ThumbnailInfo);
-
-	FRotator ThumbnailAngle;
-	ThumbnailAngle.Pitch = ThumbnailInfo->OrbitPitch;
-	ThumbnailAngle.Yaw = ThumbnailInfo->OrbitYaw;
-	ThumbnailAngle.Roll = 0;
-	const float ThumbnailDistance = ThumbnailInfo->OrbitZoom;
-
-	const float CameraY = StaticMesh->GetBounds().SphereRadius / (75.0f * PI / 360.0f);
-	SetCameraSetup(
-		FVector::ZeroVector, 
-		ThumbnailAngle,
-		FVector(0.0f, CameraY + ThumbnailDistance - AutoViewportOrbitCameraTranslate, 0.0f), 
-		StaticMesh->GetBounds().Origin, 
-		-FVector(0, CameraY, 0), 
-		FRotator(0,90.f,0)
-		);
-
-	if(!AssetThumbnailInfo && StaticMesh->EditorCameraPosition.bIsSet)
+	
+	if (bResetCamera)
 	{
-		// The static mesh editor saves the camera position in terms of an orbit camera, so ensure 
-		// that orbit mode is enabled before we set the new transform information
-		const bool bWasOrbit = bUsingOrbitCamera;
-		ToggleOrbitCamera(true);
+		// If we have a thumbnail transform, we will favor that over the camera position as the user may have customized this for a nice view
+		// If we have neither a custom thumbnail nor a valid camera position, then we'll just use the default thumbnail transform 
+		const USceneThumbnailInfo* const AssetThumbnailInfo = Cast<USceneThumbnailInfo>(StaticMesh->ThumbnailInfo);
+		const USceneThumbnailInfo* const DefaultThumbnailInfo = USceneThumbnailInfo::StaticClass()->GetDefaultObject<USceneThumbnailInfo>();
+
+		// Prefer the asset thumbnail if available
+		const USceneThumbnailInfo* const ThumbnailInfo = (AssetThumbnailInfo) ? AssetThumbnailInfo : DefaultThumbnailInfo;
+		check(ThumbnailInfo);
+
+		FRotator ThumbnailAngle;
+		ThumbnailAngle.Pitch = ThumbnailInfo->OrbitPitch;
+		ThumbnailAngle.Yaw = ThumbnailInfo->OrbitYaw;
+		ThumbnailAngle.Roll = 0;
+		const float ThumbnailDistance = ThumbnailInfo->OrbitZoom;
+
+		const float CameraY = StaticMesh->GetBounds().SphereRadius / (75.0f * PI / 360.0f);
+		SetCameraSetup(
+			FVector::ZeroVector,
+			ThumbnailAngle,
+			FVector(0.0f, CameraY + ThumbnailDistance - AutoViewportOrbitCameraTranslate, 0.0f),
+			StaticMesh->GetBounds().Origin,
+			-FVector(0, CameraY, 0),
+			FRotator(0, 90.f, 0)
+			);
+
+		if (!AssetThumbnailInfo && StaticMesh->EditorCameraPosition.bIsSet)
+		{
+			// The static mesh editor saves the camera position in terms of an orbit camera, so ensure 
+			// that orbit mode is enabled before we set the new transform information
+			const bool bWasOrbit = bUsingOrbitCamera;
+			ToggleOrbitCamera(true);
 
-		SetViewRotation(StaticMesh->EditorCameraPosition.CamOrbitRotation);
-		SetViewLocation(StaticMesh->EditorCameraPosition.CamOrbitPoint + StaticMesh->EditorCameraPosition.CamOrbitZoom);
-		SetLookAtLocation(StaticMesh->EditorCameraPosition.CamOrbitPoint);
+			SetViewRotation(StaticMesh->EditorCameraPosition.CamOrbitRotation);
+			SetViewLocation(StaticMesh->EditorCameraPosition.CamOrbitPoint + StaticMesh->EditorCameraPosition.CamOrbitZoom);
+			SetLookAtLocation(StaticMesh->EditorCameraPosition.CamOrbitPoint);
 
-		ToggleOrbitCamera(bWasOrbit);
+			ToggleOrbitCamera(bWasOrbit);
+		}
 	}
 }
 
diff --git a/Engine/Source/Editor/StaticMeshEditor/Private/StaticMeshEditorViewportClient.h b/Engine/Source/Editor/StaticMeshEditor/Private/StaticMeshEditorViewportClient.h
--- a/Engine/Source/Editor/StaticMeshEditor/Private/StaticMeshEditorViewportClient.h
+++ b/Engine/Source/Editor/StaticMeshEditor/Private/StaticMeshEditorViewportClient.h
@@ -49,7 +49,7 @@ public:
 	 *	@param	InStaticMesh				The static mesh handled by the editor.
 	 *	@param	InStaticMeshComponent		The static mesh component component from the viewport's scene.
 	 */
-	void SetPreviewMesh(UStaticMesh* InStaticMesh, UStaticMeshComponent* InPreviewStaticMeshComponent);
+	void SetPreviewMesh(UStaticMesh* InStaticMesh, UStaticMeshComponent* InPreviewStaticMeshComponent, bool bResetCamera=true);
 
 	/** Retrieves the selected edge set. */
 	TSet< int32 >& GetSelectedEdges();
diff --git a/Engine/Source/Editor/StaticMeshEditor/Public/StaticMeshEditorModule.h b/Engine/Source/Editor/StaticMeshEditor/Public/StaticMeshEditorModule.h
--- a/Engine/Source/Editor/StaticMeshEditor/Public/StaticMeshEditorModule.h
+++ b/Engine/Source/Editor/StaticMeshEditor/Public/StaticMeshEditorModule.h
@@ -1,7 +1,6 @@
 // Copyright 1998-2017 Epic Games, Inc. All Rights Reserved.
 
-#ifndef __StaticMeshEditorModule_h__
-#define __StaticMeshEditorModule_h__
+#pragma once
 
 #include "CoreMinimal.h"
 #include "Modules/ModuleInterface.h"
@@ -24,5 +23,3 @@ public:
 	 */
 	virtual TSharedRef<IStaticMeshEditor> CreateStaticMeshEditor( const EToolkitMode::Type Mode, const TSharedPtr< IToolkitHost >& InitToolkitHost, UStaticMesh* StaticMesh ) = 0;
 };
-
-#endif // __StaticMeshEditorModule_h__
diff --git a/Engine/Source/Editor/StringTableEditor/Private/AssetTypeActions_StringTable.cpp b/Engine/Source/Editor/StringTableEditor/Private/AssetTypeActions_StringTable.cpp
new file mode 100644
--- /dev/null
+++ b/Engine/Source/Editor/StringTableEditor/Private/AssetTypeActions_StringTable.cpp
@@ -0,0 +1,20 @@
+// Copyright 1998-2017 Epic Games, Inc. All Rights Reserved.
+
+#include "AssetTypeActions_StringTable.h"
+#include "StringTableEditorModule.h"
+
+#define LOCTEXT_NAMESPACE "AssetTypeActions"
+
+void FAssetTypeActions_StringTable::OpenAssetEditor(const TArray<UObject*>& InObjects, TSharedPtr<IToolkitHost> EditWithinLevelEditor)
+{
+	FStringTableEditorModule& StringTableEditorModule = FModuleManager::LoadModuleChecked<FStringTableEditorModule>("StringTableEditor");
+	for (UObject* ObjToEdit : InObjects)
+	{
+		if (UStringTable* StringTable = Cast<UStringTable>(ObjToEdit))
+		{
+			StringTableEditorModule.CreateStringTableEditor(EToolkitMode::Standalone, EditWithinLevelEditor, StringTable);
+		}
+	}
+}
+
+#undef LOCTEXT_NAMESPACE
diff --git a/Engine/Source/Editor/StringTableEditor/Private/AssetTypeActions_StringTable.h b/Engine/Source/Editor/StringTableEditor/Private/AssetTypeActions_StringTable.h
new file mode 100644
--- /dev/null
+++ b/Engine/Source/Editor/StringTableEditor/Private/AssetTypeActions_StringTable.h
@@ -0,0 +1,23 @@
+// Copyright 1998-2017 Epic Games, Inc. All Rights Reserved.
+
+#pragma once
+
+#include "CoreMinimal.h"
+#include "Toolkits/IToolkitHost.h"
+#include "AssetTypeActions_Base.h"
+#include "StringTable.h"
+
+class FMenuBuilder;
+
+class FAssetTypeActions_StringTable : public FAssetTypeActions_Base
+{
+public:
+	// IAssetTypeActions Implementation
+	virtual FColor GetTypeColor() const override { return FColor(255, 196, 128); }
+	virtual bool CanLocalize() const override { return false; }
+	virtual uint32 GetCategories() override { return EAssetTypeCategories::Misc; }
+	virtual FText GetName() const override { return NSLOCTEXT("AssetTypeActions", "AssetTypeActions_StringTable", "String Table"); }
+	virtual UClass* GetSupportedClass() const override { return UStringTable::StaticClass(); }
+	virtual void OpenAssetEditor(const TArray<UObject*>& InObjects, TSharedPtr<class IToolkitHost> EditWithinLevelEditor = TSharedPtr<IToolkitHost>()) override;
+	// End IAssetTypeActions
+};
diff --git a/Engine/Source/Editor/StringTableEditor/Public/StringTableEditorModule.h b/Engine/Source/Editor/StringTableEditor/Public/StringTableEditorModule.h
new file mode 100644
--- /dev/null
+++ b/Engine/Source/Editor/StringTableEditor/Public/StringTableEditorModule.h
@@ -0,0 +1,40 @@
+// Copyright 1998-2017 Epic Games, Inc. All Rights Reserved.
+
+#pragma once
+
+#include "CoreMinimal.h"
+#include "Toolkits/IToolkitHost.h"
+#include "Modules/ModuleInterface.h"
+#include "Toolkits/AssetEditorToolkit.h"
+
+class IStringTableEditor;
+class UStringTable;
+
+/** String Table Editor module */
+class STRINGTABLEEDITOR_API FStringTableEditorModule : public IModuleInterface, public IHasMenuExtensibility
+{
+public:
+	// IModuleInterface
+	virtual void StartupModule() override;
+	virtual void ShutdownModule() override;
+
+	/**
+	 * Creates an instance of string table editor object.
+	 *
+	 * @param	Mode					Mode that this editor should operate in
+	 * @param	InitToolkitHost			When Mode is WorldCentric, this is the level editor instance to spawn this editor within
+	 * @param	StringTable				The string table to start editing
+	 *
+	 * @return	Interface to the new string table editor
+	 */
+	TSharedRef<IStringTableEditor> CreateStringTableEditor(const EToolkitMode::Type Mode, const TSharedPtr<class IToolkitHost>& InitToolkitHost, UStringTable* StringTable);
+
+	/** Gets the extensibility managers for outside entities to extend static mesh editor's menus and toolbars */
+	virtual TSharedPtr<FExtensibilityManager> GetMenuExtensibilityManager() override { return MenuExtensibilityManager; }
+
+	/** StringTable Editor app identifier string */
+	static const FName StringTableEditorAppIdentifier;
+
+private:
+	TSharedPtr<FExtensibilityManager> MenuExtensibilityManager;
+};
diff --git a/Engine/Source/Editor/UMGEditor/Private/Designer/SDesignerView.cpp b/Engine/Source/Editor/UMGEditor/Private/Designer/SDesignerView.cpp
--- a/Engine/Source/Editor/UMGEditor/Private/Designer/SDesignerView.cpp
+++ b/Engine/Source/Editor/UMGEditor/Private/Designer/SDesignerView.cpp
@@ -933,7 +933,8 @@ void SDesignerView::PopDesignerMessage()
 
 void SDesignerView::OnEditorSelectionChanged()
 {
-	TSet<FWidgetReference> PendingSelectedWidgets = BlueprintEditor.Pin()->GetSelectedWidgets();
+	TSharedPtr<FWidgetBlueprintEditor> BPEd = BlueprintEditor.Pin();
+	TSet<FWidgetReference> PendingSelectedWidgets = BPEd->GetSelectedWidgets();
 
 	// Notify all widgets that are no longer selected.
 	for ( FWidgetReference& WidgetRef : SelectedWidgetsCache )
@@ -942,6 +943,15 @@ void SDesignerView::OnEditorSelectionChanged()
 		{
 			WidgetRef.GetPreview()->Deselect();
 		}
+
+		// Find all named slot host widgets that are hierarchical ancestors of this widget and call deselect on them as well
+		TArray<FWidgetReference> AncestorSlotHostWidgets;
+		FWidgetBlueprintEditorUtils::FindAllAncestorNamedSlotHostWidgetsForContent(AncestorSlotHostWidgets, WidgetRef.GetTemplate(), BPEd.ToSharedRef());
+
+		for (FWidgetReference SlotHostWidget : AncestorSlotHostWidgets)
+		{
+			SlotHostWidget.GetPreview()->Deselect();
+		}
 	}
 
 	// Notify all widgets that are now selected.
@@ -950,6 +960,15 @@ void SDesignerView::OnEditorSelectionChanged()
 		if ( WidgetRef.IsValid() && !SelectedWidgetsCache.Contains(WidgetRef) )
 		{
 			WidgetRef.GetPreview()->Select();
+
+			// Find all named slot host widgets that are hierarchical ancestors of this widget and call select on them as well
+			TArray<FWidgetReference> AncestorSlotHostWidgets;
+			FWidgetBlueprintEditorUtils::FindAllAncestorNamedSlotHostWidgetsForContent(AncestorSlotHostWidgets, WidgetRef.GetTemplate(), BPEd.ToSharedRef());
+
+			for (FWidgetReference SlotHostWidget : AncestorSlotHostWidgets)
+			{
+				SlotHostWidget.GetPreview()->Select();
+			}
 		}
 	}
 
diff --git a/Engine/Source/Editor/UMGEditor/Private/Designer/SRuler.cpp b/Engine/Source/Editor/UMGEditor/Private/Designer/SRuler.cpp
--- a/Engine/Source/Editor/UMGEditor/Private/Designer/SRuler.cpp
+++ b/Engine/Source/Editor/UMGEditor/Private/Designer/SRuler.cpp
@@ -72,7 +72,7 @@ struct FDrawTickArgs
 	/** Start layer for elements */
 	int32 StartLayer;
 	/** Draw effects to apply */
-	ESlateDrawEffect::Type DrawEffects;
+	ESlateDrawEffect DrawEffects;
 	/** Whether or not to only draw major ticks */
 	bool bOnlyDrawMajorTicks;
 };
@@ -295,7 +295,7 @@ int32 SRuler::DrawTicks( FSlateWindowElementList& OutDrawElements, const struct
 int32 SRuler::OnPaint(const FPaintArgs& Args, const FGeometry& AllottedGeometry, const FSlateRect& MyClippingRect, FSlateWindowElementList& OutDrawElements, int32 LayerId, const FWidgetStyle& InWidgetStyle, bool bParentEnabled) const
 {
 	const bool bEnabled = bParentEnabled;
-	const ESlateDrawEffect::Type DrawEffects = bEnabled ? ESlateDrawEffect::None : ESlateDrawEffect::DisabledEffect;
+	const ESlateDrawEffect DrawEffects = bEnabled ? ESlateDrawEffect::None : ESlateDrawEffect::DisabledEffect;
 
 	static const FSlateBrush* WhiteBrush = FCoreStyle::Get().GetBrush("GenericWhiteBox");
 
diff --git a/Engine/Source/Editor/UMGEditor/Private/Settings/WidgetDesignerSettings.cpp b/Engine/Source/Editor/UMGEditor/Private/Settings/WidgetDesignerSettings.cpp
--- a/Engine/Source/Editor/UMGEditor/Private/Settings/WidgetDesignerSettings.cpp
+++ b/Engine/Source/Editor/UMGEditor/Private/Settings/WidgetDesignerSettings.cpp
@@ -8,4 +8,5 @@ UWidgetDesignerSettings::UWidgetDesignerSettings(const FObjectInitializer& Objec
 	GridSnapEnabled = true;
 	GridSnapSize = 4;
 	bShowOutlines = true;
+	bExecutePreConstructEvent = true;
 }
