Null check added to fix crash adding component to landscape

diff --git a/Engine/Source/Editor/LandscapeEditor/Private/LandscapeEdModeComponentTools.cpp b/Engine/Source/Editor/LandscapeEditor/Private/LandscapeEdModeComponentTools.cpp
--- a/Engine/Source/Editor/LandscapeEditor/Private/LandscapeEdModeComponentTools.cpp
+++ b/Engine/Source/Editor/LandscapeEditor/Private/LandscapeEdModeComponentTools.cpp
@@ -973,8 +973,11 @@ public:
 					}
 				}
 				
-				FLandscapeEditDataInterface LandscapeEdit(LandscapeInfo);
-				NewComponent->FillLayer(BestLayerInfoObject, LandscapeEdit);
+				if (BestLayerInfoObject != nullptr)
+				{
+					FLandscapeEditDataInterface LandscapeEdit(LandscapeInfo);
+					NewComponent->FillLayer(BestLayerInfoObject, LandscapeEdit);
+				}
 			}
 
 			EdMode->LandscapeRenderAddCollision = nullptr;
