Fixed the error is size calculation for raw format (.r8) of weightmap

diff --git a/Engine/Source/Editor/LandscapeEditor/Private/LandscapeFileFormatRaw.cpp b/Engine/Source/Editor/LandscapeEditor/Private/LandscapeFileFormatRaw.cpp
--- a/Engine/Source/Editor/LandscapeEditor/Private/LandscapeFileFormatRaw.cpp
+++ b/Engine/Source/Editor/LandscapeEditor/Private/LandscapeFileFormatRaw.cpp
@@ -133,7 +133,7 @@ FLandscapeWeightmapInfo FLandscapeWeightmapFileFormat_Raw::Validate(const TCHAR*
 	}
 	else
 	{
-		Result.PossibleResolutions = CalculatePossibleRawResolutions(ImportFileSize / 2);
+		Result.PossibleResolutions = CalculatePossibleRawResolutions(ImportFileSize);
 
 		if (Result.PossibleResolutions.Num() == 0)
 		{
