Add Alt-Drag functionality to Landscape Splines for duplicating control points.

diff --git a/Engine/Source/Editor/LandscapeEditor/Private/LandscapeEdModeSplineTools.cpp b/Engine/Source/Editor/LandscapeEditor/Private/LandscapeEdModeSplineTools.cpp
--- a/Engine/Source/Editor/LandscapeEditor/Private/LandscapeEdModeSplineTools.cpp
+++ b/Engine/Source/Editor/LandscapeEditor/Private/LandscapeEdModeSplineTools.cpp
@@ -55,6 +55,16 @@ public:
 		, bAutoChangeConnectionsOnMove(true)
 		, bDeleteLooseEnds(false)
 		, bCopyMeshToNewControlPoint(false)
+		, bAllowDuplication(true)
+		, bDuplicatingControlPoint(false)
+		, bUpdatingAddSegment(false)
+		, DuplicateDelay(0)
+		, DuplicateDelayAccumulatedDrag(FVector::ZeroVector)
+		, DuplicateCachedRotation(FRotator::ZeroRotator)
+		, DuplicateCacheSplitSegmentParam(0.0f)
+		, DuplicateCacheSplitSegmentTangentLenStart(0.0f)
+		, DuplicateCacheSplitSegmentTangentLenEnd(0.0f)
+		, DuplicateCacheSplitSegmentTangentLen(0.0f)
 	{
 		// Register to update when an undo/redo operation has been called to update our list of actors
 		GEditor->RegisterForUndo(this);
@@ -351,6 +361,74 @@ public:
 		}
 	}
 
+	// called when alt-dragging a newly added end segment
+	bool UpdateAddSegment(ULandscapeSplineControlPoint* ControlPoint, FVector Location)
+	{
+		if (ControlPoint->ConnectedSegments.Num() != 1)
+		{
+			return false;
+		}
+
+		ULandscapeSplineSegment* Segment = ControlPoint->ConnectedSegments[0].Segment;
+		bool bAutoRotateStart = ControlPoint == Segment->Connections[0].ControlPoint ? false : bAutoRotateOnJoin;
+		bool bAutoRotateEnd = ControlPoint == Segment->Connections[1].ControlPoint ? false : bAutoRotateOnJoin;
+	
+		ULandscapeSplineControlPoint* Start = Segment->Connections[0].ControlPoint;
+		ULandscapeSplineControlPoint* End = Segment->Connections[1].ControlPoint;
+
+		ControlPoint->Location = Location;
+
+		FVector StartLocation; FRotator StartRotation;
+		Start->GetConnectionLocationAndRotation(Segment->Connections[0].SocketName, StartLocation, StartRotation);
+		FVector EndLocation; FRotator EndRotation;
+		End->GetConnectionLocationAndRotation(Segment->Connections[1].SocketName, EndLocation, EndRotation);
+
+		// Set up tangent lengths
+		Segment->Connections[0].TangentLen = (EndLocation - StartLocation).Size();
+		Segment->Connections[1].TangentLen = Segment->Connections[0].TangentLen;
+
+		Segment->AutoFlipTangents();
+
+		bool bUpdatedStart = false;
+		bool bUpdatedEnd = false;
+		if (bAutoRotateStart)
+		{
+			Start->AutoCalcRotation();
+			Start->UpdateSplinePoints();
+			bUpdatedStart = true;
+		}
+
+		if (bAutoRotateEnd)
+		{
+			End->AutoCalcRotation();
+			End->UpdateSplinePoints();
+			bUpdatedEnd = true;
+		}
+
+		// Control points' points are currently based on connected segments, so need to be updated.
+		if (!bUpdatedStart && (Start->Mesh || Start == ControlPoint))
+		{
+			Start->UpdateSplinePoints();
+			bUpdatedStart = true;
+		}
+		if (!bUpdatedEnd && (End->Mesh || End == ControlPoint))
+		{
+			End->UpdateSplinePoints();
+			bUpdatedEnd = true;
+		}
+
+		// If we've called UpdateSplinePoints on either control point it will already have called UpdateSplinePoints on the new segment
+		if (!(bUpdatedStart || bUpdatedEnd))
+		{
+			Segment->UpdateSplinePoints();
+		}
+
+		ULandscapeSplinesComponent* SplinesComponent = ControlPoint->GetOuterULandscapeSplinesComponent();
+		SplinesComponent->MarkRenderStateDirty();
+
+		return true;
+	}
+
 	void AddControlPoint(ULandscapeSplinesComponent* SplinesComponent, const FVector& LocalLocation)
 	{
 		FScopedTransaction Transaction(LOCTEXT("LandscapeSpline_AddControlPoint", "Add Landscape Spline Control Point"));
@@ -365,7 +443,17 @@ public:
 		if (SelectedSplineControlPoints.Num() > 0)
 		{
 			ULandscapeSplineControlPoint* FirstPoint = *SelectedSplineControlPoints.CreateConstIterator();
-			NewControlPoint->Rotation    = (NewControlPoint->Location - FirstPoint->Location).Rotation();
+
+			if (bDuplicatingControlPoint)
+			{
+				NewControlPoint->Rotation = FirstPoint->Rotation;
+			}
+			else
+			{
+				FVector NewSegmentDirection = (NewControlPoint->Location - FirstPoint->Location) * (FirstPoint->ConnectedSegments[0].End ? 1.0f : -1.0f);
+				NewControlPoint->Rotation = NewSegmentDirection.Rotation();
+			}
+
 			NewControlPoint->Width       = FirstPoint->Width;
 			NewControlPoint->SideFalloff = FirstPoint->SideFalloff;
 			NewControlPoint->EndFalloff  = FirstPoint->EndFalloff;
@@ -381,7 +469,14 @@ public:
 
 			for (ULandscapeSplineControlPoint* ControlPoint : SelectedSplineControlPoints)
 			{
-				AddSegment(ControlPoint, NewControlPoint, bAutoRotateOnJoin, true);
+				if (ControlPoint->ConnectedSegments[0].End)
+				{
+					AddSegment(ControlPoint, NewControlPoint, bAutoRotateOnJoin, !bDuplicatingControlPoint);
+				}
+				else
+				{
+					AddSegment(NewControlPoint, ControlPoint, !bDuplicatingControlPoint, bAutoRotateOnJoin);
+				}
 			}
 		}
 		else
@@ -542,12 +637,29 @@ public:
 		FVector Tangent;
 		Segment->FindNearest(LocalLocation, t, Location, Tangent);
 
+		if (bDuplicatingControlPoint)
+		{
+			DuplicateCacheSplitSegmentParam = t;
+			DuplicateCacheSplitSegmentTangentLenStart = Segment->Connections[0].TangentLen;
+			DuplicateCacheSplitSegmentTangentLenEnd = Segment->Connections[1].TangentLen;
+			DuplicateCacheSplitSegmentTangentLen = Tangent.Size();
+		}
+
 		ULandscapeSplineControlPoint* NewControlPoint = NewObject<ULandscapeSplineControlPoint>(SplinesComponent, NAME_None, RF_Transactional);
 		SplinesComponent->ControlPoints.Add(NewControlPoint);
 
-		NewControlPoint->Location = Location;
-		NewControlPoint->Rotation = Tangent.Rotation();
-		NewControlPoint->Rotation.Roll = FMath::Lerp(Segment->Connections[0].ControlPoint->Rotation.Roll, Segment->Connections[1].ControlPoint->Rotation.Roll, t);
+		if (bDuplicatingControlPoint)
+		{
+			NewControlPoint->Location = LocalLocation;
+			NewControlPoint->Rotation = DuplicateCachedRotation;
+		}
+		else
+		{
+			NewControlPoint->Location = Location;
+			NewControlPoint->Rotation = Tangent.Rotation();
+			NewControlPoint->Rotation.Roll = FMath::Lerp(Segment->Connections[0].ControlPoint->Rotation.Roll, Segment->Connections[1].ControlPoint->Rotation.Roll, t);
+		}
+
 		NewControlPoint->Width = FMath::Lerp(Segment->Connections[0].ControlPoint->Width, Segment->Connections[1].ControlPoint->Width, t);
 		NewControlPoint->SideFalloff = FMath::Lerp(Segment->Connections[0].ControlPoint->SideFalloff, Segment->Connections[1].ControlPoint->SideFalloff, t);
 		NewControlPoint->EndFalloff = FMath::Lerp(Segment->Connections[0].ControlPoint->EndFalloff, Segment->Connections[1].ControlPoint->EndFalloff, t);
@@ -584,6 +696,79 @@ public:
 		SplinesComponent->MarkRenderStateDirty();
 	}
 
+	bool UpdateSplitSegment(ULandscapeSplineControlPoint* ControlPoint, const FVector& LocalLocation)
+	{
+		ULandscapeSplineSegment* Segment0 = ControlPoint->ConnectedSegments[0].Segment;
+		ULandscapeSplineSegment* Segment1 = ControlPoint->ConnectedSegments[1].Segment;
+		ULandscapeSplineSegment *Segment, *NewSegment;
+		if (ControlPoint->ConnectedSegments[0].End == 0)
+		{
+			Segment = ControlPoint->ConnectedSegments[1].Segment;
+			NewSegment = ControlPoint->ConnectedSegments[0].Segment;
+		}
+		else
+		{
+			Segment = ControlPoint->ConnectedSegments[0].Segment;
+			NewSegment = ControlPoint->ConnectedSegments[1].Segment;
+		}
+
+		float t, t0, t1, tseg;
+		FVector Location, Location0, Location1;
+		FVector Tangent, Tangent0, Tangent1;
+		Segment->FindNearest(LocalLocation, t0, Location0, Tangent0);
+		NewSegment->FindNearest(LocalLocation, t1, Location1, Tangent1);
+
+		float Len0 = Tangent0.Size();
+		float Len1 = Tangent1.Size();
+
+		ULandscapeSplineSegment* UseSegment;
+		if (FVector::Distance(LocalLocation, Location0) < FVector::Distance(LocalLocation, Location1))
+		{
+			t = DuplicateCacheSplitSegmentParam * t0;
+			tseg = t0;
+			Location = Location0;
+			Tangent = Tangent0;
+			UseSegment = Segment;
+		}
+		else
+		{
+			t = DuplicateCacheSplitSegmentParam + (1 - DuplicateCacheSplitSegmentParam) * t1;
+			tseg = t1;
+			Location = Location1;
+			Tangent = Tangent1;
+			UseSegment = NewSegment;
+		}
+		DuplicateCacheSplitSegmentParam = t;
+
+		ControlPoint->Location = LocalLocation;
+
+		// Do not update rotation during alt-drag.
+		//ControlPoint->Rotation.Roll = FMath::Lerp(UseSegment->Connections[0].ControlPoint->Rotation.Roll, UseSegment->Connections[1].ControlPoint->Rotation.Roll, tseg);
+		ControlPoint->Width = FMath::Lerp(UseSegment->Connections[0].ControlPoint->Width, UseSegment->Connections[1].ControlPoint->Width, tseg);
+		ControlPoint->SideFalloff = FMath::Lerp(UseSegment->Connections[0].ControlPoint->SideFalloff, UseSegment->Connections[1].ControlPoint->SideFalloff, tseg);
+		ControlPoint->EndFalloff = FMath::Lerp(UseSegment->Connections[0].ControlPoint->EndFalloff, UseSegment->Connections[1].ControlPoint->EndFalloff, tseg);
+
+		Segment->Connections[0].TangentLen = DuplicateCacheSplitSegmentTangentLenStart * t;
+		Segment->Connections[1].TangentLen = -DuplicateCacheSplitSegmentTangentLen * t;
+
+		NewSegment->Connections[0].TangentLen = DuplicateCacheSplitSegmentTangentLen * (1 - t);
+		NewSegment->Connections[1].TangentLen = DuplicateCacheSplitSegmentTangentLenEnd * (1 - t);
+
+		if (bAutoChangeConnectionsOnMove)
+		{
+			ControlPoint->AutoSetConnections(true);
+		}
+
+		ControlPoint->UpdateSplinePoints();
+		Segment->UpdateSplinePoints();
+		NewSegment->UpdateSplinePoints();
+
+		ULandscapeSplinesComponent* SplinesComponent = ControlPoint->GetOuterULandscapeSplinesComponent();
+		SplinesComponent->MarkRenderStateDirty();
+
+		return true;
+	}
+
 	void FlipSegment(ULandscapeSplineSegment* Segment)
 	{
 		FScopedTransaction Transaction(LOCTEXT("LandscapeSpline_FlipSegment", "Flip Landscape Spline Segment"));
@@ -1234,7 +1419,15 @@ public:
 						checkSlow(SelectedSplineControlPoints.Num() > 0);
 						bMovingControlPoint = true;
 
-						GEditor->BeginTransaction(LOCTEXT("LandscapeSpline_ModifyControlPoint", "Modify Landscape Spline Control Point"));
+						if (SelectedSplineControlPoints.Num() == 1 && InViewportClient->IsAltPressed() && InViewportClient->GetWidgetMode() == FWidget::WM_Translate && InViewportClient->GetCurrentWidgetAxis() != EAxisList::None)
+						{
+							GEditor->BeginTransaction(LOCTEXT("LandscapeSpline_DuplicateControlPoint", "Duplicate Landscape Spline Control Point"));
+						}
+						else
+						{
+							GEditor->BeginTransaction(LOCTEXT("LandscapeSpline_ModifyControlPoint", "Modify Landscape Spline Control Point"));
+						}
+
 						for (ULandscapeSplineControlPoint* ControlPoint : SelectedSplineControlPoints)
 						{
 							ControlPoint->Modify();
@@ -1266,9 +1459,16 @@ public:
 
 					for (ULandscapeSplineControlPoint* ControlPoint : SelectedSplineControlPoints)
 					{
+						if (bDuplicatingControlPoint && bAutoRotateOnJoin)
+						{
+							ControlPoint->AutoCalcRotation();
+						}
+
 						ControlPoint->UpdateSplinePoints(true);
 					}
 
+					ResetAllowDuplication();
+
 					GEditor->EndTransaction();
 
 					return false; // We're not actually handling this case ourselves, just wrapping it in a transaction
@@ -1289,6 +1489,109 @@ public:
 		return false;
 	}
 
+	virtual void ResetAllowDuplication()
+	{
+		bAllowDuplication = true;
+		bDuplicatingControlPoint = false;
+		bUpdatingAddSegment = false;
+		DuplicateDelay = 0;
+		DuplicateDelayAccumulatedDrag = FVector::ZeroVector;
+		DuplicateCachedRotation = FRotator::ZeroRotator;
+		DuplicateCacheSplitSegmentParam = 0.0f;
+		DuplicateCacheSplitSegmentTangentLenStart = 0.0f;
+		DuplicateCacheSplitSegmentTangentLenEnd = 0.0f;
+		DuplicateCacheSplitSegmentTangentLen = 0.0f;
+	}
+
+	virtual bool DuplicateControlPoint(FVector& InDrag)
+	{
+		if (InDrag.IsZero())
+		{
+			return false;
+		}
+
+		ULandscapeSplineControlPoint* SelectedControlPoint = *SelectedSplineControlPoints.CreateConstIterator();
+		ULandscapeSplinesComponent* SplinesComponent = SelectedControlPoint->GetOuterULandscapeSplinesComponent();
+		FVector LocalDrag = SplinesComponent->GetComponentTransform().InverseTransformVector(InDrag);
+
+		ULandscapeSplineSegment* SegmentToSplit = nullptr;
+
+		if (SelectedControlPoint->ConnectedSegments.Num() > 0)
+		{
+			bool bHasPrevAngle = false;
+			float PrevAngle = 0.0f;
+
+			for (const FLandscapeSplineConnection& Connection : SelectedControlPoint->ConnectedSegments)
+			{
+				ULandscapeSplineControlPoint* AdjacentControlPoint = Connection.Segment->Connections[Connection.End == 1 ? 0 : 1].ControlPoint;
+				FVector SegmentDirection = AdjacentControlPoint->Location - SelectedControlPoint->Location;
+				if (SegmentDirection.IsZero())
+				{
+					continue;
+				}
+
+				float CurrentAngle = FMath::Acos(FVector::DotProduct(LocalDrag, SegmentDirection) / (LocalDrag.Size() * SegmentDirection.Size()));
+
+				// Create a new segment if there is no segment within 90 degrees of drag direction.
+				// Otherwise split segment that is closest to the drag direction.
+				if ((SelectedControlPoint->ConnectedSegments.Num() == 1 && CurrentAngle < HALF_PI) ||
+					(SelectedControlPoint->ConnectedSegments.Num() > 1 && (!bHasPrevAngle || CurrentAngle < PrevAngle)))
+				{
+					SegmentToSplit = Connection.Segment;
+				}
+
+				bHasPrevAngle = true;
+				PrevAngle = CurrentAngle;
+			}
+		}
+
+		FVector Location = SelectedControlPoint->Location + LocalDrag;
+		DuplicateCachedRotation = SelectedControlPoint->Rotation;
+
+		if (SegmentToSplit)
+		{
+			SplitSegment(SegmentToSplit, Location);
+
+			FWidget::EWidgetMode WidgetMode = EdMode->GetModeManager()->GetWidgetMode(); 
+			SelectControlPoint(SplinesComponent->ControlPoints.Last());
+			EdMode->GetModeManager()->SetWidgetMode(WidgetMode);
+		}
+		else
+		{
+			AddControlPoint(SplinesComponent, Location);
+			bUpdatingAddSegment = true;
+			GUnrealEd->RedrawLevelEditingViewports();
+		}
+
+		// Get newly-created control point
+		SelectedControlPoint = *SelectedSplineControlPoints.CreateConstIterator();
+
+		if (bAutoChangeConnectionsOnMove)
+		{
+			SelectedControlPoint->AutoSetConnections(true);
+		}
+
+		SelectedControlPoint->UpdateSplinePoints(false);
+
+		return true;
+	}
+
+	// called when alt-dragging duplicated control point 
+	virtual bool UpdateDuplicateControlPoint(FVector& InDrag)
+	{
+		ULandscapeSplineControlPoint* SelectedControlPoint = *SelectedSplineControlPoints.CreateConstIterator();
+		ULandscapeSplinesComponent* SplinesComponent = SelectedControlPoint->GetOuterULandscapeSplinesComponent();
+		FVector LocalDrag = SplinesComponent->GetComponentTransform().InverseTransformVector(InDrag);
+		FVector Location = SelectedControlPoint->Location + LocalDrag;
+
+		if (bUpdatingAddSegment)
+		{
+			return UpdateAddSegment(SelectedControlPoint, Location);
+		}
+
+		return UpdateSplitSegment(SelectedControlPoint, Location);
+	}
+
 	virtual bool InputDelta(FEditorViewportClient* InViewportClient, FViewport* InViewport, FVector& InDrag, FRotator& InRot, FVector& InScale) override
 	{
 		FVector Drag = InDrag;
@@ -1302,7 +1605,7 @@ public:
 			Connection.ControlPoint->GetConnectionLocationAndRotation(Connection.SocketName, StartLocation, StartRotation);
 
 			float OldTangentLen = Connection.TangentLen;
-			Connection.TangentLen += SplinesComponent->GetComponentTransform().InverseTransformVector(-Drag) | StartRotation.Vector();
+			Connection.TangentLen += SplinesComponent->GetComponentTransform().InverseTransformVector(Drag) | StartRotation.Vector();
 
 			// Disallow a tangent of exactly 0
 			if (Connection.TangentLen == 0)
@@ -1328,6 +1631,36 @@ public:
 			return true;
 		}
 
+		if (SelectedSplineControlPoints.Num() == 1 && InViewportClient->IsAltPressed() && InViewportClient->GetWidgetMode() == FWidget::WM_Translate && InViewportClient->GetCurrentWidgetAxis() != EAxisList::None)
+		{
+			static const int MaxDuplicationDelay = 3;
+
+			if (bAllowDuplication)
+			{
+				if (DuplicateDelay < MaxDuplicationDelay)
+				{
+					DuplicateDelay++;
+					DuplicateDelayAccumulatedDrag += Drag;
+
+					return true;
+				}
+				else
+				{
+					Drag += DuplicateDelayAccumulatedDrag;
+					DuplicateDelayAccumulatedDrag = FVector::ZeroVector;
+				}
+
+				bAllowDuplication = false;
+				bDuplicatingControlPoint = true;
+
+				return DuplicateControlPoint(Drag);
+			}
+			else
+			{
+				return UpdateDuplicateControlPoint(Drag);
+			}
+		}
+
 		if (SelectedSplineControlPoints.Num() > 0 && InViewportClient->GetCurrentWidgetAxis() != EAxisList::None)
 		{
 			for (ULandscapeSplineControlPoint* ControlPoint : SelectedSplineControlPoints)
@@ -1825,6 +2158,36 @@ protected:
 	uint32 bDeleteLooseEnds : 1;
 	uint32 bCopyMeshToNewControlPoint : 1;
 
+	/** Alt-drag: True when control point may be duplicated. */
+	uint32 bAllowDuplication : 1;
+
+	/** Alt-drag: True when in process of duplicating a control point. */
+	uint32 bDuplicatingControlPoint : 1;
+
+	/** Alt-drag: True when in process of adding end segment. */
+	uint32 bUpdatingAddSegment : 1;
+
+	/** Alt-drag: Delays duplicating control point to accumulate sufficient drag input offset. */
+	uint32 DuplicateDelay;
+
+	/** Alt-drag: Accumulates delayed drag offset. */
+	FVector DuplicateDelayAccumulatedDrag;
+
+	/** Alt-drag: Cached control point rotation when adding new control point at end of the spline. */
+	FRotator DuplicateCachedRotation;
+
+	/** Alt-drag: Cached segment parameter for split segment at new control point */
+	float DuplicateCacheSplitSegmentParam;
+
+	/** Alt-drag: Cached pre-split segment start tangent length. */
+	float DuplicateCacheSplitSegmentTangentLenStart;
+
+	/** Alt-drag: Cached pre-split segment end tangent length. */
+	float DuplicateCacheSplitSegmentTangentLenEnd;
+
+	/** Alt-drag: Cached tangent length for split segment at new control point. */
+	float DuplicateCacheSplitSegmentTangentLen;
+
 	friend FEdModeLandscape;
 };
 
