Added missing landscape actor check

diff --git a/Engine/Source/Editor/LandscapeEditor/Private/LandscapeEdModeComponentTools.cpp b/Engine/Source/Editor/LandscapeEditor/Private/LandscapeEdModeComponentTools.cpp
--- a/Engine/Source/Editor/LandscapeEditor/Private/LandscapeEdModeComponentTools.cpp
+++ b/Engine/Source/Editor/LandscapeEditor/Private/LandscapeEdModeComponentTools.cpp
@@ -912,7 +912,7 @@ public:
 				NewComponents[Idx]->RegisterComponent();
 			}
 
-			if (LandscapeInfo->LandscapeActor.Get()->HasProceduralContent && !GetMutableDefault<UEditorExperimentalSettings>()->bProceduralLandscape)
+			if (LandscapeInfo->LandscapeActor.IsValid() && LandscapeInfo->LandscapeActor.Get()->HasProceduralContent && !GetMutableDefault<UEditorExperimentalSettings>()->bProceduralLandscape)
 			{
 				FMessageLog("MapCheck").Warning()->AddToken(FTextToken::Create(LOCTEXT("LandscapeProcedural_ChangingDataWithoutSettings", "This map contains landscape procedural content, modifying the landscape data will result in data loss when the map is reopened with Landscape Procedural settings on. Please enable Landscape Procedural settings before modifying the data.")));
 				FMessageLog("MapCheck").Open(EMessageSeverity::Warning);
