Added basic Procedural Layer management (Create New Layer, Delete Layer, Lock/Unlock Layers, Toggle/Isolate Layer Visibility)

diff --git a/Engine/Source/Editor/LandscapeEditor/Private/LandscapeEdMode.cpp b/Engine/Source/Editor/LandscapeEditor/Private/LandscapeEdMode.cpp
--- a/Engine/Source/Editor/LandscapeEditor/Private/LandscapeEdMode.cpp
+++ b/Engine/Source/Editor/LandscapeEditor/Private/LandscapeEdMode.cpp
@@ -561,9 +561,8 @@ void FEdModeLandscape::Enter()
 
 		if (GetMutableDefault<UEditorExperimentalSettings>()->bProceduralLandscape)
 		{
-			ALandscape* Landscape = CurrentToolTarget.LandscapeInfo->LandscapeActor.Get();
-
-			if (Landscape != nullptr)
+			ALandscape* Landscape = GetLandscape();
+			if (Landscape)
 			{
 				Landscape->RequestProceduralContentUpdate(EProceduralContentUpdateFlag::All_Render);
 			}
@@ -1898,8 +1897,18 @@ bool FEdModeLandscape::InputKey(FEditorViewportClient* ViewportClient, FViewport
 					if (LandscapeMouseTrace(ViewportClient, HitLocation))
 					{
 						const FName CurrentToolName = CurrentTool->GetToolName();
+						ALandscape* Landscape = GetLandscape();
+						bool bProceduralLandscape = GetMutableDefault<UEditorExperimentalSettings>()->bProceduralLandscape;
 
-						if (CurrentToolTarget.TargetType == ELandscapeToolTargetType::Weightmap && CurrentToolTarget.LayerInfo == NULL && CurrentToolName != TEXT("BPCustom"))
+						if (bProceduralLandscape && GetCurrentProceduralLayer() && !GetCurrentProceduralLayer()->bVisible)
+						{
+							FMessageDialog::Open(EAppMsgType::Ok, NSLOCTEXT("UnrealEd", "LandscapeProceduralLayerHidden", "Painting in a hidden layer is not allowed."));
+						}
+						else if (bProceduralLandscape && GetCurrentProceduralLayer() && GetCurrentProceduralLayer()->bLocked)
+						{
+							FMessageDialog::Open(EAppMsgType::Ok, NSLOCTEXT("UnrealEd", "LandscapeProceduralLayerLocked", "This layer is locked. You must unlock it before you can work on this layer."));
+						}
+						else if (CurrentToolTarget.TargetType == ELandscapeToolTargetType::Weightmap && CurrentToolTarget.LayerInfo == NULL && CurrentToolName != TEXT("BPCustom"))
 						{
 							FMessageDialog::Open(EAppMsgType::Ok,
 								NSLOCTEXT("UnrealEd", "LandscapeNeedToCreateLayerInfo", "This layer has no layer info assigned yet. You must create or assign a layer info before you can paint this layer."));
@@ -3341,6 +3350,7 @@ void FEdModeLandscape::ReimportData(const FLandscapeTargetListInfo& TargetInfo)
 	const FString& SourceFilePath = TargetInfo.ReimportFilePath();
 	if (SourceFilePath.Len())
 	{
+		FScopedSetLandscapeCurrentEditingProceduralLayer Scope(GetLandscape(), GetCurrentProceduralLayerGuid(), [&] { RequestProceduralContentUpdate(); });
 		ImportData(TargetInfo, SourceFilePath);
 	}
 	else
@@ -3445,42 +3455,34 @@ void FEdModeLandscape::ImportData(const FLandscapeTargetListInfo& TargetInfo, co
 				return;
 			}
 
-			if (GetMutableDefault<UEditorExperimentalSettings>()->bProceduralLandscape)
-			{
-				ChangeHeightmapsToCurrentProceduralLayerHeightmaps();
-			}
-
-			TArray<uint16> Data;
-			if (ImportResolution != LandscapeResolution)
 			{
-				// Cloned from FLandscapeEditorDetailCustomization_NewLandscape.OnCreateButtonClicked
-				// so that reimports behave the same as the initial import :)
+				ALandscape* Landscape = GetLandscape();
+				FScopedSetLandscapeCurrentEditingProceduralLayer Scope(Landscape, GetCurrentProceduralLayerGuid(), [&] { check(Landscape); Landscape->RequestProceduralContentUpdate(EProceduralContentUpdateFlag::Heightmap_All); });
 
-				const int32 OffsetX = (int32)(LandscapeResolution.Width - ImportResolution.Width) / 2;
-				const int32 OffsetY = (int32)(LandscapeResolution.Height - ImportResolution.Height) / 2;
-
-				Data.SetNumUninitialized(LandscapeResolution.Width * LandscapeResolution.Height * sizeof(uint16));
+				TArray<uint16> Data;
+				if (ImportResolution != LandscapeResolution)
+				{
+					// Cloned from FLandscapeEditorDetailCustomization_NewLandscape.OnCreateButtonClicked
+					// so that reimports behave the same as the initial import :)
 
-				LandscapeEditorUtils::ExpandData<uint16>(Data.GetData(), ImportData.Data.GetData(),
-					0, 0, ImportResolution.Width - 1, ImportResolution.Height - 1,
-					-OffsetX, -OffsetY, LandscapeResolution.Width - OffsetX - 1, LandscapeResolution.Height - OffsetY - 1);
-			}
-			else
-			{
-				Data = MoveTemp(ImportData.Data);
-			}
+					const int32 OffsetX = (int32)(LandscapeResolution.Width - ImportResolution.Width) / 2;
+					const int32 OffsetY = (int32)(LandscapeResolution.Height - ImportResolution.Height) / 2;
 
-			FScopedTransaction Transaction(LOCTEXT("Undo_ImportHeightmap", "Importing Landscape Heightmap"));
+					Data.SetNumUninitialized(LandscapeResolution.Width * LandscapeResolution.Height * sizeof(uint16));
 
-			FHeightmapAccessor<false> HeightmapAccessor(LandscapeInfo);
-			HeightmapAccessor.SetData(MinX, MinY, MaxX, MaxY, Data.GetData());
+					LandscapeEditorUtils::ExpandData<uint16>(Data.GetData(), ImportData.Data.GetData(),
+						0, 0, ImportResolution.Width - 1, ImportResolution.Height - 1,
+						-OffsetX, -OffsetY, LandscapeResolution.Width - OffsetX - 1, LandscapeResolution.Height - OffsetY - 1);
+				}
+				else
+				{
+					Data = MoveTemp(ImportData.Data);
+				}
 
-			if (GetMutableDefault<UEditorExperimentalSettings>()->bProceduralLandscape)
-			{
-				ChangeHeightmapsToCurrentProceduralLayerHeightmaps(true);
+				FScopedTransaction Transaction(LOCTEXT("Undo_ImportHeightmap", "Importing Landscape Heightmap"));
 
-				check(CurrentToolTarget.LandscapeInfo->LandscapeActor.IsValid());
-				CurrentToolTarget.LandscapeInfo->LandscapeActor.Get()->RequestProceduralContentUpdate(EProceduralContentUpdateFlag::Heightmap_All);
+				FHeightmapAccessor<false> HeightmapAccessor(LandscapeInfo);
+				HeightmapAccessor.SetData(MinX, MinY, MaxX, MaxY, Data.GetData());
 			}
 		}
 		else
@@ -3569,42 +3571,34 @@ void FEdModeLandscape::ImportData(const FLandscapeTargetListInfo& TargetInfo, co
 				return;
 			}
 
-			if (GetMutableDefault<UEditorExperimentalSettings>()->bProceduralLandscape)
 			{
-				ChangeWeightmapsToCurrentProceduralLayerWeightmaps();
-			}
+				ALandscape* Landscape = GetLandscape();
+				FScopedSetLandscapeCurrentEditingProceduralLayer Scope(Landscape, GetCurrentProceduralLayerGuid(), [&] { check(Landscape); Landscape->RequestProceduralContentUpdate(EProceduralContentUpdateFlag::Weightmap_All); });
 
-			TArray<uint8> Data;
-			if (ImportResolution != LandscapeResolution)
-			{
-				// Cloned from FLandscapeEditorDetailCustomization_NewLandscape.OnCreateButtonClicked
-				// so that reimports behave the same as the initial import :)
-
-				const int32 OffsetX = (int32)(LandscapeResolution.Width - ImportResolution.Width) / 2;
-				const int32 OffsetY = (int32)(LandscapeResolution.Height - ImportResolution.Height) / 2;
-
-				Data.SetNumUninitialized(LandscapeResolution.Width * LandscapeResolution.Height * sizeof(uint8));
+				TArray<uint8> Data;
+				if (ImportResolution != LandscapeResolution)
+				{
+					// Cloned from FLandscapeEditorDetailCustomization_NewLandscape.OnCreateButtonClicked
+					// so that reimports behave the same as the initial import :)
 
-				LandscapeEditorUtils::ExpandData<uint8>(Data.GetData(), ImportData.Data.GetData(),
-					0, 0, ImportResolution.Width - 1, ImportResolution.Height - 1,
-					-OffsetX, -OffsetY, LandscapeResolution.Width - OffsetX - 1, LandscapeResolution.Height - OffsetY - 1);
-			}
-			else
-			{
-				Data = MoveTemp(ImportData.Data);
-			}
+					const int32 OffsetX = (int32)(LandscapeResolution.Width - ImportResolution.Width) / 2;
+					const int32 OffsetY = (int32)(LandscapeResolution.Height - ImportResolution.Height) / 2;
 
-			FScopedTransaction Transaction(LOCTEXT("Undo_ImportWeightmap", "Importing Landscape Layer"));
+					Data.SetNumUninitialized(LandscapeResolution.Width * LandscapeResolution.Height * sizeof(uint8));
 
-			FAlphamapAccessor<false, false> AlphamapAccessor(LandscapeInfo, TargetInfo.LayerInfoObj.Get());
-			AlphamapAccessor.SetData(MinX, MinY, MaxX, MaxY, Data.GetData(), ELandscapeLayerPaintingRestriction::None);
+					LandscapeEditorUtils::ExpandData<uint8>(Data.GetData(), ImportData.Data.GetData(),
+						0, 0, ImportResolution.Width - 1, ImportResolution.Height - 1,
+						-OffsetX, -OffsetY, LandscapeResolution.Width - OffsetX - 1, LandscapeResolution.Height - OffsetY - 1);
+				}
+				else
+				{
+					Data = MoveTemp(ImportData.Data);
+				}
 
-			if (GetMutableDefault<UEditorExperimentalSettings>()->bProceduralLandscape)
-			{
-				ChangeWeightmapsToCurrentProceduralLayerWeightmaps(true);
+				FScopedTransaction Transaction(LOCTEXT("Undo_ImportWeightmap", "Importing Landscape Layer"));
 
-				check(CurrentToolTarget.LandscapeInfo->LandscapeActor.IsValid());
-				CurrentToolTarget.LandscapeInfo->LandscapeActor.Get()->RequestProceduralContentUpdate(EProceduralContentUpdateFlag::Weightmap_All);
+				FAlphamapAccessor<false, false> AlphamapAccessor(LandscapeInfo, TargetInfo.LayerInfoObj.Get());
+				AlphamapAccessor.SetData(MinX, MinY, MaxX, MaxY, Data.GetData(), ELandscapeLayerPaintingRestriction::None);
 			}
 		}
 	}
@@ -4130,19 +4124,8 @@ ELandscapeEditingState FEdModeLandscape::GetEditingState() const
 
 int32 FEdModeLandscape::GetProceduralLayerCount() const
 {
-	if (CurrentToolTarget.LandscapeInfo == nullptr)
-	{
-		return 0;
-	}
-
-	ALandscape* Landscape = CurrentToolTarget.LandscapeInfo->LandscapeActor.Get();
-
-	if (Landscape == nullptr)
-	{
-		return 0;
-	}
-
-	return Landscape->ProceduralLayers.Num();
+	ALandscape* Landscape = GetLandscape();
+	return Landscape ? Landscape->ProceduralLayers.Num() : 0;
 }
 
 void FEdModeLandscape::SetCurrentProceduralLayer(int32 InLayerIndex)
@@ -4157,137 +4140,108 @@ int32 FEdModeLandscape::GetCurrentProceduralLayerIndex() const
 	return CurrentToolTarget.CurrentProceduralLayerIndex;
 }
 
-FName FEdModeLandscape::GetProceduralLayerName(int32 InLayerIndex) const
+ALandscape* FEdModeLandscape::GetLandscape() const
 {
-	if (CurrentToolTarget.LandscapeInfo == nullptr)
-	{
-		return NAME_None;
-	}
-
-	ALandscape* Landscape = CurrentToolTarget.LandscapeInfo->LandscapeActor.Get();
-
-	if (Landscape == nullptr || !Landscape->ProceduralLayers.IsValidIndex(InLayerIndex))
-	{
-		return NAME_None;
-	}
+	return CurrentToolTarget.LandscapeInfo.IsValid() ? CurrentToolTarget.LandscapeInfo->LandscapeActor.Get() : nullptr;
+}
 
-	return Landscape->ProceduralLayers[InLayerIndex].Name;
+FProceduralLayer* FEdModeLandscape::GetProceduralLayer(int32 InLayerIndex) const
+{
+	ALandscape* Landscape = GetLandscape();
+	return Landscape ? Landscape->GetProceduralLayer(InLayerIndex) : nullptr;
 }
 
-FName FEdModeLandscape::GetCurrentProceduralLayerName() const
+FName FEdModeLandscape::GetProceduralLayerName(int32 InLayerIndex) const
 {
-	return GetProceduralLayerName(CurrentToolTarget.CurrentProceduralLayerIndex);
+	FProceduralLayer* Layer = GetProceduralLayer(InLayerIndex);
+	return Layer ? Layer->Name : NAME_None;
 }
 
-void FEdModeLandscape::SetProceduralLayerName(int32 InLayerIndex, const FName& InName)
+bool FEdModeLandscape::CanRenameProceduralLayerTo(int32 InLayerIndex, const FName& InNewName)
 {
-	if (CurrentToolTarget.LandscapeInfo == nullptr)
+	ALandscape* Landscape = GetLandscape();
+	if (Landscape)
 	{
-		return;
+		int32 LayerCount = GetProceduralLayerCount();
+		for (int32 LayerIdx = 0; LayerIdx < LayerCount; ++LayerIdx)
+		{
+			if (LayerIdx != InLayerIndex && GetProceduralLayerName(LayerIdx) == InNewName)
+			{
+				return false;
+			}
+		}
 	}
+	return true;
+}
 
-	ALandscape* Landscape = CurrentToolTarget.LandscapeInfo->LandscapeActor.Get();
-
-	if (Landscape == nullptr || !Landscape->ProceduralLayers.IsValidIndex(InLayerIndex))
+void FEdModeLandscape::SetProceduralLayerName(int32 InLayerIndex, const FName& InName)
+{
+	ALandscape* Landscape = GetLandscape();
+	if (Landscape)
 	{
-		return;
+		Landscape->SetProceduralLayerName(InLayerIndex, InName);
 	}
+}
 
-	Landscape->ProceduralLayers[InLayerIndex].Name = InName;
+FName FEdModeLandscape::GetCurrentProceduralLayerName() const
+{
+	return GetProceduralLayerName(CurrentToolTarget.CurrentProceduralLayerIndex);
 }
 
 float FEdModeLandscape::GetProceduralLayerAlpha(int32 InLayerIndex) const
 {
-	if (CurrentToolTarget.LandscapeInfo == nullptr)
+	FProceduralLayer* Layer = GetProceduralLayer(InLayerIndex);
+	if (Layer)
 	{
-		return 1.0f;
+		return (CurrentToolTarget.TargetType == ELandscapeToolTargetType::Heightmap) ? Layer->HeightmapAlpha : Layer->WeightmapAlpha;
 	}
-
-	ALandscape* Landscape = CurrentToolTarget.LandscapeInfo->LandscapeActor.Get();
-
-	if (Landscape == nullptr || !Landscape->ProceduralLayers.IsValidIndex(InLayerIndex))
-	{
-		return 1.0f;
-	}
-
-	return CurrentToolTarget.TargetType == ELandscapeToolTargetType::Heightmap ? Landscape->ProceduralLayers[InLayerIndex].HeightmapAlpha : Landscape->ProceduralLayers[InLayerIndex].WeightmapAlpha;
+	return 1.0f;
 }
 
-void FEdModeLandscape::SetProceduralLayerAlpha(float InWeight, int32 InLayerIndex)
+void FEdModeLandscape::SetProceduralLayerAlpha(int32 InLayerIndex, float InAlpha)
 {
-	if (CurrentToolTarget.LandscapeInfo == nullptr)
-	{
-		return;
-	}
-
-	ALandscape* Landscape = CurrentToolTarget.LandscapeInfo->LandscapeActor.Get();
-
-	if (Landscape == nullptr || !Landscape->ProceduralLayers.IsValidIndex(InLayerIndex))
-	{
-		return;
-	}
-
-	if (CurrentToolTarget.TargetType == ELandscapeToolTargetType::Heightmap)
-	{
-		Landscape->ProceduralLayers[InLayerIndex].HeightmapAlpha = InWeight;
-	}
-	else
+	ALandscape* Landscape = GetLandscape();
+	if (Landscape)
 	{
-		Landscape->ProceduralLayers[InLayerIndex].WeightmapAlpha = InWeight;
+		Landscape->SetProceduralLayerAlpha(InLayerIndex, InAlpha, CurrentToolTarget.TargetType == ELandscapeToolTargetType::Heightmap);
 	}
-
-
-	Landscape->RequestProceduralContentUpdate(EProceduralContentUpdateFlag::All, true);
 }
 
-void FEdModeLandscape::SetProceduralLayerVisibility(bool InVisible, int32 InLayerIndex)
+bool FEdModeLandscape::IsProceduralLayerVisible(int32 InLayerIndex) const
 {
-	if (CurrentToolTarget.LandscapeInfo == nullptr)
-	{
-		return;
-	}
-
-	ALandscape* Landscape = CurrentToolTarget.LandscapeInfo->LandscapeActor.Get();
-
-	if (Landscape == nullptr || !Landscape->ProceduralLayers.IsValidIndex(InLayerIndex))
-	{
-		return;
-	}
-
-	Landscape->ProceduralLayers[InLayerIndex].Visible = InVisible;
-
-	Landscape->RequestProceduralContentUpdate(EProceduralContentUpdateFlag::All, true);
+	FProceduralLayer* Layer = GetProceduralLayer(InLayerIndex);
+	return Layer ? Layer->bVisible : false;
 }
 
-bool FEdModeLandscape::IsProceduralLayerVisible(int32 InLayerIndex) const
+void FEdModeLandscape::SetProceduralLayerVisibility(bool bInVisible, int32 InLayerIndex)
 {
-	if (CurrentToolTarget.LandscapeInfo == nullptr)
+	ALandscape* Landscape = GetLandscape();
+	if (Landscape)
 	{
-		return true;
+		Landscape->SetProceduralLayerVisibility(InLayerIndex, bInVisible);
 	}
+}
 
-	ALandscape* Landscape = CurrentToolTarget.LandscapeInfo->LandscapeActor.Get();
+bool FEdModeLandscape::IsProceduralLayerLocked(int32 InLayerIndex) const
+{
+	FProceduralLayer* Layer = GetProceduralLayer(InLayerIndex);
+	return Layer ? Layer->bLocked : false;
+}
 
-	if (Landscape == nullptr || !Landscape->ProceduralLayers.IsValidIndex(InLayerIndex))
+void FEdModeLandscape::SetProceduralLayerLocked(int32 InLayerIndex, bool bInLocked)
+{
+	FProceduralLayer* Layer = GetProceduralLayer(InLayerIndex);
+	if (Layer)
 	{
-		return true;
+		Layer->bLocked = bInLocked;
 	}
-
-	return Landscape->ProceduralLayers[InLayerIndex].Visible;
 }
 
 void FEdModeLandscape::AddBrushToCurrentProceduralLayer(int32 InTargetType, ALandscapeBlueprintCustomBrush* InBrush)
 {
-	if (CurrentToolTarget.LandscapeInfo == nullptr || !CurrentToolTarget.LandscapeInfo->LandscapeActor.IsValid())
-	{
-		return;
-	}
-
-	ALandscape* Landscape = CurrentToolTarget.LandscapeInfo->LandscapeActor.Get();
-
+	ALandscape* Landscape = GetLandscape();
 	FProceduralLayer* Layer = GetCurrentProceduralLayer();
-
-	if (Layer == nullptr)
+	if (!Landscape || !Layer)
 	{
 		return;
 	}
@@ -4310,24 +4264,20 @@ void FEdModeLandscape::AddBrushToCurrentProceduralLayer(int32 InTargetType, ALan
 
 void FEdModeLandscape::RequestProceduralContentUpdate(bool InUpdateAllMaterials)
 {
-	if (CurrentToolTarget.LandscapeInfo == nullptr || !CurrentToolTarget.LandscapeInfo->LandscapeActor.IsValid())
+	ALandscape* Landscape = GetLandscape();
+	if (Landscape)
 	{
-		return;
+		Landscape->RequestProceduralContentUpdate(EProceduralContentUpdateFlag::All_Render, InUpdateAllMaterials);
 	}
-
-	ALandscape* Landscape = CurrentToolTarget.LandscapeInfo->LandscapeActor.Get();
-	Landscape->RequestProceduralContentUpdate(EProceduralContentUpdateFlag::All_Render, InUpdateAllMaterials);
 }
 
 void FEdModeLandscape::RemoveBrushFromCurrentProceduralLayer(int32 InTargetType, ALandscapeBlueprintCustomBrush* InBrush)
 {
-	if (CurrentToolTarget.LandscapeInfo == nullptr || !CurrentToolTarget.LandscapeInfo->LandscapeActor.IsValid())
+	ALandscape* Landscape = GetLandscape();
+	if (!Landscape)
 	{
 		return;
 	}
-
-	ALandscape* Landscape = CurrentToolTarget.LandscapeInfo->LandscapeActor.Get();
-
 	for (FProceduralLayer& Layer : Landscape->ProceduralLayers)
 	{
 		int32 IndexToRemove = INDEX_NONE;
@@ -4548,146 +4498,15 @@ void FEdModeLandscape::SetCurrentProceduralLayerSubstractiveBlendStatus(bool InS
 
 FProceduralLayer* FEdModeLandscape::GetCurrentProceduralLayer() const
 {
-	if (!CurrentToolTarget.LandscapeInfo.IsValid())
-	{
-		return nullptr;
-	}
-
-	ALandscape* Landscape = CurrentToolTarget.LandscapeInfo->LandscapeActor.Get();
-
-	if (Landscape == nullptr)
-	{
-		return nullptr;
-	}
-
-	FName CurrentLayerName = GetCurrentProceduralLayerName();
-
-	if (CurrentLayerName == NAME_None)
-	{
-		return nullptr;
-	}
-
-	for (FProceduralLayer& Layer : Landscape->ProceduralLayers)
-	{
-		if (Layer.Name == CurrentLayerName)
-		{
-			return &Layer;
-		}
-	}
-
-	return nullptr;
+	return GetProceduralLayer(GetCurrentProceduralLayerIndex());
 }
 
-void FEdModeLandscape::ChangeHeightmapsToCurrentProceduralLayerHeightmaps(bool InResetCurrentEditingHeightmap)
+FGuid FEdModeLandscape::GetCurrentProceduralLayerGuid() const
 {
-	if (!CurrentToolTarget.LandscapeInfo.IsValid() || !CurrentToolTarget.LandscapeInfo->LandscapeActor.IsValid())
-	{
-		return;
-	}
-
-	TArray<ALandscapeProxy*> AllLandscapes;
-	AllLandscapes.Add(CurrentToolTarget.LandscapeInfo->LandscapeActor.Get());
-
-	for (const auto& It : CurrentToolTarget.LandscapeInfo->Proxies)
-	{
-		AllLandscapes.Add(It);
-	}
-
-	FName CurrentLayerName = GetCurrentProceduralLayerName();
-
-	if (CurrentLayerName == NAME_None)
-	{
-		return;
-	}
-
-	for (ALandscapeProxy* LandscapeProxy : AllLandscapes)
-	{
-		FProceduralLayerData* CurrentLayerData = LandscapeProxy->ProceduralLayersData.Find(CurrentLayerName);
-
-		if (CurrentLayerData == nullptr)
-		{
-			continue;
-		}
-
-		for (ULandscapeComponent* Component : LandscapeProxy->LandscapeComponents)
-		{
-			if (InResetCurrentEditingHeightmap)
-			{
-				Component->SetCurrentEditingHeightmap(nullptr);
-			}
-			else
-			{
-				UTexture2D** LayerHeightmap = CurrentLayerData->Heightmaps.Find(Component->GetHeightmap());
-
-				if (LayerHeightmap != nullptr)
-				{
-					Component->SetCurrentEditingHeightmap(*LayerHeightmap);
-				}
-			}
-
-			Component->MarkRenderStateDirty();
-		}
-	}
-}
-void FEdModeLandscape::ChangeWeightmapsToCurrentProceduralLayerWeightmaps(bool InResetCurrentEditingHeightmap)
-{
-	if (!CurrentToolTarget.LandscapeInfo.IsValid() || !CurrentToolTarget.LandscapeInfo->LandscapeActor.IsValid())
-	{
-		return;
-	}
-
-	TArray<ALandscapeProxy*> AllLandscapes;
-	AllLandscapes.Add(CurrentToolTarget.LandscapeInfo->LandscapeActor.Get());
-
-	for (const auto& It : CurrentToolTarget.LandscapeInfo->Proxies)
-	{
-		AllLandscapes.Add(It);
-	}
-
-	FName CurrentLayerName = GetCurrentProceduralLayerName();
-
-	if (CurrentLayerName == NAME_None)
-	{
-		return;
-	}
-
-	for (ALandscapeProxy* LandscapeProxy : AllLandscapes)
-	{
-		FProceduralLayerData* CurrentLayerData = LandscapeProxy->ProceduralLayersData.Find(CurrentLayerName);
-
-		if (CurrentLayerData == nullptr)
-		{
-			continue;
-		}
-
-		for (ULandscapeComponent* Component : LandscapeProxy->LandscapeComponents)
-		{
-			if (InResetCurrentEditingHeightmap)
-			{
-				Component->SetCurrentEditingWeightmapLayerAllocations(nullptr);
-				Component->SetCurrentEditingWeightmaps(nullptr);
-				Component->SetCurrentProceduralLayerName(NAME_None);
-				Component->SetCurrentEditingWeightmapTexturesUsage(nullptr);			
-			}
-			else
-			{
-				FWeightmapLayerData* WeightmapData = CurrentLayerData->WeightmapData.Find(Component);
-
-				if (WeightmapData != nullptr)
-				{
-					Component->SetCurrentEditingWeightmaps(&WeightmapData->Weightmaps);
-					Component->SetCurrentEditingWeightmapLayerAllocations(&WeightmapData->WeightmapLayerAllocations);
-					Component->SetCurrentProceduralLayerName(CurrentLayerName);
-					Component->SetCurrentEditingWeightmapTexturesUsage(&WeightmapData->WeightmapTextureUsages);
-				}
-			}
-
-			Component->MarkRenderStateDirty();
-		}
-	}
+	FProceduralLayer* CurrentProceduralLayer = GetCurrentProceduralLayer();
+	return CurrentProceduralLayer ? CurrentProceduralLayer->Guid : FGuid();
 }
 
-
 bool FEdModeLandscape::NeedToFillEmptyLayersForProcedural() const
 {
 	if (!CurrentToolTarget.LandscapeInfo.IsValid() || !CurrentToolTarget.LandscapeInfo->LandscapeActor.IsValid())
diff --git a/Engine/Source/Editor/LandscapeEditor/Private/LandscapeEdMode.h b/Engine/Source/Editor/LandscapeEditor/Private/LandscapeEdMode.h
--- a/Engine/Source/Editor/LandscapeEditor/Private/LandscapeEdMode.h
+++ b/Engine/Source/Editor/LandscapeEditor/Private/LandscapeEdMode.h
@@ -494,16 +494,22 @@ public:
 	void SetCurrentProceduralLayer(int32 InLayerIndex);
 	int32 GetCurrentProceduralLayerIndex() const;
 	FName GetCurrentProceduralLayerName() const;
+	ALandscape* GetLandscape() const;
+	struct FProceduralLayer* GetProceduralLayer(int32 InLayerIndex) const;
 	FName GetProceduralLayerName(int32 InLayerIndex) const;
 	void SetProceduralLayerName(int32 InLayerIndex, const FName& InName);
+	bool CanRenameProceduralLayerTo(int32 InLayerIndex, const FName& InNewName);
 	float GetProceduralLayerAlpha(int32 InLayerIndex) const;
-	void SetProceduralLayerAlpha(float InWeight, int32 InLayerIndex);
+	void SetProceduralLayerAlpha(int32 InLayerIndex, float InAlpha);
 	void SetProceduralLayerVisibility(bool InVisible, int32 InLayerIndex);
 	bool IsProceduralLayerVisible(int32 InLayerIndex) const;
+	bool IsProceduralLayerLocked(int32 InLayerIndex) const;
+	void SetProceduralLayerLocked(int32 InLayerIndex, bool bInLocked);
 	struct FProceduralLayer* GetCurrentProceduralLayer() const;
+	FGuid GetCurrentProceduralLayerGuid() const;
 	bool IsCurrentProceduralLayerBlendSubstractive(const TWeakObjectPtr<ULandscapeLayerInfoObject>& InLayerInfoObj) const;
 	void SetCurrentProceduralLayerSubstractiveBlendStatus(bool InStatus, const TWeakObjectPtr<ULandscapeLayerInfoObject>& InLayerInfoObj);
-	
+
 	void AddBrushToCurrentProceduralLayer(int32 InTargetType, class ALandscapeBlueprintCustomBrush* InBrush);
 	void RemoveBrushFromCurrentProceduralLayer(int32 InTargetType, class ALandscapeBlueprintCustomBrush* InBrush);
 	bool AreAllBrushesCommitedToCurrentProceduralLayer(int32 InTargetType);
@@ -512,9 +518,6 @@ public:
 	class ALandscapeBlueprintCustomBrush* GetBrushForCurrentProceduralLayer(int32 InTargetType, int8 BrushIndex) const;
 	TArray<class ALandscapeBlueprintCustomBrush*> GetBrushesForCurrentProceduralLayer(int32 InTargetType);
 	
-	void ChangeHeightmapsToCurrentProceduralLayerHeightmaps(bool InResetCurrentEditingHeightmap = false);
-	void ChangeWeightmapsToCurrentProceduralLayerWeightmaps(bool InResetCurrentEditingHeightmap = false);
-
 	bool NeedToFillEmptyLayersForProcedural() const;
 	void RequestProceduralContentUpdate(bool InUpdateAllMaterials = false);
 
diff --git a/Engine/Source/Editor/LandscapeEditor/Private/LandscapeEdModeComponentTools.cpp b/Engine/Source/Editor/LandscapeEditor/Private/LandscapeEdModeComponentTools.cpp
--- a/Engine/Source/Editor/LandscapeEditor/Private/LandscapeEdModeComponentTools.cpp
+++ b/Engine/Source/Editor/LandscapeEditor/Private/LandscapeEdModeComponentTools.cpp
@@ -344,11 +344,11 @@ public:
 
 		if (GetMutableDefault<UEditorExperimentalSettings>()->bProceduralLandscape)
 		{
-			ALandscape* Landscape = this->EdMode->CurrentToolTarget.LandscapeInfo->LandscapeActor.Get();
-			if (Landscape != nullptr)
+			ALandscape* Landscape = this->EdMode->GetLandscape();
+			if (Landscape)
 			{
 				Landscape->RequestProceduralContentUpdate(EProceduralContentUpdateFlag::Weightmap_Render);
-				this->EdMode->ChangeWeightmapsToCurrentProceduralLayerWeightmaps();
+				Landscape->SetCurrentEditingProceduralLayer(this->EdMode->GetCurrentProceduralLayerGuid());
 			}
 		}
 
@@ -361,8 +361,8 @@ public:
 
 		if (GetMutableDefault<UEditorExperimentalSettings>()->bProceduralLandscape && this->IsToolActive())
 		{
-			ALandscape* Landscape = this->EdMode->CurrentToolTarget.LandscapeInfo->LandscapeActor.Get();
-			if (Landscape != nullptr)
+			ALandscape* Landscape = this->EdMode->GetLandscape();
+			if (Landscape)
 			{
 				Landscape->RequestProceduralContentUpdate(EProceduralContentUpdateFlag::Weightmap_Render);
 			}
@@ -373,11 +373,11 @@ public:
 	{
 		if (GetMutableDefault<UEditorExperimentalSettings>()->bProceduralLandscape)
 		{
-			this->EdMode->ChangeWeightmapsToCurrentProceduralLayerWeightmaps(true);
-
-			if (this->EdMode->CurrentToolTarget.LandscapeInfo->LandscapeActor.IsValid())
+			ALandscape* Landscape = this->EdMode->GetLandscape();
+			if (Landscape)
 			{
-				this->EdMode->CurrentToolTarget.LandscapeInfo->LandscapeActor->RequestProceduralContentUpdate(EProceduralContentUpdateFlag::Weightmap_All);
+				Landscape->SetCurrentEditingProceduralLayer();
+				Landscape->RequestProceduralContentUpdate(EProceduralContentUpdateFlag::Weightmap_All);
 			}
 		}
 
diff --git a/Engine/Source/Editor/LandscapeEditor/Private/LandscapeEdModePaintTools.cpp b/Engine/Source/Editor/LandscapeEditor/Private/LandscapeEdModePaintTools.cpp
--- a/Engine/Source/Editor/LandscapeEditor/Private/LandscapeEdModePaintTools.cpp
+++ b/Engine/Source/Editor/LandscapeEditor/Private/LandscapeEdModePaintTools.cpp
@@ -77,19 +77,12 @@ public:
 	{
 		if (GetMutableDefault<UEditorExperimentalSettings>()->bProceduralLandscape)
 		{
-			ALandscape* Landscape = this->EdMode->CurrentToolTarget.LandscapeInfo->LandscapeActor.Get();
-			if (Landscape != nullptr)
+			ALandscape* Landscape = this->EdMode->GetLandscape();
+			if (Landscape)
 			{
-				if (this->EdMode->CurrentToolTarget.TargetType == ELandscapeToolTargetType::Type::Heightmap)
-				{
-					Landscape->RequestProceduralContentUpdate(EProceduralContentUpdateFlag::Heightmap_Render);
-					this->EdMode->ChangeHeightmapsToCurrentProceduralLayerHeightmaps();
-				}
-				else
-				{
-					Landscape->RequestProceduralContentUpdate(EProceduralContentUpdateFlag::Weightmap_Render);
-					this->EdMode->ChangeWeightmapsToCurrentProceduralLayerWeightmaps();
-				}
+				bool bUpdateHeightmap = this->EdMode->CurrentToolTarget.TargetType == ELandscapeToolTargetType::Type::Heightmap;
+				Landscape->RequestProceduralContentUpdate(bUpdateHeightmap ? EProceduralContentUpdateFlag::Heightmap_Render : EProceduralContentUpdateFlag::Weightmap_Render);
+				Landscape->SetCurrentEditingProceduralLayer(this->EdMode->GetCurrentProceduralLayerGuid());
 			}
 		}
 
@@ -100,23 +93,12 @@ public:
 	{
 		if (GetMutableDefault<UEditorExperimentalSettings>()->bProceduralLandscape)
 		{
-			if (this->EdMode->CurrentToolTarget.TargetType == ELandscapeToolTargetType::Type::Heightmap)
-			{
-				this->EdMode->ChangeHeightmapsToCurrentProceduralLayerHeightmaps(true);
-
-				if (this->EdMode->CurrentToolTarget.LandscapeInfo->LandscapeActor.IsValid())
-				{
-					this->EdMode->CurrentToolTarget.LandscapeInfo->LandscapeActor->RequestProceduralContentUpdate(EProceduralContentUpdateFlag::Heightmap_All);
-				}
-			}
-			else
+			ALandscape* Landscape = this->EdMode->GetLandscape();
+			if (Landscape)
 			{
-				this->EdMode->ChangeWeightmapsToCurrentProceduralLayerWeightmaps(true);
-
-				if (this->EdMode->CurrentToolTarget.LandscapeInfo->LandscapeActor.IsValid())
-				{
-					this->EdMode->CurrentToolTarget.LandscapeInfo->LandscapeActor->RequestProceduralContentUpdate(EProceduralContentUpdateFlag::Weightmap_All);
-				}
+				bool bUpdateHeightmap = this->EdMode->CurrentToolTarget.TargetType == ELandscapeToolTargetType::Type::Heightmap;
+				Landscape->SetCurrentEditingProceduralLayer();
+				Landscape->RequestProceduralContentUpdate(bUpdateHeightmap ? EProceduralContentUpdateFlag::Heightmap_All : EProceduralContentUpdateFlag::Weightmap_All);
 			}
 		}
 
diff --git a/Engine/Source/Editor/LandscapeEditor/Private/LandscapeEditorDetailCustomization_ProceduralLayers.cpp b/Engine/Source/Editor/LandscapeEditor/Private/LandscapeEditorDetailCustomization_ProceduralLayers.cpp
--- a/Engine/Source/Editor/LandscapeEditor/Private/LandscapeEditorDetailCustomization_ProceduralLayers.cpp
+++ b/Engine/Source/Editor/LandscapeEditor/Private/LandscapeEditorDetailCustomization_ProceduralLayers.cpp
@@ -22,6 +22,7 @@
 #include "EditorModes.h"
 #include "LandscapeEditorModule.h"
 #include "LandscapeEditorObject.h"
+#include "Landscape.h"
 
 #include "DetailLayoutBuilder.h"
 #include "IDetailPropertyRow.h"
@@ -44,6 +45,8 @@
 #include "LandscapeEditorDetailCustomization_TargetLayers.h"
 #include "Widgets/Input/SEditableText.h"
 #include "Widgets/Input/SNumericEntryBox.h"
+#include "Widgets/Text/SInlineEditableTextBlock.h"
+#include "LandscapeEditorCommands.h"
 
 #define LOCTEXT_NAMESPACE "LandscapeEditor.Layers"
 
@@ -60,12 +63,7 @@ void FLandscapeEditorDetailCustomization_ProceduralLayers::CustomizeDetails(IDet
 	FEdModeLandscape* LandscapeEdMode = GetEditorMode();
 	if (LandscapeEdMode && LandscapeEdMode->CurrentToolMode != nullptr)
 	{
-		const FName CurrentToolName = LandscapeEdMode->CurrentTool->GetToolName();
-
-		if (LandscapeEdMode->CurrentToolMode->SupportedTargetTypes != 0)
-		{
-			LayerCategory.AddCustomBuilder(MakeShareable(new FLandscapeEditorCustomNodeBuilder_ProceduralLayers(DetailBuilder.GetThumbnailPool().ToSharedRef())));
-		}
+		LayerCategory.AddCustomBuilder(MakeShareable(new FLandscapeEditorCustomNodeBuilder_ProceduralLayers(DetailBuilder.GetThumbnailPool().ToSharedRef())));
 	}
 }
 END_SLATE_FUNCTION_BUILD_OPTIMIZATION
@@ -103,7 +101,7 @@ void FLandscapeEditorCustomNodeBuilder_ProceduralLayers::GenerateHeaderRowConten
 		[
 			SNew(STextBlock)
 			.Font(IDetailLayoutBuilder::GetDetailFont())
-			.Text(FText::FromString(TEXT("")))
+			.Text(FText::FromString(TEXT("Layers")))
 		];
 }
 
@@ -127,6 +125,8 @@ void FLandscapeEditorCustomNodeBuilder_ProceduralLayers::GenerateChildContent(ID
 				LayerList.ToSharedRef()
 			];
 
+		InlineTextBlocks.Empty();
+		InlineTextBlocks.AddDefaulted(LandscapeEdMode->GetProceduralLayerCount());
 		for (int32 i = 0; i < LandscapeEdMode->GetProceduralLayerCount(); ++i)
 		{
 			TSharedPtr<SWidget> GeneratedRowWidget = GenerateRow(i);
@@ -151,58 +151,55 @@ TSharedPtr<SWidget> FLandscapeEditorCustomNodeBuilder_ProceduralLayers::Generate
 	TSharedPtr<SWidget> RowWidget = SNew(SLandscapeEditorSelectableBorder)
 		.Padding(0)
 		.VAlign(VAlign_Center)
-		//.OnContextMenuOpening_Static(&FLandscapeEditorCustomNodeBuilder_Layers::OnTargetLayerContextMenuOpening, Target)
+		.OnContextMenuOpening(this, &FLandscapeEditorCustomNodeBuilder_ProceduralLayers::OnLayerContextMenuOpening, InLayerIndex)
 		.OnSelected(this, &FLandscapeEditorCustomNodeBuilder_ProceduralLayers::OnLayerSelectionChanged, InLayerIndex)
 		.IsSelected(TAttribute<bool>::Create(TAttribute<bool>::FGetter::CreateSP(this, &FLandscapeEditorCustomNodeBuilder_ProceduralLayers::IsLayerSelected, InLayerIndex)))
 		.Visibility(EVisibility::Visible)
 		[
 			SNew(SHorizontalBox)
 			
-			/*+ SHorizontalBox::Slot()
+			+ SHorizontalBox::Slot()
 			.AutoWidth()
 			.VAlign(VAlign_Center)
-			.Padding(FMargin(2))
 			[
-				SNew(SImage)
-				.Image(FEditorStyle::GetBrush(TEXT("LandscapeEditor.Target_Heightmap")))
+				SNew(SButton)
+				.ButtonStyle(FEditorStyle::Get(), "NoBorder")
+				.OnClicked(this, &FLandscapeEditorCustomNodeBuilder_ProceduralLayers::OnToggleLock, InLayerIndex)
+				.ToolTipText(LOCTEXT("FLandscapeEditorCustomNodeBuilder_ProceduralLayerLock", "Locks the current layer"))
+				[
+					SNew(SImage)
+					.Image(this, &FLandscapeEditorCustomNodeBuilder_ProceduralLayers::GetLockBrushForLayer, InLayerIndex)
+				]
 			]
-			*/
 
-			+ SHorizontalBox::Slot()
+			+SHorizontalBox::Slot()
+			.AutoWidth()
 			.VAlign(VAlign_Center)
-			.FillWidth(1.0f)
-			.Padding(4, 0)
 			[
-				SNew(SVerticalBox)
-				+ SVerticalBox::Slot()
-				.AutoHeight()
+				SNew(SButton)
+				.ContentPadding(0)
+				.ButtonStyle(FEditorStyle::Get(), "NoBorder")
+				.OnClicked(this, &FLandscapeEditorCustomNodeBuilder_ProceduralLayers::OnToggleVisibility, InLayerIndex)
+				.ToolTipText(LOCTEXT("FLandscapeEditorCustomNodeBuilder_ProceduralLayerVisibility", "Toggle Layer Visibility"))
+				.HAlign(HAlign_Center)
 				.VAlign(VAlign_Center)
-				.Padding(0, 2)
-				.HAlign(HAlign_Left)				
+				.Content()
 				[
-					SNew(SEditableText)
-					.SelectAllTextWhenFocused(true)
-					.IsReadOnly(true)
-					.Text(this, &FLandscapeEditorCustomNodeBuilder_ProceduralLayers::GetLayerText, InLayerIndex)
-					.ToolTipText(LOCTEXT("FLandscapeEditorCustomNodeBuilder_ProceduralLayers_tooltip", "Name of the Layer"))
-					.OnTextCommitted(FOnTextCommitted::CreateSP(this, &FLandscapeEditorCustomNodeBuilder_ProceduralLayers::OnLayerTextCommitted, InLayerIndex))
+					SNew(SImage)
+					.Image(this, &FLandscapeEditorCustomNodeBuilder_ProceduralLayers::GetVisibilityBrushForLayer, InLayerIndex)
 				]
 			]
+
 			+ SHorizontalBox::Slot()
-			.FillWidth(1.0f)
+			.FillWidth(1.0)
 			.VAlign(VAlign_Center)
-			.Padding(0, 2)
-			.HAlign(HAlign_Center)				
+			.Padding(4, 0)
 			[
-				SNew(SCheckBox)
-				.OnCheckStateChanged(this, &FLandscapeEditorCustomNodeBuilder_ProceduralLayers::OnLayerVisibilityChanged, InLayerIndex)
-				.IsChecked(TAttribute<ECheckBoxState>::Create(TAttribute<ECheckBoxState>::FGetter::CreateSP(this, &FLandscapeEditorCustomNodeBuilder_ProceduralLayers::IsLayerVisible, InLayerIndex)))
-				.ToolTipText(LOCTEXT("FLandscapeEditorCustomNodeBuilder_ProceduralLayerVisibility_Tooltips", "Is layer visible"))
-				.Content()
-				[
-					SNew(STextBlock)
-					.Text(LOCTEXT("FLandscapeEditorCustomNodeBuilder_ProceduralLayerVisibility", "Visibility"))
-				]
+				SAssignNew(InlineTextBlocks[InLayerIndex], SInlineEditableTextBlock)
+				.Text(this, &FLandscapeEditorCustomNodeBuilder_ProceduralLayers::GetLayerText, InLayerIndex)
+				.ToolTipText(LOCTEXT("FLandscapeEditorCustomNodeBuilder_ProceduralLayers_tooltip", "Name of the Layer"))
+				.OnVerifyTextChanged(FOnVerifyTextChanged::CreateSP(this, &FLandscapeEditorCustomNodeBuilder_ProceduralLayers::CanRenameProceduralLayerTo, InLayerIndex))
+				.OnTextCommitted(FOnTextCommitted::CreateSP(this, &FLandscapeEditorCustomNodeBuilder_ProceduralLayers::SetProceduralLayerName, InLayerIndex))
 			]
 
 			+ SHorizontalBox::Slot()
@@ -244,37 +241,206 @@ TSharedPtr<SWidget> FLandscapeEditorCustomNodeBuilder_ProceduralLayers::Generate
 }
 END_SLATE_FUNCTION_BUILD_OPTIMIZATION
 
-void FLandscapeEditorCustomNodeBuilder_ProceduralLayers::OnLayerTextCommitted(const FText& InText, ETextCommit::Type InCommitType, int32 InLayerIndex)
+FText FLandscapeEditorCustomNodeBuilder_ProceduralLayers::GetLayerText(int32 InLayerIndex) const
 {
 	FEdModeLandscape* LandscapeEdMode = GetEditorMode();
+	if (LandscapeEdMode)
+	{
+		return FText::FromName(LandscapeEdMode->GetProceduralLayerName(InLayerIndex));
+	}
 
-	if (LandscapeEdMode != nullptr)
+	return FText::FromString(TEXT("None"));
+}
+
+bool FLandscapeEditorCustomNodeBuilder_ProceduralLayers::IsLayerSelected(int32 InLayerIndex)
+{
+	FEdModeLandscape* LandscapeEdMode = GetEditorMode();
+	if (LandscapeEdMode)
 	{
+		return LandscapeEdMode->GetCurrentProceduralLayerIndex() == InLayerIndex;
+	}
+
+	return false;
+}
+
+bool FLandscapeEditorCustomNodeBuilder_ProceduralLayers::CanRenameProceduralLayerTo(const FText& InNewText, FText& OutErrorMessage, int32 InLayerIndex)
+{
+	FEdModeLandscape* LandscapeEdMode = GetEditorMode();
+	if (LandscapeEdMode)
+	{
+		if (!LandscapeEdMode->CanRenameProceduralLayerTo(InLayerIndex, *InNewText.ToString()))
+		{
+			OutErrorMessage = LOCTEXT("RenameFailed_AlreadyExists", "This layer already exists");
+			return false;
+		}
+	}
+	return true;
+}
+
+void FLandscapeEditorCustomNodeBuilder_ProceduralLayers::SetProceduralLayerName(const FText& InText, ETextCommit::Type InCommitType, int32 InLayerIndex)
+{
+	FEdModeLandscape* LandscapeEdMode = GetEditorMode();
+	if (LandscapeEdMode)
+	{
+		const FScopedTransaction Transaction(LOCTEXT("Landscape_ProceduralLayers_Rename", "Rename Procedural Layer"));
 		LandscapeEdMode->SetProceduralLayerName(InLayerIndex, *InText.ToString());
 	}
 }
 
-FText FLandscapeEditorCustomNodeBuilder_ProceduralLayers::GetLayerText(int32 InLayerIndex) const
+TSharedPtr<SWidget> FLandscapeEditorCustomNodeBuilder_ProceduralLayers::OnLayerContextMenuOpening(int32 InLayerIndex)
 {
 	FEdModeLandscape* LandscapeEdMode = GetEditorMode();
+	ALandscape* Landscape = LandscapeEdMode ? LandscapeEdMode->GetLandscape() : nullptr;
+	if (LandscapeEdMode && Landscape)
+	{
+		FProceduralLayer* Layer = LandscapeEdMode->GetProceduralLayer(InLayerIndex);
+		TSharedRef<FLandscapeEditorCustomNodeBuilder_ProceduralLayers> SharedThis = AsShared();
+		FMenuBuilder MenuBuilder(true, NULL);
+		MenuBuilder.BeginSection("LandscapeEditorProceduralLayerActions", LOCTEXT("LandscapeEditorProceduralLayerActions.Heading", "Layers"));
+		{
+			// Create Layer
+			FUIAction CreateLayerAction = FUIAction(FExecuteAction::CreateLambda([SharedThis] { SharedThis->CreateLayer(); }));
+			MenuBuilder.AddMenuEntry(LOCTEXT("CreateLayer", "Create"), LOCTEXT("CreateLayerTooltip", "Create Layer"), FSlateIcon(), CreateLayerAction);
 
-	if (LandscapeEdMode != nullptr)
+			if (Layer)
+			{
+				// Rename Layer
+				FUIAction RenameLayerAction = FUIAction(FExecuteAction::CreateLambda([SharedThis, InLayerIndex] { SharedThis->RenameLayer(InLayerIndex); }));
+				MenuBuilder.AddMenuEntry(LOCTEXT("RenameLayer", "Rename..."), LOCTEXT("RenameLayerTooltip", "Rename Layer"), FSlateIcon(), RenameLayerAction);
+
+				if (!Layer->bLocked)
+				{
+					// Clear Layer
+					FUIAction ClearLayerAction = FUIAction(FExecuteAction::CreateLambda([SharedThis, InLayerIndex] { SharedThis->ClearLayer(InLayerIndex); }));
+					MenuBuilder.AddMenuEntry(LOCTEXT("ClearLayer", "Clear..."), LOCTEXT("ClearLayerTooltip", "Clear Layer"), FSlateIcon(), ClearLayerAction);
+
+					if (Landscape->ProceduralLayers.Num() > 1)
+					{
+						// Delete Layer
+						FUIAction DeleteLayerAction = FUIAction(FExecuteAction::CreateLambda([SharedThis, InLayerIndex] { SharedThis->DeleteLayer(InLayerIndex); } ));
+						MenuBuilder.AddMenuEntry(LOCTEXT("DeleteLayer", "Delete..."), LOCTEXT("DeleteLayerTooltip", "Delete Layer"), FSlateIcon(), DeleteLayerAction);
+					}
+				}
+			}
+		}
+		MenuBuilder.EndSection();
+
+		MenuBuilder.BeginSection("LandscapeEditorProceduralLayerVisibility", LOCTEXT("LandscapeEditorProceduralLayerVisibility.Heading", "Visibility"));
+		{
+			if (Layer)
+			{
+				if (Layer->bVisible)
+				{
+					// Hide Selected Layer
+					FUIAction HideSelectedLayerAction = FUIAction(FExecuteAction::CreateLambda([SharedThis, InLayerIndex] { SharedThis->OnToggleVisibility(InLayerIndex); }));
+					MenuBuilder.AddMenuEntry(LOCTEXT("HideSelectedLayer", "Hide Selected"), LOCTEXT("HideSelectedLayerTooltip", "Hide Selected Layer"), FSlateIcon(), HideSelectedLayerAction);
+				}
+				else
+				{
+					// Show Selected Layer
+					FUIAction ShowSelectedLayerAction = FUIAction(FExecuteAction::CreateLambda([SharedThis, InLayerIndex] { SharedThis->OnToggleVisibility(InLayerIndex); }));
+					MenuBuilder.AddMenuEntry(LOCTEXT("ShowSelectedLayer", "Show Selected"), LOCTEXT("ShowSelectedLayerTooltip", "Show Selected Layer"), FSlateIcon(), ShowSelectedLayerAction);
+				}
+
+				// Show Only Selected Layer
+				FUIAction ShowOnlySelectedLayerAction = FUIAction(FExecuteAction::CreateLambda([SharedThis, InLayerIndex] { SharedThis->ShowOnlySelectedLayer(InLayerIndex); }));
+				MenuBuilder.AddMenuEntry(LOCTEXT("ShowOnlySelectedLayer", "Show Only Selected"), LOCTEXT("ShowOnlySelectedLayerTooltip", "Show Only Selected Layer"), FSlateIcon(), ShowOnlySelectedLayerAction);
+			}
+
+			// Show All Layers
+			FUIAction ShowAllLayersAction = FUIAction(FExecuteAction::CreateLambda([SharedThis] { SharedThis->ShowAllLayers(); }));
+			MenuBuilder.AddMenuEntry(LOCTEXT("ShowAllLayers", "Show All Layers"), LOCTEXT("ShowAllLayersTooltip", "Show All Layers"), FSlateIcon(), ShowAllLayersAction);
+		}
+		MenuBuilder.EndSection();
+
+		return MenuBuilder.MakeWidget();
+	}
+	return NULL;
+}
+
+void FLandscapeEditorCustomNodeBuilder_ProceduralLayers::RenameLayer(int32 InLayerIndex)
+{
+	if (InlineTextBlocks.IsValidIndex(InLayerIndex) && InlineTextBlocks[InLayerIndex].IsValid())
 	{
-		return FText::FromName(LandscapeEdMode->GetProceduralLayerName(InLayerIndex));
+		InlineTextBlocks[InLayerIndex]->EnterEditingMode();
 	}
+}
 
-	return FText::FromString(TEXT("None"));
+void FLandscapeEditorCustomNodeBuilder_ProceduralLayers::ClearLayer(int32 InLayerIndex)
+{
+	FEdModeLandscape* LandscapeEdMode = GetEditorMode();
+	ALandscape* Landscape = LandscapeEdMode ? LandscapeEdMode->GetLandscape() : nullptr;
+	if (Landscape)
+	{
+		FProceduralLayer* Layer = LandscapeEdMode->GetProceduralLayer(InLayerIndex);
+		if (Layer)
+		{
+			EAppReturnType::Type Result = FMessageDialog::Open(EAppMsgType::YesNo, FText::Format(LOCTEXT("LandscapeMode_Message", "The layer {0} content will be completely cleared.  Continue?"), FText::FromName(Layer->Name)));
+			if (Result == EAppReturnType::Yes)
+			{
+				const FScopedTransaction Transaction(LOCTEXT("Landscape_ProceduralLayers_Clean", "Clean Procedural Layer"));
+				Landscape->ClearProceduralLayer(InLayerIndex);
+			}
+		}
+	}
 }
 
-bool FLandscapeEditorCustomNodeBuilder_ProceduralLayers::IsLayerSelected(int32 InLayerIndex)
+void FLandscapeEditorCustomNodeBuilder_ProceduralLayers::DeleteLayer(int32 InLayerIndex)
 {
 	FEdModeLandscape* LandscapeEdMode = GetEditorMode();
-	if (LandscapeEdMode)
+	ALandscape* Landscape = LandscapeEdMode ? LandscapeEdMode->GetLandscape() : nullptr;
+	if (Landscape && Landscape->ProceduralLayers.Num() > 1)
 	{
-		return LandscapeEdMode->GetCurrentProceduralLayerIndex() == InLayerIndex;
+		FProceduralLayer* Layer = LandscapeEdMode->GetProceduralLayer(InLayerIndex);
+		if (Layer)
+		{
+			EAppReturnType::Type Result = FMessageDialog::Open(EAppMsgType::YesNo, FText::Format(LOCTEXT("LandscapeMode_Message", "The layer {0} will be deleted.  Continue?"), FText::FromName(Layer->Name)));
+			if (Result == EAppReturnType::Yes)
+			{
+				const FScopedTransaction Transaction(LOCTEXT("Landscape_ProceduralLayers_Delete", "Delete Procedural Layer"));
+				Landscape->DeleteProceduralLayer(InLayerIndex);
+				int32 NewLayerSelectionIndex = Landscape->GetProceduralLayer(InLayerIndex) ? InLayerIndex : 0;
+				OnLayerSelectionChanged(NewLayerSelectionIndex);
+				LandscapeEdMode->RefreshDetailPanel();
+			}
+		}
 	}
+}
 
-	return false;
+void FLandscapeEditorCustomNodeBuilder_ProceduralLayers::ShowOnlySelectedLayer(int32 InLayerIndex)
+{
+	FEdModeLandscape* LandscapeEdMode = GetEditorMode();
+	ALandscape* Landscape = LandscapeEdMode ? LandscapeEdMode->GetLandscape() : nullptr;
+	if (Landscape)
+	{
+		const FScopedTransaction Transaction(LOCTEXT("ShowOnlySelectedLayer", "Show Only Selected Layer"));
+		Landscape->ShowOnlySelectedProceduralLayer(InLayerIndex);
+	}
+}
+
+void FLandscapeEditorCustomNodeBuilder_ProceduralLayers::ShowAllLayers()
+{
+	FEdModeLandscape* LandscapeEdMode = GetEditorMode();
+	ALandscape* Landscape = LandscapeEdMode ? LandscapeEdMode->GetLandscape() : nullptr;
+	if (Landscape)
+	{
+		const FScopedTransaction Transaction(LOCTEXT("ShowAllLayers", "Show All Layers"));
+		Landscape->ShowAllProceduralLayers();
+	}
+}
+
+void FLandscapeEditorCustomNodeBuilder_ProceduralLayers::CreateLayer()
+{
+	FEdModeLandscape* LandscapeEdMode = GetEditorMode();
+	ALandscape* Landscape = LandscapeEdMode ? LandscapeEdMode->GetLandscape() : nullptr;
+	if (Landscape)
+	{
+		{
+			const FScopedTransaction Transaction(LOCTEXT("Landscape_ProceduralLayers_Create", "Create Procedural Layer"));
+			Landscape->CreateProceduralLayer();
+		}
+		LandscapeEdMode->RefreshDetailPanel();
+	}
 }
 
 void FLandscapeEditorCustomNodeBuilder_ProceduralLayers::OnLayerSelectionChanged(int32 InLayerIndex)
@@ -305,30 +471,43 @@ void FLandscapeEditorCustomNodeBuilder_ProceduralLayers::SetLayerAlpha(float InA
 
 	if (LandscapeEdMode)
 	{
-		LandscapeEdMode->SetProceduralLayerAlpha(InAlpha, InLayerIndex);
+		const FScopedTransaction Transaction(LOCTEXT("Landscape_ProceduralLayers_SetAlpha", "Set Procedural Layer Alpha"));
+		LandscapeEdMode->SetProceduralLayerAlpha(InLayerIndex, InAlpha);
 	}
 }
 
-void FLandscapeEditorCustomNodeBuilder_ProceduralLayers::OnLayerVisibilityChanged(ECheckBoxState NewState, int32 InLayerIndex)
+FReply FLandscapeEditorCustomNodeBuilder_ProceduralLayers::OnToggleVisibility(int32 InLayerIndex)
 {
 	FEdModeLandscape* LandscapeEdMode = GetEditorMode();
-
 	if (LandscapeEdMode)
 	{
-		LandscapeEdMode->SetProceduralLayerVisibility(NewState == ECheckBoxState::Checked, InLayerIndex);
+		LandscapeEdMode->SetProceduralLayerVisibility(!LandscapeEdMode->IsProceduralLayerVisible(InLayerIndex), InLayerIndex);
 	}
+	return FReply::Handled();
 }
 
-ECheckBoxState FLandscapeEditorCustomNodeBuilder_ProceduralLayers::IsLayerVisible(int32 InLayerIndex) const
+const FSlateBrush* FLandscapeEditorCustomNodeBuilder_ProceduralLayers::GetVisibilityBrushForLayer(int32 InLayerIndex) const
 {
 	FEdModeLandscape* LandscapeEdMode = GetEditorMode();
+	bool bIsVisible = LandscapeEdMode && LandscapeEdMode->IsProceduralLayerVisible(InLayerIndex);
+	return bIsVisible ? FEditorStyle::GetBrush("Level.VisibleIcon16x") : FEditorStyle::GetBrush("Level.NotVisibleIcon16x");
+}
 
+FReply FLandscapeEditorCustomNodeBuilder_ProceduralLayers::OnToggleLock(int32 InLayerIndex)
+{
+	FEdModeLandscape* LandscapeEdMode = GetEditorMode();
 	if (LandscapeEdMode)
 	{
-		return LandscapeEdMode->IsProceduralLayerVisible(InLayerIndex) ? ECheckBoxState::Checked : ECheckBoxState::Unchecked;
+		LandscapeEdMode->SetProceduralLayerLocked(InLayerIndex, !LandscapeEdMode->IsProceduralLayerLocked(InLayerIndex));
+	}
+	return FReply::Handled();
 	}
 
-	return ECheckBoxState::Unchecked;
+const FSlateBrush* FLandscapeEditorCustomNodeBuilder_ProceduralLayers::GetLockBrushForLayer(int32 InLayerIndex) const
+{
+	FEdModeLandscape* LandscapeEdMode = GetEditorMode();
+	bool bIsLocked = LandscapeEdMode && LandscapeEdMode->IsProceduralLayerLocked(InLayerIndex);
+	return bIsLocked ? FEditorStyle::GetBrush(TEXT("PropertyWindow.Locked")) : FEditorStyle::GetBrush(TEXT("PropertyWindow.Unlocked"));
 }
 
 FReply FLandscapeEditorCustomNodeBuilder_ProceduralLayers::HandleDragDetected(const FGeometry& MyGeometry, const FPointerEvent& MouseEvent, int32 SlotIndex, SVerticalBox::FSlot* Slot)
diff --git a/Engine/Source/Editor/LandscapeEditor/Private/LandscapeEditorDetailCustomization_ProceduralLayers.h b/Engine/Source/Editor/LandscapeEditor/Private/LandscapeEditorDetailCustomization_ProceduralLayers.h
--- a/Engine/Source/Editor/LandscapeEditor/Private/LandscapeEditorDetailCustomization_ProceduralLayers.h
+++ b/Engine/Source/Editor/LandscapeEditor/Private/LandscapeEditorDetailCustomization_ProceduralLayers.h
@@ -7,6 +7,7 @@
 #include "Layout/Visibility.h"
 #include "Layout/Margin.h"
 #include "Styling/SlateColor.h"
+#include "Styling/SlateBrush.h"
 #include "Input/Reply.h"
 #include "Widgets/DeclarativeSyntaxSupport.h"
 #include "Widgets/SWidget.h"
@@ -65,13 +66,28 @@ protected:
 
 	bool IsLayerSelected(int32 LayerIndex);
 	void OnLayerSelectionChanged(int32 LayerIndex);
-
-	void OnLayerTextCommitted(const FText& InText, ETextCommit::Type InCommitType, int32 InLayerIndex);
+	TSharedPtr<SWidget> OnLayerContextMenuOpening(int32 InLayerIndex);
+	void CreateLayer();
+	void ClearLayer(int32 InLayerIndex);
+	void RenameLayer(int32 InLayerIndex);
+	void DeleteLayer(int32 InLayerIndex);
+	void ShowOnlySelectedLayer(int32 InLayerIndex);
+	void ShowAllLayers();
+	bool CanRenameProceduralLayerTo(const FText& NewText, FText& OutErrorMessage, int32 InLayerIndex);
+	void SetProceduralLayerName(const FText& InText, ETextCommit::Type InCommitType, int32 InLayerIndex);
 	FText GetLayerText(int32 InLayerIndex) const;
 
 	TOptional<float> GetLayerAlpha(int32 InLayerIndex) const;
 	void SetLayerAlpha(float InAlpha, int32 InLayerIndex);
 
-	void OnLayerVisibilityChanged(ECheckBoxState NewState, int32 InLayerIndex);
-	ECheckBoxState IsLayerVisible(int32 InLayerIndex) const;
+	FReply OnToggleVisibility(int32 InLayerIndex);
+	const FSlateBrush* GetVisibilityBrushForLayer(int32 InLayerIndex) const;
+	
+	FReply OnToggleLock(int32 InLayerIndex);
+	const FSlateBrush* GetLockBrushForLayer(int32 InLayerIndex) const;
+
+private:
+
+	/** Widgets for displaying and editing the layer name */
+	TArray< TSharedPtr< SInlineEditableTextBlock > > InlineTextBlocks;
 };
\ No newline at end of file
diff --git a/Engine/Source/Editor/LandscapeEditor/Private/LandscapeEditorDetailCustomization_TargetLayers.cpp b/Engine/Source/Editor/LandscapeEditor/Private/LandscapeEditorDetailCustomization_TargetLayers.cpp
--- a/Engine/Source/Editor/LandscapeEditor/Private/LandscapeEditorDetailCustomization_TargetLayers.cpp
+++ b/Engine/Source/Editor/LandscapeEditor/Private/LandscapeEditorDetailCustomization_TargetLayers.cpp
@@ -22,6 +22,7 @@
 #include "EditorModes.h"
 #include "LandscapeEditorModule.h"
 #include "LandscapeEditorObject.h"
+#include "Landscape.h"
 
 #include "DetailLayoutBuilder.h"
 #include "IDetailPropertyRow.h"
@@ -1079,17 +1080,7 @@ void FLandscapeEditorCustomNodeBuilder_TargetLayers::OnReimportLayer(const TShar
 	FEdModeLandscape* LandscapeEdMode = GetEditorMode();
 	if (LandscapeEdMode)
 	{
-		if (GetMutableDefault<UEditorExperimentalSettings>()->bProceduralLandscape)
-		{
-			LandscapeEdMode->ChangeWeightmapsToCurrentProceduralLayerWeightmaps();
-			LandscapeEdMode->ReimportData(*Target);
-			LandscapeEdMode->ChangeWeightmapsToCurrentProceduralLayerWeightmaps(true);
-			LandscapeEdMode->RequestProceduralContentUpdate();
-		}
-		else
-		{
-			LandscapeEdMode->ReimportData(*Target);
-		}
+		LandscapeEdMode->ReimportData(*Target);
 	}
 }
 
@@ -1103,17 +1094,8 @@ void FLandscapeEditorCustomNodeBuilder_TargetLayers::OnFillLayer(const TSharedRe
 		FEdModeLandscape* LandscapeEdMode = GetEditorMode();
 		if (LandscapeEdMode)
 		{
-			if (GetMutableDefault<UEditorExperimentalSettings>()->bProceduralLandscape)
-			{
-				LandscapeEdMode->ChangeWeightmapsToCurrentProceduralLayerWeightmaps();
-				LandscapeEdit.FillLayer(Target->LayerInfoObj.Get());
-				LandscapeEdMode->ChangeWeightmapsToCurrentProceduralLayerWeightmaps(true);
-				LandscapeEdMode->RequestProceduralContentUpdate(true);
-			}
-			else
-			{
-				LandscapeEdit.FillLayer(Target->LayerInfoObj.Get());
-			}
+			FScopedSetLandscapeCurrentEditingProceduralLayer Scope(LandscapeEdMode->GetLandscape(), LandscapeEdMode->GetCurrentProceduralLayerGuid(), [&] { LandscapeEdMode->RequestProceduralContentUpdate(true); });
+			LandscapeEdit.FillLayer(Target->LayerInfoObj.Get());
 		}
 	}
 }
@@ -1124,14 +1106,13 @@ void FLandscapeEditorCustomNodeBuilder_TargetLayers::FillEmptyLayers(ULandscapeI
 	if (LandscapeEdMode)
 	{
 		FLandscapeEditDataInterface LandscapeEdit(LandscapeInfo);
-
+		
 		if (GetMutableDefault<UEditorExperimentalSettings>()->bProceduralLandscape)
 		{
 			if (LandscapeEdMode->NeedToFillEmptyLayersForProcedural())
 			{
-				LandscapeEdMode->ChangeWeightmapsToCurrentProceduralLayerWeightmaps();
+				FScopedSetLandscapeCurrentEditingProceduralLayer Scope(LandscapeEdMode->GetLandscape(), LandscapeEdMode->GetCurrentProceduralLayerGuid());
 				LandscapeEdit.FillEmptyLayers(LandscapeInfoObject);
-				LandscapeEdMode->ChangeWeightmapsToCurrentProceduralLayerWeightmaps(true);
 			}
 		}
 		else
@@ -1151,19 +1132,9 @@ void FLandscapeEditorCustomNodeBuilder_TargetLayers::OnClearLayer(const TSharedR
 		FEdModeLandscape* LandscapeEdMode = GetEditorMode();
 		if (LandscapeEdMode)
 		{
+			FScopedSetLandscapeCurrentEditingProceduralLayer Scope(LandscapeEdMode->GetLandscape(), LandscapeEdMode->GetCurrentProceduralLayerGuid(), [&] { LandscapeEdMode->RequestProceduralContentUpdate(true); });
 			FLandscapeEditDataInterface LandscapeEdit(Target->LandscapeInfo.Get());
-
-			if (GetMutableDefault<UEditorExperimentalSettings>()->bProceduralLandscape)
-			{
-				LandscapeEdMode->ChangeWeightmapsToCurrentProceduralLayerWeightmaps();
-				LandscapeEdit.DeleteLayer(Target->LayerInfoObj.Get());
-				LandscapeEdMode->ChangeWeightmapsToCurrentProceduralLayerWeightmaps(true);
-				LandscapeEdMode->RequestProceduralContentUpdate(true);
-			}
-			else
-			{
-				LandscapeEdit.DeleteLayer(Target->LayerInfoObj.Get());
-			}
+			LandscapeEdit.DeleteLayer(Target->LayerInfoObj.Get());
 		}
 	}
 }
diff --git a/Engine/Source/Editor/LandscapeEditor/Private/LandscapeEditorDetails.cpp b/Engine/Source/Editor/LandscapeEditor/Private/LandscapeEditorDetails.cpp
--- a/Engine/Source/Editor/LandscapeEditor/Private/LandscapeEditorDetails.cpp
+++ b/Engine/Source/Editor/LandscapeEditor/Private/LandscapeEditorDetails.cpp
@@ -108,20 +108,20 @@ void FLandscapeEditorDetails::CustomizeDetails(IDetailLayoutBuilder& DetailBuild
 	Customization_AlphaBrush = MakeShareable(new FLandscapeEditorDetailCustomization_AlphaBrush);
 	Customization_AlphaBrush->CustomizeDetails(DetailBuilder);
 
-	// Target Layers:
-	Customization_TargetLayers = MakeShareable(new FLandscapeEditorDetailCustomization_TargetLayers);
-	Customization_TargetLayers->CustomizeDetails(DetailBuilder);
-
 	if (GetMutableDefault<UEditorExperimentalSettings>()->bProceduralLandscape)
 	{
-		// Brush Stack
-		Customization_ProceduralBrushStack = MakeShareable(new FLandscapeEditorDetailCustomization_ProceduralBrushStack);
-		Customization_ProceduralBrushStack->CustomizeDetails(DetailBuilder);
-
 		// Procedural Layers
 		Customization_ProceduralLayers = MakeShareable(new FLandscapeEditorDetailCustomization_ProceduralLayers);
 		Customization_ProceduralLayers->CustomizeDetails(DetailBuilder);
+
+		// Brush Stack
+		Customization_ProceduralBrushStack = MakeShareable(new FLandscapeEditorDetailCustomization_ProceduralBrushStack);
+		Customization_ProceduralBrushStack->CustomizeDetails(DetailBuilder);
 	}
+
+	// Target Layers:
+	Customization_TargetLayers = MakeShareable(new FLandscapeEditorDetailCustomization_TargetLayers);
+	Customization_TargetLayers->CustomizeDetails(DetailBuilder);
 }
 
 FText FLandscapeEditorDetails::GetLocalizedName(FString Name)
diff --git a/Engine/Source/Runtime/Landscape/Classes/Landscape.h b/Engine/Source/Runtime/Landscape/Classes/Landscape.h
--- a/Engine/Source/Runtime/Landscape/Classes/Landscape.h
+++ b/Engine/Source/Runtime/Landscape/Classes/Landscape.h
@@ -145,17 +145,25 @@ struct FProceduralLayer
 	GENERATED_USTRUCT_BODY()
 
 	FProceduralLayer()
-		: Name(NAME_None)
-		, Visible(true)
+		: Guid(FGuid::NewGuid())
+		, Name(NAME_None)
+		, bVisible(true)
+		, bLocked(false)
 		, HeightmapAlpha(1.0f)
 		, WeightmapAlpha(1.0f)
 	{}
 
+	UPROPERTY(meta = (IgnoreForMemberInitializationTest))
+	FGuid Guid;
+
 	UPROPERTY()
 	FName Name;
 
+	UPROPERTY(Transient)
+	bool bVisible;
+
 	UPROPERTY()
-	bool Visible;
+	bool bLocked;
 
 	UPROPERTY()
 	float HeightmapAlpha;
@@ -222,6 +230,19 @@ public:
 	// Procedural stuff
 #if WITH_EDITOR
 	LANDSCAPE_API void RequestProceduralContentUpdate(uint32 InDataFlags, bool InUpdateAllMaterials = false);
+	LANDSCAPE_API void CreateProceduralLayer(FName InName = NAME_None, bool bInUpdateProceduralContent = true);
+	LANDSCAPE_API bool IsProceduralLayerNameUnique(const FName& InName) const;
+	LANDSCAPE_API void SetProceduralLayerName(int32 InLayerIndex, const FName& InName);
+	LANDSCAPE_API void SetProceduralLayerAlpha(int32 InLayerIndex, const float InAlpha, bool bInHeightmap);
+	LANDSCAPE_API void SetProceduralLayerVisibility(int32 InLayerIndex, bool bInVisible);
+	LANDSCAPE_API struct FProceduralLayer* GetProceduralLayer(int32 InLayerIndex);
+	LANDSCAPE_API const struct FProceduralLayer* GetProceduralLayer(int32 InLayerIndex) const;
+	LANDSCAPE_API void ClearProceduralLayer(int32 InLayerIndex);
+	LANDSCAPE_API void ClearProceduralLayer(const FGuid& InLayerGuid);
+	LANDSCAPE_API void DeleteProceduralLayer(int32 InLayerIndex);
+	LANDSCAPE_API void SetCurrentEditingProceduralLayer(FGuid InLayerGuid = FGuid());
+	LANDSCAPE_API void ShowOnlySelectedProceduralLayer(int32 InLayerIndex);
+	LANDSCAPE_API void ShowAllProceduralLayers();
 
 private:
 	void TickProcedural(float DeltaTime, ELevelTick TickType, FActorTickFunction& ThisTickFunction);
@@ -314,4 +335,23 @@ private:
 	// Used in packing the paint layers data contained into CombinedProcLayerWeightmapAllLayersResource to be set again for each component weightmap (size of the landscape)
 	class FLandscapeTexture2DResource* WeightmapScratchPackLayerTextureResource;	
 #endif
+
+protected:
+#if WITH_EDITOR
+	FName GenerateUniqueProceduralLayerName(FName InName = NAME_None) const;
+#endif
+};
+
+#if WITH_EDITOR
+class LANDSCAPE_API FScopedSetLandscapeCurrentEditingProceduralLayer
+{
+public:
+	FScopedSetLandscapeCurrentEditingProceduralLayer(ALandscape* InLandscape, const FGuid& InProceduralLayer, TFunction<void()> InCompletionCallback = TFunction<void()>());
+	~FScopedSetLandscapeCurrentEditingProceduralLayer();
+
+private:
+	TWeakObjectPtr<ALandscape> Landscape;
+	const FGuid& ProceduralLayer;
+	TFunction<void()> CompletionCallback;
 };
+#endif
diff --git a/Engine/Source/Runtime/Landscape/Private/LandscapeEditProcedural.cpp b/Engine/Source/Runtime/Landscape/Private/LandscapeEditProcedural.cpp
--- a/Engine/Source/Runtime/Landscape/Private/LandscapeEditProcedural.cpp
+++ b/Engine/Source/Runtime/Landscape/Private/LandscapeEditProcedural.cpp
@@ -28,6 +28,7 @@ LandscapeEditProcedural.cpp: Landscape editing procedural mode
 #include "LandscapeMaterialInstanceConstant.h"
 #include "Materials/MaterialExpressionLandscapeVisibilityMask.h"
 #include "ShaderCompiler.h"
+#include "Algo/Count.h"
 #endif
 
 #define LOCTEXT_NAMESPACE "Landscape"
@@ -1226,7 +1227,6 @@ void ALandscapeProxy::SetupProceduralLayers(int32 InNumComponentsX, int32 InNumC
 	check(Landscape);
 
 	ULandscapeInfo* Info = GetLandscapeInfo();
-
 	if (Info == nullptr)
 	{
 		return;
@@ -1234,54 +1234,16 @@ void ALandscapeProxy::SetupProceduralLayers(int32 InNumComponentsX, int32 InNumC
 
 	TArray<ALandscapeProxy*> AllLandscapes;
 	AllLandscapes.Add(Landscape);
-
 	for (auto& It : Info->Proxies)
 	{
 		AllLandscapes.Add(It);
 	}
 
-	// TEMP STUFF START
-	bool Layer1Exist = false;
-	FProceduralLayer Layer1;
-	Layer1.Name = TEXT("Layer1");
-
-	bool Layer2Exist = false;
-	FProceduralLayer Layer2;
-	Layer2.Name = TEXT("Layer2");	
-
-	for (int32 i = 0; i < Landscape->ProceduralLayers.Num(); ++i)
-	{
-		if (Landscape->ProceduralLayers[i].Name == Layer1.Name)
-		{
-			Layer1Exist = true;				
-		}
-
-		if (Landscape->ProceduralLayers[i].Name == Layer2.Name)
-		{
-			Layer2Exist = true;
-		}
-	}
-
-	if (!Layer1Exist)
+	// Make sure we have at least 1 layer
+	if (Landscape->ProceduralLayers.Num() == 0)
 	{
-		Landscape->ProceduralLayers.Add(Layer1);
-
-		for (ALandscapeProxy* LandscapeProxy : AllLandscapes)
-		{
-			LandscapeProxy->ProceduralLayersData.Add(Layer1.Name, FProceduralLayerData());
-		}
-	}
-
-	if (!Layer2Exist)
-	{
-		Landscape->ProceduralLayers.Add(Layer2);
-
-		for (ALandscapeProxy* LandscapeProxy : AllLandscapes)
-		{
-			LandscapeProxy->ProceduralLayersData.Add(Layer2.Name, FProceduralLayerData());
-		}
+		Landscape->CreateProceduralLayer(FName(TEXT("Layer")), false);
 	}	
-	///// TEMP STUFF END
 
 	// TODO: When using Change Component Size, we will call Setup again, currently all the existing data will get collapsed into the layer 1, it should keep the layers data.
 
@@ -2665,7 +2627,7 @@ void ALandscape::RegenerateProceduralHeightmaps()
 		{
 			//Draw Layer heightmap to Combined RT Atlas
 			ShaderParams.ApplyLayerModifiers = true;
-			ShaderParams.LayerVisible = Layer.Visible;
+			ShaderParams.LayerVisible = Layer.bVisible;
 			ShaderParams.LayerAlpha = Layer.HeightmapAlpha;
 
 			for (ALandscapeProxy* Landscape : AllLandscapes)
@@ -2696,7 +2658,7 @@ void ALandscape::RegenerateProceduralHeightmaps()
 
 			ShaderParams.ApplyLayerModifiers = false;
 
-			if (Layer.Visible)
+			if (Layer.bVisible)
 			{
 				// Draw each Combined RT into a Non Atlas RT format to be use as base for all brush rendering
 				if (Layer.Brushes.Num() > 0)
@@ -3598,7 +3560,7 @@ void ALandscape::RegenerateProceduralWeightmaps()
 			TMap<ULandscapeLayerInfoObject*, int32> LayerInfoObjects; // <LayerInfoObj, LayerIndex>
 
 			// Determine if some brush want to write to layer that we have currently no data on
-			if (ProceduralLayer.Visible)
+			if (ProceduralLayer.bVisible)
 			{
 				for (int32 LayerInfoSettingsIndex = 0; LayerInfoSettingsIndex < Info->Layers.Num(); ++LayerInfoSettingsIndex)
 				{
@@ -3684,7 +3646,7 @@ void ALandscape::RegenerateProceduralWeightmaps()
 					PrintProceduralDebugRT(OutputDebugName ? FString::Printf(TEXT("LS Weight: %s CurrentProcLayerWeightmapAllLayersResource -> Paint Layer RT %s"), *ProceduralLayer.Name.ToString(), *LandscapeScratchRT1->GetName()) : TEXT(""), LandscapeScratchRT1, 0, false);
 
 					PSShaderParams.ApplyLayerModifiers = true;
-					PSShaderParams.LayerVisible = ProceduralLayer.Visible;
+					PSShaderParams.LayerVisible = ProceduralLayer.bVisible;
 					PSShaderParams.LayerAlpha = LayerInfoObj == ALandscapeProxy::VisibilityLayer ? 1.0f : ProceduralLayer.WeightmapAlpha; // visibility can't be affected by weight
 
 					DrawWeightmapComponentsToRenderTarget(OutputDebugName ? FString::Printf(TEXT("LS Weight: %s Paint: %s += -> %s"), *ProceduralLayer.Name.ToString(), *LayerInfoObj->LayerName.ToString(), *LandscapeScratchRT1->GetName(), *LandscapeScratchRT2->GetName()) : TEXT(""),
@@ -3719,7 +3681,7 @@ void ALandscape::RegenerateProceduralWeightmaps()
 					CopyProceduralTexture(OutputDebugName ? FString::Printf(TEXT("Weight: %s PaintLayer: %s %s"), *ProceduralLayer.Name.ToString(), *LayerInfoObj->LayerName.ToString(), *LandscapeScratchRT3->GetName()) : TEXT(""), LandscapeScratchRT3->GameThread_GetRenderTargetResource(), TEXT("CombinedProcLayerWeightmap"), CombinedProcLayerWeightmapAllLayersResource, nullptr, FIntPoint(0, 0), 0, 0, 0, LayerIndex);
 
 					// Handle brush blending
-					if (ProceduralLayer.Visible)
+					if (ProceduralLayer.bVisible)
 					{
 						// Draw each brushes				
 						for (int32 i = 0; i < ProceduralLayer.WeightmapBrushOrderIndices.Num(); ++i)
@@ -4395,4 +4357,325 @@ void ALandscape::FinishDestroy()
 	Super::FinishDestroy();
 }
 
+#if WITH_EDITOR
+bool ALandscape::IsProceduralLayerNameUnique(const FName& InName) const
+{
+	return Algo::CountIf(ProceduralLayers, [InName](const FProceduralLayer& Layer) { return (Layer.Name == InName); }) == 0;
+}
+
+void ALandscape::SetProceduralLayerName(int32 InLayerIndex, const FName& InName)
+{
+	FProceduralLayer* Layer = GetProceduralLayer(InLayerIndex);
+	ULandscapeInfo* LandscapeInfo = GetLandscapeInfo();
+	if (!LandscapeInfo || !Layer || Layer->Name == InName)
+	{
+		return;
+	}
+
+	if (!IsProceduralLayerNameUnique(InName))
+	{
+		return;
+	}
+
+	// TODO: ProceduralLayersData should be changed to map GUID to Data
+	FName OldName = ProceduralLayers[InLayerIndex].Name;
+	ProceduralLayers[InLayerIndex].Name = InName;
+	LandscapeInfo->ForAllLandscapeProxies([OldName, InName](ALandscapeProxy* Proxy)
+	{
+		FProceduralLayerData* Value = Proxy->ProceduralLayersData.Find(OldName);
+		check(Value);
+		Proxy->ProceduralLayersData.Add(InName, *Value);
+		Proxy->ProceduralLayersData.Remove(OldName);
+	});
+}
+
+void ALandscape::SetProceduralLayerAlpha(int32 InLayerIndex, const float InAlpha, bool bInHeightmap)
+{
+	FProceduralLayer* Layer = GetProceduralLayer(InLayerIndex);
+	if (!Layer)
+	{
+		return;
+	}
+	float& LayerAlpha = bInHeightmap ? Layer->HeightmapAlpha : Layer->WeightmapAlpha;
+	if (LayerAlpha == InAlpha)
+	{
+		return;
+	}
+
+	LayerAlpha = InAlpha;
+	RequestProceduralContentUpdate(EProceduralContentUpdateFlag::All, true);
+}
+
+void ALandscape::SetProceduralLayerVisibility(int32 InLayerIndex, bool bInVisible)
+{
+	FProceduralLayer* Layer = GetProceduralLayer(InLayerIndex);
+	if (!Layer || Layer->bVisible == bInVisible)
+	{
+		return;
+	}
+
+	Layer->bVisible = bInVisible;
+	RequestProceduralContentUpdate(EProceduralContentUpdateFlag::All, true);
+}
+
+FProceduralLayer* ALandscape::GetProceduralLayer(int32 InLayerIndex)
+{
+	if (ProceduralLayers.IsValidIndex(InLayerIndex))
+	{
+		return &ProceduralLayers[InLayerIndex];
+	}
+	return nullptr;
+}
+
+const FProceduralLayer* ALandscape::GetProceduralLayer(int32 InLayerIndex) const
+{
+	if (ProceduralLayers.IsValidIndex(InLayerIndex))
+	{
+		return &ProceduralLayers[InLayerIndex];
+	}
+	return nullptr;
+}
+
+void ALandscape::DeleteProceduralLayer(int32 InLayerIndex)
+{
+	ensure(GetMutableDefault<UEditorExperimentalSettings>()->bProceduralLandscape);
+
+	ULandscapeInfo* LandscapeInfo = GetLandscapeInfo();
+	const FProceduralLayer* Layer = GetProceduralLayer(InLayerIndex);
+	if (!LandscapeInfo || !Layer || ProceduralLayers.Num() <= 1)
+	{
+		return;
+	}
+	
+	Modify();
+	FName LayerName = Layer->Name;
+
+	// Clean up Weightmap usage in LandscapeProxies
+	LandscapeInfo->ForAllLandscapeProxies([LayerName](ALandscapeProxy* Proxy)
+	{
+		const FProceduralLayerData* LayerData = Proxy->ProceduralLayersData.Find(LayerName);
+		if (LayerData)
+		{
+			for (ULandscapeComponent* Component : Proxy->LandscapeComponents)
+			{
+				const FWeightmapLayerData* WeightmapLayer = LayerData->WeightmapData.Find(Component);
+				if (WeightmapLayer)
+				{
+					for (const FWeightmapLayerAllocationInfo& Allocation : WeightmapLayer->WeightmapLayerAllocations)
+					{
+						UTexture2D* WeightmapTexture = WeightmapLayer->Weightmaps[Allocation.WeightmapTextureIndex];
+						ULandscapeWeightmapUsage** Usage = Proxy->WeightmapUsageMap.Find(WeightmapTexture);
+						if (Usage != nullptr && (*Usage) != nullptr)
+						{
+							(*Usage)->ChannelUsage[Allocation.WeightmapTextureChannel] = nullptr;
+							if ((*Usage)->FreeChannelCount() == 4)
+							{
+								Proxy->WeightmapUsageMap.Remove(WeightmapTexture);
+							}
+						}
+					}
+				}
+			}
+		}
+	});
+	
+	// Remove associated layer data of each landscape proxy
+	LandscapeInfo->ForAllLandscapeProxies([LayerName](ALandscapeProxy* Proxy)
+	{
+		Proxy->ProceduralLayersData.Remove(LayerName);
+	});
+
+	// Remove layer from list
+	ProceduralLayers.RemoveAt(InLayerIndex);
+
+	// Request Update
+	RequestProceduralContentUpdate(EProceduralContentUpdateFlag::All_Setup | All, true);
+}
+
+void ALandscape::ClearProceduralLayer(int32 InLayerIndex)
+{
+	const FProceduralLayer* Layer = GetProceduralLayer(InLayerIndex);
+	if (Layer)
+	{
+		ClearProceduralLayer(Layer->Guid);
+	}
+}
+
+void ALandscape::ClearProceduralLayer(const FGuid& InLayerGuid)
+{
+	ensure(GetMutableDefault<UEditorExperimentalSettings>()->bProceduralLandscape);
+
+	FProceduralLayer* Layer = ProceduralLayers.FindByPredicate([InLayerGuid](const FProceduralLayer& Other) { return Other.Guid == InLayerGuid; });
+	ULandscapeInfo* LandscapeInfo = GetLandscapeInfo();
+	if (!LandscapeInfo || !Layer)
+	{
+		return;
+	}
+	
+	Modify();
+	FScopedSetLandscapeCurrentEditingProceduralLayer Scope(this, Layer ? Layer->Guid : FGuid(), [=] { RequestProceduralContentUpdate(EProceduralContentUpdateFlag::All, true); });
+
+	TArray<uint16> NewData;
+	NewData.AddZeroed(FMath::Square(ComponentSizeQuads + 1));
+	uint16 ZeroValue = LandscapeDataAccess::GetTexHeight(0.f);
+	for (uint16& NewDataValue : NewData)
+	{
+		NewDataValue = ZeroValue;
+	}
+
+	FLandscapeEditDataInterface LandscapeEdit(LandscapeInfo);
+	LandscapeInfo->ForAllLandscapeProxies([&](ALandscapeProxy* Proxy)
+	{
+		for (ULandscapeComponent* Component : Proxy->LandscapeComponents)
+		{
+			int32 MinX = MAX_int32;
+			int32 MinY = MAX_int32;
+			int32 MaxX = MIN_int32;
+			int32 MaxY = MIN_int32;
+			Component->GetComponentExtent(MinX, MinY, MaxX, MaxY);
+			check(ComponentSizeQuads == (MaxX - MinX));
+			check(ComponentSizeQuads == (MaxY - MinY));
+
+			TArray<uint16> OldData;
+			OldData.AddZeroed((1 + MaxY - MinY) * (1 + MaxX - MinX));
+
+			LandscapeEdit.GetHeightData(MinX, MinY, MaxX, MaxY, OldData.GetData(), 0);
+			if (FMemory::Memcmp(OldData.GetData(), NewData.GetData(), NewData.Num() * NewData.GetTypeSize()) != 0)
+			{
+				LandscapeEdit.SetHeightData(MinX, MinY, MaxX, MaxY, NewData.GetData(), 0, true);
+			}
+
+			// Clear weight maps
+			for (FLandscapeInfoLayerSettings& LayerSettings : LandscapeInfo->Layers)
+			{
+				Component->DeleteLayer(LayerSettings.LayerInfoObj, LandscapeEdit);
+			}
+		}
+	});
+}
+
+void ALandscape::ShowOnlySelectedProceduralLayer(int32 InLayerIndex)
+{
+	const FProceduralLayer* VisibleLayer = GetProceduralLayer(InLayerIndex);
+	if (VisibleLayer)
+	{
+		for (FProceduralLayer& Layer : ProceduralLayers)
+		{
+			Layer.bVisible = (&Layer == VisibleLayer);
+		}
+		RequestProceduralContentUpdate(EProceduralContentUpdateFlag::All, true);
+	}
+}
+
+void ALandscape::ShowAllProceduralLayers()
+{
+	if (ProceduralLayers.Num() > 0)
+	{
+		for (FProceduralLayer& Layer : ProceduralLayers)
+		{
+			Layer.bVisible = true;
+		}
+		RequestProceduralContentUpdate(EProceduralContentUpdateFlag::All, true);
+	}
+}
+
+FScopedSetLandscapeCurrentEditingProceduralLayer::FScopedSetLandscapeCurrentEditingProceduralLayer(ALandscape* InLandscape, const FGuid& InProceduralLayer, TFunction<void()> InCompletionCallback)
+	: Landscape(InLandscape)
+	, ProceduralLayer(InProceduralLayer)
+	, CompletionCallback(MoveTemp(InCompletionCallback))
+{
+	if (GetMutableDefault<UEditorExperimentalSettings>()->bProceduralLandscape && Landscape.IsValid() && ProceduralLayer.IsValid())
+	{
+		Landscape->SetCurrentEditingProceduralLayer(ProceduralLayer);
+	}
+}
+
+FScopedSetLandscapeCurrentEditingProceduralLayer::~FScopedSetLandscapeCurrentEditingProceduralLayer()
+{
+	if (GetMutableDefault<UEditorExperimentalSettings>()->bProceduralLandscape && Landscape.IsValid() && ProceduralLayer.IsValid())
+	{
+		Landscape->SetCurrentEditingProceduralLayer();
+		if (CompletionCallback)
+		{
+			CompletionCallback();
+		}
+	}
+}
+
+void ALandscape::SetCurrentEditingProceduralLayer(FGuid InLayerGuid)
+{
+	ensure(GetMutableDefault<UEditorExperimentalSettings>()->bProceduralLandscape);
+
+	ULandscapeInfo* LandscapeInfo = GetLandscapeInfo();
+	if (!LandscapeInfo)
+	{
+		return;
+	}
+
+	LandscapeInfo->ForAllLandscapeProxies([InLayerGuid, this](ALandscapeProxy* Proxy)
+	{
+		// TODO: ProceduralLayersData should be changed to map GUID to Data
+		FProceduralLayer* Layer = ProceduralLayers.FindByPredicate([InLayerGuid](const FProceduralLayer& Other) { return Other.Guid == InLayerGuid; });
+		FProceduralLayerData* CurrentLayerData = Layer ? Proxy->ProceduralLayersData.Find(Layer->Name) : nullptr;
+
+		for (ULandscapeComponent* Component : Proxy->LandscapeComponents)
+		{
+			// Update Current Heightmap
+			UTexture2D** LayerHeightmap = CurrentLayerData ? CurrentLayerData->Heightmaps.Find(Component->GetHeightmap()) : nullptr;
+			Component->SetCurrentEditingHeightmap(LayerHeightmap ? *LayerHeightmap : nullptr);
+
+			// Update Current Weightmaps
+			FWeightmapLayerData* WeightmapData = CurrentLayerData ? CurrentLayerData->WeightmapData.Find(Component) : nullptr;
+			Component->SetCurrentEditingWeightmapLayerAllocations(WeightmapData ? &WeightmapData->WeightmapLayerAllocations : nullptr);
+			Component->SetCurrentEditingWeightmaps(WeightmapData ? &WeightmapData->Weightmaps : nullptr);
+			Component->SetCurrentProceduralLayerName(WeightmapData && Layer ? Layer->Name : NAME_None);
+			Component->SetCurrentEditingWeightmapTexturesUsage(WeightmapData ? &WeightmapData->WeightmapTextureUsages : nullptr);
+
+			Component->MarkRenderStateDirty();
+		}
+	});
+}
+
+void ALandscape::CreateProceduralLayer(FName InName, bool bInUpdateProceduralContent)
+{
+	ULandscapeInfo* LandscapeInfo = GetLandscapeInfo();
+	if (!LandscapeInfo || !GetMutableDefault<UEditorExperimentalSettings>()->bProceduralLandscape)
+	{
+		return;
+	}
+
+	Modify();
+	FProceduralLayer NewLayer;
+	NewLayer.Name = GenerateUniqueProceduralLayerName(InName);
+	ProceduralLayers.Add(NewLayer);
+
+	// Create associated layer data in each landscape proxy
+	LandscapeInfo->ForAllLandscapeProxies([NewLayer](ALandscapeProxy* Proxy)
+	{
+		Proxy->ProceduralLayersData.Add(NewLayer.Name, FProceduralLayerData());
+	});
+
+	if (bInUpdateProceduralContent)
+	{
+		// Request Update
+		RequestProceduralContentUpdate(EProceduralContentUpdateFlag::All_Setup);
+		RegenerateProceduralContent();
+	}
+}
+
+FName ALandscape::GenerateUniqueProceduralLayerName(FName InName) const
+{
+	FString BaseName = InName == NAME_None ? "Layer" : InName.ToString();
+	FName NewName;
+	int32 LayerIndex = 0;
+	do
+	{
+		++LayerIndex;
+		NewName = FName(*FString::Printf(TEXT("%s%d"), *BaseName, LayerIndex));
+	} while (ProceduralLayers.ContainsByPredicate([NewName](const FProceduralLayer& Layer) { return Layer.Name == NewName; }));
+
+	return NewName;
+}
+#endif // WITH_EDITOR
+
 #undef LOCTEXT_NAMESPACE
