Changed the name from Landscape Procedural to Landscape Layers
Moved some function from LandscapeEdMode to LandscapeEditLayers

diff --git a/Engine/Shaders/Private/LandscapeLayersCS.usf b/Engine/Shaders/Private/LandscapeLayersCS.usf
new file mode 100644
--- /dev/null
+++ b/Engine/Shaders/Private/LandscapeLayersCS.usf
@@ -0,0 +1,90 @@
+// Copyright 1998-2019 Epic Games, Inc. All Rights Reserved.
+
+#include "Common.ush"
+
+uint ComponentSize;
+
+struct FLandscapeLayerWeightmapExtractMaterialLayersComponentData
+{
+	int2 ComponentVertexPosition;		// Section Base of the component converted to vertex
+	uint DestinationPaintLayerIndex;	// correspond to which layer info object index the data should be stored in the texture 2d array
+	uint WeightmapChannelToProcess;		// correspond to which RGBA channel to process
+	int2 AtlasTexturePositionOutput;	// This represent the location we will write layer information
+};
+
+StructuredBuffer<FLandscapeLayerWeightmapExtractMaterialLayersComponentData> InExtractLayersComponentsData;
+Texture2D<float4> InComponentWeightMaps;
+RWTexture2DArray<float> OutAtlasPaintLayers;
+
+[numthreads(THREADGROUP_SIZEX, THREADGROUP_SIZEY, 1)]
+void ComputeWeightmapPerPaintLayer(uint3 DispatchThreadId : SV_DispatchThreadID)
+{
+	if (DispatchThreadId.x < ComponentSize && DispatchThreadId.y < ComponentSize)
+	{
+		FLandscapeLayerWeightmapExtractMaterialLayersComponentData ComponentData = InExtractLayersComponentsData[DispatchThreadId.z];
+		int2 CurrentTexel = DispatchThreadId.xy + ComponentData.ComponentVertexPosition;
+
+		float4 LayerBlend = InComponentWeightMaps.Load(int3(CurrentTexel, 0));
+		OutAtlasPaintLayers[int3(DispatchThreadId.xy + ComponentData.AtlasTexturePositionOutput, ComponentData.DestinationPaintLayerIndex)] = LayerBlend[ComponentData.WeightmapChannelToProcess];
+	}
+}
+
+struct FLandscapeLayerWeightmapPackMaterialLayersComponentData
+{
+	int4 ComponenVertexPositionX;	// Section Base of the component converted to vertex for each rgba value (as a texture can contain data from different component)
+	int4 ComponenVertexPositionY;	// Section Base of the component converted to vertex for each rgba value (as a texture can contain data from different component)
+	int4 SourcePaintLayerIndex;		// correspond to which layer info object index the data should be stored in the texture 2d array
+	int4 WeightmapChannelToProcess;	// correspond to which RGBA channel to process
+};
+
+StructuredBuffer<FLandscapeLayerWeightmapPackMaterialLayersComponentData> InPackLayersComponentsData;
+Buffer<float> InWeightmapWeightBlendMode;
+Buffer<float2> InWeightmapTextureOutputOffset;
+Texture2DArray<float> InAtlasPaintLayers;
+RWTexture2D<float4> OutComponentWeightMaps;
+
+float ComputeWeightmapChannel(int ChannelIndexToProcess, uint2 TexelPosition, uint SourcePaintLayerIndex)
+{
+	float Channel = 0;
+
+	if (ChannelIndexToProcess != -1)
+	{
+		Channel = InAtlasPaintLayers.Load(int4(TexelPosition, SourcePaintLayerIndex, 0));
+		float AllLayerWeightSum = 0;
+
+		float Unused = 0.0;
+		float ArraySize = 0.0;
+		InAtlasPaintLayers.GetDimensions(Unused, Unused, ArraySize);
+
+		for (int i = 0; i < ArraySize; ++i)
+		{
+			if (InWeightmapWeightBlendMode[i] == 1.0) // Only include the one that are weight blended
+			{
+				AllLayerWeightSum += InAtlasPaintLayers.Load(int4(TexelPosition, i, 0));
+			}
+		}
+
+		Channel = AllLayerWeightSum == 0 ? Channel : clamp(Channel / AllLayerWeightSum, 0.0, 1.0);
+	}
+
+	return Channel;
+}
+
+[numthreads(THREADGROUP_SIZEX, THREADGROUP_SIZEY, 1)]
+void PackPaintLayerToWeightmap(uint3 DispatchThreadId : SV_DispatchThreadID)
+{
+	if (DispatchThreadId.x < ComponentSize && DispatchThreadId.y < ComponentSize)
+	{
+		FLandscapeLayerWeightmapPackMaterialLayersComponentData ComponentData = InPackLayersComponentsData[DispatchThreadId.z];
+
+		float4 Channels;
+
+		for (int i = 0; i < 4; ++i)
+		{
+			Channels[i] = ComputeWeightmapChannel(ComponentData.WeightmapChannelToProcess[i], DispatchThreadId.xy + int2(ComponentData.ComponenVertexPositionX[i], ComponentData.ComponenVertexPositionY[i]), ComponentData.SourcePaintLayerIndex[i]);
+		}
+
+		float2 TextureOutputOffset = InWeightmapTextureOutputOffset[DispatchThreadId.z];
+		OutComponentWeightMaps[DispatchThreadId.xy + TextureOutputOffset] = Channels;
+	}
+}
\ No newline at end of file
diff --git a/Engine/Shaders/Private/LandscapeLayersPS.usf b/Engine/Shaders/Private/LandscapeLayersPS.usf
new file mode 100644
--- /dev/null
+++ b/Engine/Shaders/Private/LandscapeLayersPS.usf
@@ -0,0 +1,405 @@
+// Copyright 1998-2019 Epic Games, Inc. All Rights Reserved.
+
+#include "Common.ush"
+
+Texture2D ReadTexture1;
+Texture2D ReadTexture2;					// This one is optional, it's valid when OutputConfig.z == 1
+SamplerState ReadTexture1Sampler;		
+SamplerState ReadTexture2Sampler;
+float2 LayerInfo;					// x == weight, y == visibility
+float4 OutputConfig;					// x == ApplyLayerModifiers, y == OutputAsDelta or Output weightmap as substractive/additive, z == Use ReadTexture2 and out Delta with ReadTexture1, w == Output Normals or Output weightmap as normalized weight blend
+float2 TextureSize;					// x == read texture 1 width, y == read texture 1 height
+float3 LandscapeGridScale;				// x == LS Actor DrawScale.X, y == LS Actor DrawScale.y, z == LS Actor DrawScale.z / 128.0f (ZSCALE)
+float CurrentMipComponentVertexCount;
+float2 CurrentMipTextureSize;
+float2 ParentMipTextureSize;
+
+float ExtractHeight(float2 InPackedHeight)
+{
+	return float(((int)round(InPackedHeight.r * 255.0) << 8) | (int)round(InPackedHeight.g * 255.0));
+}
+
+float2 PackHeight(float InExtractedHeight)
+{
+	InExtractedHeight = clamp(InExtractedHeight, 0.0f, 65536.0f);
+	int iHeight = (int)InExtractedHeight;
+	return float2((float)((iHeight - (iHeight & 255)) >> 8) / 255.0, (float)(iHeight & 255) / 255.0);
+}
+
+float3 SampleTextureYNormal(float2 InTextureCoordinates, float2 InTexelSize, bool InUp)
+{
+	float2 SamplingUV = InTextureCoordinates;
+	SamplingUV.y = InUp ? SamplingUV.y - InTexelSize.y : SamplingUV.y + InTexelSize.y;
+	float2 PosSample = ReadTexture1.SampleLevel(ReadTexture1Sampler, SamplingUV, 0).rg;
+	float PosHeight = ExtractHeight(PosSample) - 32768.0;
+	return float3(SamplingUV * TextureSize * LandscapeGridScale.xy, PosHeight * LandscapeGridScale.z * TERRAIN_ZSCALE);
+}
+
+float3 SampleTextureXNormal(float2 InTextureCoordinates, float2 InTexelSize, bool InLeft)
+{
+	float2 SamplingUV = InTextureCoordinates;
+	SamplingUV.x = InLeft ? SamplingUV.x - InTexelSize.x : SamplingUV.x + InTexelSize.x;
+	float2 PosSample = ReadTexture1.SampleLevel(ReadTexture1Sampler, SamplingUV, 0).rg;
+	float PosHeight = ExtractHeight(PosSample) - 32768.0;
+	return float3(SamplingUV * TextureSize * LandscapeGridScale.xy, PosHeight * LandscapeGridScale.z * TERRAIN_ZSCALE);
+}
+
+void PSHeightmapMain(float2 InTextureCoordinates : TEXCOORD0, out float4 OutColor : SV_Target0)
+{
+	float4 SourceColor = ReadTexture1.Sample(ReadTexture1Sampler, InTextureCoordinates);
+	float Height = ExtractHeight(SourceColor.rg);
+
+	// Perform calculation 0 based
+	Height -= 32768.0f;
+
+	// Output as Layer, so apply Layer info
+	if (OutputConfig.x == 1.0)
+	{
+		Height = round(Height * LayerInfo.x * LayerInfo.y);
+	}
+
+	// Output using 2nd heightmap
+	if (OutputConfig.z == 1.0)
+	{
+		float2 BaseColor = ReadTexture2.Sample(ReadTexture2Sampler, InTextureCoordinates).rg;
+		float BaseHeight = ExtractHeight(BaseColor.rg);
+		Height = OutputConfig.y == 1 ? BaseHeight + (Height - BaseHeight) : BaseHeight + Height;
+	}
+	else
+	{
+		// Output as Delta to Source
+		if (OutputConfig.y == 0.0)
+		{
+			Height += 32768.0f;
+		}
+	}
+
+	float2 PackedHeight = PackHeight(Height);
+
+	OutColor.r = PackedHeight.x;
+	OutColor.g = PackedHeight.y;
+	OutColor.b = SourceColor.b;
+	OutColor.a = SourceColor.a;
+
+	// Output normals
+	if (OutputConfig.w == 1.0)
+	{
+		float2 TexelSize = 1.0 / TextureSize;
+
+		bool IsMinBorderTexelX = InTextureCoordinates.x <= TexelSize.x;
+		bool IsMinBorderTexelY = InTextureCoordinates.y <= TexelSize.y;
+		bool IsMaxBorderTexelX = InTextureCoordinates.x >= TexelSize.x * min(CurrentMipComponentVertexCount - 1, 1.0);
+		bool IsMaxBorderTexelY = InTextureCoordinates.y >= TexelSize.y * min(CurrentMipComponentVertexCount - 1, 1.0);
+
+		float3 CurrentPos = float3(InTextureCoordinates * TextureSize * LandscapeGridScale.xy, (Height - 32768.0f) * LandscapeGridScale.z * TERRAIN_ZSCALE);
+		float3 FinalNormal = 0.0;
+
+		if (IsMinBorderTexelX) // left border
+		{
+			if (IsMinBorderTexelY) // top border
+			{
+				float3 PosDown = SampleTextureYNormal(InTextureCoordinates, TexelSize, false);
+				float3 PosRight = SampleTextureXNormal(InTextureCoordinates, TexelSize, false);
+				FinalNormal = normalize(cross(PosRight - CurrentPos, PosDown - CurrentPos));
+			}
+			else if (IsMaxBorderTexelY) // bottom border
+			{
+				float3 PosUp = SampleTextureYNormal(InTextureCoordinates, TexelSize, true);
+				float3 PosRight = SampleTextureXNormal(InTextureCoordinates, TexelSize, false);
+				FinalNormal = normalize(cross(PosUp - CurrentPos, PosRight - CurrentPos));
+			}
+			else
+			{
+				float3 PosUp = SampleTextureYNormal(InTextureCoordinates, TexelSize, true);
+				float3 PosRight = SampleTextureXNormal(InTextureCoordinates, TexelSize, false);
+				float3 PosDown = SampleTextureYNormal(InTextureCoordinates, TexelSize, false);
+
+				float3 VectorUpRight = cross(PosUp - CurrentPos, PosRight - CurrentPos);
+				float3 VectorRightDown = cross(PosRight - CurrentPos, PosDown - CurrentPos);
+				FinalNormal = normalize(VectorUpRight + VectorRightDown);
+			}
+		}
+		else if (IsMaxBorderTexelX) // right border
+		{
+			if (IsMinBorderTexelY) // top border
+			{
+				float3 PosDown = SampleTextureYNormal(InTextureCoordinates, TexelSize, false);
+				float3 PosLeft = SampleTextureXNormal(InTextureCoordinates, TexelSize, true);
+				FinalNormal = normalize(cross(PosDown - CurrentPos, PosLeft - CurrentPos));
+			}
+			else if (IsMaxBorderTexelY) // bottom border
+			{
+				float3 PosUp = SampleTextureYNormal(InTextureCoordinates, TexelSize, true);
+				float3 PosLeft = SampleTextureXNormal(InTextureCoordinates, TexelSize, true);
+				FinalNormal = normalize(cross(PosLeft - CurrentPos, PosUp - CurrentPos));
+			}
+			else
+			{
+				float3 PosUp = SampleTextureYNormal(InTextureCoordinates, TexelSize, true);
+				float3 PosLeft = SampleTextureXNormal(InTextureCoordinates, TexelSize, true);
+				float3 PosDown = SampleTextureYNormal(InTextureCoordinates, TexelSize, false);
+
+				float3 VectorLeftUp = cross(PosLeft - CurrentPos, PosUp - CurrentPos);
+				float3 VectorDownLeft = cross(PosDown - CurrentPos, PosLeft - CurrentPos);
+				FinalNormal = normalize(VectorDownLeft + VectorLeftUp);
+			}
+		}
+		else // center border
+		{
+			if (IsMinBorderTexelY) // top border
+			{
+				float3 PosRight = SampleTextureXNormal(InTextureCoordinates, TexelSize, false);
+				float3 PosLeft = SampleTextureXNormal(InTextureCoordinates, TexelSize, true);
+				float3 PosDown = SampleTextureYNormal(InTextureCoordinates, TexelSize, false);
+
+				float3 VectorRightDown = cross(PosRight - CurrentPos, PosDown - CurrentPos);
+				float3 VectorDownLeft = cross(PosDown - CurrentPos, PosLeft - CurrentPos);
+				FinalNormal = normalize(VectorRightDown + VectorDownLeft);
+			}
+			else if (IsMaxBorderTexelY) // bottom border
+			{
+				float3 PosRight = SampleTextureXNormal(InTextureCoordinates, TexelSize, false);
+				float3 PosLeft = SampleTextureXNormal(InTextureCoordinates, TexelSize, true);
+				float3 PosUp = SampleTextureYNormal(InTextureCoordinates, TexelSize, true);
+
+				float3 VectorUpRight = cross(PosUp - CurrentPos, PosRight - CurrentPos);
+				float3 VectorLeftUp = cross(PosLeft - CurrentPos, PosUp - CurrentPos);
+				FinalNormal = normalize(VectorUpRight + VectorLeftUp);
+			}
+			else
+			{
+				float3 PosRight = SampleTextureXNormal(InTextureCoordinates, TexelSize, false);
+				float3 PosLeft = SampleTextureXNormal(InTextureCoordinates, TexelSize, true);
+				float3 PosUp = SampleTextureYNormal(InTextureCoordinates, TexelSize, true);
+				float3 PosDown = SampleTextureYNormal(InTextureCoordinates, TexelSize, false);
+
+				float3 VectorRightDown = cross(PosRight - CurrentPos, PosDown - CurrentPos);
+				float3 VectorDownLeft = cross(PosDown - CurrentPos, PosLeft - CurrentPos);
+
+				float3 VectorUpRight = cross(PosUp - CurrentPos, PosRight - CurrentPos);
+				float3 VectorLeftUp = cross(PosLeft - CurrentPos, PosUp - CurrentPos);
+
+				FinalNormal = normalize(VectorUpRight + VectorLeftUp + VectorRightDown + VectorDownLeft);
+			}
+		}		
+
+		// Scale back to be 0-1 normal
+		OutColor.b = (FinalNormal.x + 1.0) * 0.5;
+		OutColor.a = (FinalNormal.y + 1.0) * 0.5;
+	}
+}
+
+void PSHeightmapMainMips(float2 InTextureCoordinates : TEXCOORD0, out float4 OutColor : SV_Target0)
+{
+	bool IsMinBorderTexelX = false;
+	bool IsMinBorderTexelY = false;
+	bool IsMaxBorderTexelX = false;
+	bool IsMaxBorderTexelY = false;
+
+	// Special case of 1 texel size component
+	if (CurrentMipComponentVertexCount == 1)
+	{		
+		if (CurrentMipTextureSize.x >= CurrentMipTextureSize.y) // x size 1 texel
+		{
+			IsMinBorderTexelY = true;
+			IsMinBorderTexelX = (InTextureCoordinates.x * CurrentMipTextureSize.x) <= 1.0;
+			IsMaxBorderTexelX = (InTextureCoordinates.x * CurrentMipTextureSize.x) >= CurrentMipTextureSize.x - 1.0;
+		}
+		else
+		{
+			IsMinBorderTexelX = true;
+			IsMinBorderTexelY = (InTextureCoordinates.y * CurrentMipTextureSize.y) <= 1.0;
+			IsMaxBorderTexelY = (InTextureCoordinates.y * CurrentMipTextureSize.y) >= CurrentMipTextureSize.y - 1.0;
+		}
+	}
+	else
+	{
+		float2 CurrentMipQuadTexelSize = 1.0 / CurrentMipComponentVertexCount;
+		float2 LandscapeQuadUV = frac(InTextureCoordinates * CurrentMipTextureSize / CurrentMipComponentVertexCount);
+
+		IsMinBorderTexelX = LandscapeQuadUV.x <= CurrentMipQuadTexelSize.x;
+		IsMinBorderTexelY = LandscapeQuadUV.y <= CurrentMipQuadTexelSize.y;
+		IsMaxBorderTexelX = LandscapeQuadUV.x >= (CurrentMipQuadTexelSize.x * max(CurrentMipComponentVertexCount - 1, 1.0));
+		IsMaxBorderTexelY = LandscapeQuadUV.y >= (CurrentMipQuadTexelSize.y * max(CurrentMipComponentVertexCount - 1, 1.0));
+	}
+
+	float2 ParentMipTexelSize = 1.0 / ParentMipTextureSize;
+	float2 SourceUV = InTextureCoordinates - (ParentMipTexelSize * 0.5);
+	float HeightResult = 0;
+	float2 NormalResult = 0;
+	float Alpha = 0.5;
+	OutColor = 0; // Default to 0
+
+	if (IsMinBorderTexelX) // on left border
+	{
+		float4 SourceSample = ReadTexture1.SampleLevel(ReadTexture1Sampler, SourceUV, 0);
+		float4 DownSample = ReadTexture1.SampleLevel(ReadTexture1Sampler, SourceUV + float2(0.0, ParentMipTexelSize.y), 0);
+		float SourceHeight = ExtractHeight(SourceSample.rg);
+		float DownHeight = ExtractHeight(DownSample.rg);
+
+		if (IsMinBorderTexelY) // on top border
+		{
+			HeightResult = SourceHeight;
+			NormalResult = float2(SourceSample.ba);
+		}
+		else if (IsMaxBorderTexelY) // on bottom border
+		{
+			HeightResult = DownHeight;
+			NormalResult = float2(DownSample.ba);
+		}
+		else
+		{
+			HeightResult = clamp(round(lerp(SourceHeight, DownHeight, Alpha)), 0.0, 65536.0);
+			NormalResult.x = lerp(SourceSample.b, DownSample.b, Alpha);
+			NormalResult.y = lerp(SourceSample.a, DownSample.a, Alpha);
+		}
+	}
+	else if (IsMaxBorderTexelX) // on right border
+	{
+		float4 RightSample = ReadTexture1.SampleLevel(ReadTexture1Sampler, SourceUV + float2(ParentMipTexelSize.x, 0.0), 0);
+		float4 DownRightSample = ReadTexture1.SampleLevel(ReadTexture1Sampler, SourceUV + ParentMipTexelSize, 0);
+		float RightHeight = ExtractHeight(RightSample.rg);
+		float DownRightHeight = ExtractHeight(DownRightSample.rg);
+
+		if (IsMinBorderTexelY) // on top border
+		{
+			HeightResult = RightHeight;
+			NormalResult = float2(RightSample.ba);
+		}		
+		else if (IsMaxBorderTexelY) // on bottom border
+		{
+			HeightResult = DownRightHeight;
+			NormalResult = float2(DownRightSample.ba);
+		}
+		else
+		{
+			HeightResult = clamp(round(lerp(RightHeight, DownRightHeight, Alpha)), 0.0, 65536.0);;
+			NormalResult.x = lerp(RightSample.b, DownRightSample.b, Alpha);
+			NormalResult.y = lerp(RightSample.a, DownRightSample.a, Alpha);
+		}
+	}
+	else // center texel
+	{
+		float4 SourceSample = ReadTexture1.SampleLevel(ReadTexture1Sampler, SourceUV, 0);
+		float4 DownSample = ReadTexture1.SampleLevel(ReadTexture1Sampler, SourceUV + float2(0.0, ParentMipTexelSize.y), 0);
+		float4 RightSample = ReadTexture1.SampleLevel(ReadTexture1Sampler, SourceUV + float2(ParentMipTexelSize.x, 0.0), 0);
+		float4 DownRightSample = ReadTexture1.SampleLevel(ReadTexture1Sampler, SourceUV + ParentMipTexelSize, 0);
+
+		float SourceHeight = ExtractHeight(SourceSample.rg);
+		float DownHeight = ExtractHeight(DownSample.rg);
+		float RightHeight = ExtractHeight(RightSample.rg);
+		float DownRightHeight = ExtractHeight(DownRightSample.rg);
+
+		if (IsMinBorderTexelY) // on top border
+		{
+			HeightResult = clamp(round(lerp(SourceHeight, RightHeight, Alpha)), 0.0, 65536.0);
+			NormalResult.x = lerp(SourceSample.b, RightSample.b, Alpha);
+			NormalResult.y = lerp(SourceSample.a, RightSample.a, Alpha);
+		}
+		else if (IsMaxBorderTexelY) // on bottom border
+		{
+			HeightResult = clamp(round(lerp(DownHeight, DownRightHeight, Alpha)), 0.0, 65536.0);
+			NormalResult.x = lerp(DownSample.b, DownRightSample.b, Alpha);
+			NormalResult.y = lerp(DownSample.a, DownRightSample.a, Alpha);
+		}
+		else
+		{
+			float SourceToRightHeight = lerp(SourceHeight, RightHeight, Alpha);
+			float DownToDownRightHeight = lerp(DownHeight, DownRightHeight, Alpha);
+
+			HeightResult = clamp(round(lerp(SourceToRightHeight, DownToDownRightHeight, Alpha)), 0.0, 65536.0);
+			NormalResult.x = lerp(lerp(SourceSample.b, RightSample.b, Alpha), lerp(DownSample.b, DownRightSample.b, Alpha), Alpha);
+			NormalResult.y = lerp(lerp(SourceSample.a, RightSample.a, Alpha), lerp(DownSample.a, DownRightSample.a, Alpha), Alpha);
+		}
+	}	
+
+	float2 PackedHeightResult = PackHeight(HeightResult);
+	OutColor = float4(PackedHeightResult.xy, NormalResult.xy);
+}
+
+float4 WeightmapLayerWeightBlend; // each xyzw == 1 texture channel
+
+void PSWeightmapMain(float2 InTextureCoordinates : TEXCOORD0, out float4 OutColor : SV_Target0)
+{
+	float4 SourceColor = ReadTexture1.Sample(ReadTexture1Sampler, InTextureCoordinates);
+
+	// Output as Layer, so apply Layer info
+	SourceColor = OutputConfig.x == 1.0 ? SourceColor * LayerInfo.x * LayerInfo.y : SourceColor;
+
+	// Output using 2nd heightmap
+	if (OutputConfig.z == 1.0)
+	{
+		float4 BaseColor = ReadTexture2.Sample(ReadTexture2Sampler, InTextureCoordinates);
+
+		// Apply addtive/substractive logic
+		SourceColor = OutputConfig.y == 0.0f ? clamp(BaseColor + SourceColor, 0.0, 1.0) : clamp(BaseColor - SourceColor, 0.0, 1.0);
+	}	
+
+	OutColor = SourceColor;
+}
+
+void PSWeightmapMainMips(float2 InTextureCoordinates : TEXCOORD0, out float4 OutColor : SV_Target0)
+{
+	bool IsMinBorderTexelX = false;
+	bool IsMinBorderTexelY = false;
+	bool IsMaxBorderTexelX = false;
+	bool IsMaxBorderTexelY = false;
+
+	// Special case of 1 texel size component
+	if (CurrentMipComponentVertexCount == 1)
+	{		
+		if (CurrentMipTextureSize.x >= CurrentMipTextureSize.y) // x size 1 texel
+		{
+			IsMinBorderTexelY = true;
+			IsMinBorderTexelX = (InTextureCoordinates.x * CurrentMipTextureSize.x) <= 1.0;
+			IsMaxBorderTexelX = (InTextureCoordinates.x * CurrentMipTextureSize.x) >= CurrentMipTextureSize.x - 1.0;
+		}
+		else
+		{
+			IsMinBorderTexelX = true;
+			IsMinBorderTexelY = (InTextureCoordinates.y * CurrentMipTextureSize.y) <= 1.0;
+			IsMaxBorderTexelY = (InTextureCoordinates.y * CurrentMipTextureSize.y) >= CurrentMipTextureSize.y - 1.0;
+		}
+	}
+	else
+	{
+		float2 CurrentMipQuadTexelSize = 1.0 / CurrentMipComponentVertexCount;
+		float2 LandscapeQuadUV = frac(InTextureCoordinates * CurrentMipTextureSize / CurrentMipComponentVertexCount);
+
+		IsMinBorderTexelX = LandscapeQuadUV.x <= CurrentMipQuadTexelSize.x;
+		IsMinBorderTexelY = LandscapeQuadUV.y <= CurrentMipQuadTexelSize.y;
+		IsMaxBorderTexelX = LandscapeQuadUV.x >= (CurrentMipQuadTexelSize.x * max(CurrentMipComponentVertexCount - 1, 1.0));
+		IsMaxBorderTexelY = LandscapeQuadUV.y >= (CurrentMipQuadTexelSize.y * max(CurrentMipComponentVertexCount - 1, 1.0));
+	}
+
+	float2 ParentMipTexelSize = 1.0 / ParentMipTextureSize;
+	float2 SourceUV = InTextureCoordinates - (ParentMipTexelSize * 0.5);
+	float2 HeightResult = 0;
+	float2 NormalResult = 0;
+	float Alpha = 0.5;
+	OutColor = 0; // Default to 0
+
+	if (IsMinBorderTexelX) // on left border
+	{
+		float4 SourceSample = ReadTexture1.SampleLevel(ReadTexture1Sampler, SourceUV, 0);
+		float4 DownSample = ReadTexture1.SampleLevel(ReadTexture1Sampler, SourceUV + float2(0.0, ParentMipTexelSize.y), 0);
+
+		OutColor = IsMinBorderTexelY ? SourceSample : (IsMaxBorderTexelY ? DownSample : lerp(SourceSample, DownSample, Alpha));		
+	}
+	else if (IsMaxBorderTexelX) // on right border
+	{
+		float4 RightSample = ReadTexture1.SampleLevel(ReadTexture1Sampler, SourceUV + float2(ParentMipTexelSize.x, 0.0), 0);
+		float4 DownRightSample = ReadTexture1.SampleLevel(ReadTexture1Sampler, SourceUV + ParentMipTexelSize, 0);
+
+		OutColor = IsMinBorderTexelY ? RightSample : (IsMaxBorderTexelY ? DownRightSample : lerp(RightSample, DownRightSample, Alpha));
+	}
+	else // center texel
+	{
+		float4 SourceSample = ReadTexture1.SampleLevel(ReadTexture1Sampler, SourceUV, 0);
+		float4 DownSample = ReadTexture1.SampleLevel(ReadTexture1Sampler, SourceUV + float2(0.0, ParentMipTexelSize.y), 0);
+		float4 RightSample = ReadTexture1.SampleLevel(ReadTexture1Sampler, SourceUV + float2(ParentMipTexelSize.x, 0.0), 0);
+		float4 DownRightSample = ReadTexture1.SampleLevel(ReadTexture1Sampler, SourceUV + ParentMipTexelSize, 0);
+
+		OutColor = IsMinBorderTexelY ? lerp(SourceSample, RightSample, Alpha) : (IsMaxBorderTexelY ? lerp(DownSample, DownRightSample, Alpha) : lerp(lerp(SourceSample, RightSample, Alpha), lerp(DownSample, DownRightSample, Alpha), Alpha));
+	}
+}
\ No newline at end of file
diff --git a/Engine/Shaders/Private/LandscapeLayersVS.usf b/Engine/Shaders/Private/LandscapeLayersVS.usf
new file mode 100644
--- /dev/null
+++ b/Engine/Shaders/Private/LandscapeLayersVS.usf
@@ -0,0 +1,15 @@
+// Copyright 1998-2019 Epic Games, Inc. All Rights Reserved.
+
+#include "Common.ush"
+
+float4x4 Transform;
+
+void VSMain(in float4 InPosition : ATTRIBUTE0,
+			in float2 InTextureCoordinate : ATTRIBUTE1,
+			out float2 OutTextureCoordinate : TEXCOORD0,
+			out float4 OutPosition : SV_POSITION)
+{
+	//OutPosition = float4(InPosition.xy, 0, 1);
+	OutPosition = mul(InPosition, Transform);
+	OutTextureCoordinate = InTextureCoordinate;
+}
\ No newline at end of file
diff --git a/Engine/Source/Editor/LandscapeEditor/Private/LandscapeEdMode.cpp b/Engine/Source/Editor/LandscapeEditor/Private/LandscapeEdMode.cpp
--- a/Engine/Source/Editor/LandscapeEditor/Private/LandscapeEdMode.cpp
+++ b/Engine/Source/Editor/LandscapeEditor/Private/LandscapeEdMode.cpp
@@ -144,7 +144,7 @@ void ALandscape::SplitHeightmap(ULandscapeComponent* Comp, bool bMoveToCurrentLe
 
 #if WITH_EDITORONLY_DATA
 		// TODO: There is a bug here with updating the mip regions
-		/*if (GetMutableDefault<UEditorExperimentalSettings>()->bProceduralLandscape)
+		/*if (GetMutableDefault<UEditorExperimentalSettings>()->bLandscapeLayerSystem)
 		{
 			// Add component to new heightmap
 			FRenderDataPerHeightmap* RenderData = Landscape->RenderDataPerHeightmap.Find(HeightmapTexture);
@@ -155,7 +155,7 @@ void ALandscape::SplitHeightmap(ULandscapeComponent* Comp, bool bMoveToCurrentLe
 				NewData.Components.Add(Comp);
 				NewData.TopLeftSectionBase = Comp->GetSectionBase();
 				NewData.HeightmapKey = HeightmapTexture;
-				NewData.HeightmapsCPUReadBack = new FLandscapeProceduralTexture2DCPUReadBackResource(HeightmapTexture->Source.GetSizeX(), HeightmapTexture->Source.GetSizeY(), Comp->GetHeightmap(false)->GetPixelFormat(), HeightmapTexture->Source.GetNumMips());
+				NewData.HeightmapsCPUReadBack = new FLandscapeLayersTexture2DCPUReadBackResource(HeightmapTexture->Source.GetSizeX(), HeightmapTexture->Source.GetSizeY(), Comp->GetHeightmap(false)->GetPixelFormat(), HeightmapTexture->Source.GetNumMips());
 				BeginInitResource(NewData.HeightmapsCPUReadBack);
 
 				Landscape->RenderDataPerHeightmap.Add(HeightmapTexture, NewData);
@@ -479,7 +479,7 @@ void FEdModeLandscape::InitializeToolModes()
 	ToolMode_Sculpt->ValidTools.Add(TEXT("Retopologize"));
 	ToolMode_Sculpt->ValidTools.Add(TEXT("Visibility"));
 
-	if (GetMutableDefault<UEditorExperimentalSettings>()->bProceduralLandscape)
+	if (GetMutableDefault<UEditorExperimentalSettings>()->bLandscapeLayerSystem)
 	{
 		ToolMode_Sculpt->ValidTools.Add(TEXT("BPCustom"));
 	}
@@ -495,7 +495,7 @@ void FEdModeLandscape::InitializeToolModes()
 	ToolMode_Paint->ValidTools.Add(TEXT("Noise"));
 	ToolMode_Paint->ValidTools.Add(TEXT("Visibility"));
 
-	if (GetMutableDefault<UEditorExperimentalSettings>()->bProceduralLandscape)
+	if (GetMutableDefault<UEditorExperimentalSettings>()->bLandscapeLayerSystem)
 	{
 		ToolMode_Paint->ValidTools.Add(TEXT("BPCustom"));
 	}
@@ -559,12 +559,12 @@ void FEdModeLandscape::Enter()
 		ALandscapeProxy* LandscapeProxy = CurrentToolTarget.LandscapeInfo->GetLandscapeProxy();
 		LandscapeProxy->OnMaterialChangedDelegate().AddRaw(this, &FEdModeLandscape::OnLandscapeMaterialChangedDelegate);
 
-		if (GetMutableDefault<UEditorExperimentalSettings>()->bProceduralLandscape)
+		if (GetMutableDefault<UEditorExperimentalSettings>()->bLandscapeLayerSystem)
 		{
 			ALandscape* Landscape = GetLandscape();
 			if (Landscape)
 			{
-				Landscape->RequestProceduralContentUpdate(EProceduralContentUpdateFlag::All_Render);
+				Landscape->RequestLayersContentUpdate(ELandscapeLayersContentUpdateFlag::All_Render);
 			}
 		}
 	}
@@ -1898,15 +1898,15 @@ bool FEdModeLandscape::InputKey(FEditorViewportClient* ViewportClient, FViewport
 					{
 						const FName CurrentToolName = CurrentTool->GetToolName();
 						ALandscape* Landscape = GetLandscape();
-						bool bProceduralLandscape = GetMutableDefault<UEditorExperimentalSettings>()->bProceduralLandscape;
+						bool bLandscapeLayerSystem = GetMutableDefault<UEditorExperimentalSettings>()->bLandscapeLayerSystem;
 
-						if (bProceduralLandscape && GetCurrentProceduralLayer() && !GetCurrentProceduralLayer()->bVisible)
+						if (bLandscapeLayerSystem && GetCurrentLayer() && !GetCurrentLayer()->bVisible)
 						{
-							FMessageDialog::Open(EAppMsgType::Ok, NSLOCTEXT("UnrealEd", "LandscapeProceduralLayerHidden", "Painting in a hidden layer is not allowed."));
+							FMessageDialog::Open(EAppMsgType::Ok, NSLOCTEXT("UnrealEd", "LandscapeLayerHidden", "Painting in a hidden layer is not allowed."));
 						}
-						else if (bProceduralLandscape && GetCurrentProceduralLayer() && GetCurrentProceduralLayer()->bLocked)
+						else if (bLandscapeLayerSystem && GetCurrentLayer() && GetCurrentLayer()->bLocked)
 						{
-							FMessageDialog::Open(EAppMsgType::Ok, NSLOCTEXT("UnrealEd", "LandscapeProceduralLayerLocked", "This layer is locked. You must unlock it before you can work on this layer."));
+							FMessageDialog::Open(EAppMsgType::Ok, NSLOCTEXT("UnrealEd", "LandscapeLayerLocked", "This layer is locked. You must unlock it before you can work on this layer."));
 						}
 						else if (CurrentToolTarget.TargetType == ELandscapeToolTargetType::Weightmap && CurrentToolTarget.LayerInfo == NULL && CurrentToolName != TEXT("BPCustom"))
 						{
@@ -2399,7 +2399,7 @@ int32 FEdModeLandscape::UpdateLandscapeList()
 			CurrentToolTarget.LandscapeInfo = LandscapeList[0].Info;
 			CurrentIndex = 0;
 
-			SetCurrentProceduralLayer(0);
+			SetCurrentLayer(0);
 
 			// Init UI to saved value
 			ALandscapeProxy* LandscapeProxy = CurrentToolTarget.LandscapeInfo->GetLandscapeProxy();
@@ -2446,7 +2446,7 @@ void FEdModeLandscape::UpdateTargetList()
 			bool bFoundSelected = false;
 
 			// Add heightmap
-			LandscapeTargetList.Add(MakeShareable(new FLandscapeTargetListInfo(LOCTEXT("Heightmap", "Heightmap"), ELandscapeToolTargetType::Heightmap, CurrentToolTarget.LandscapeInfo.Get(), CurrentToolTarget.CurrentProceduralLayerIndex)));
+			LandscapeTargetList.Add(MakeShareable(new FLandscapeTargetListInfo(LOCTEXT("Heightmap", "Heightmap"), ELandscapeToolTargetType::Heightmap, CurrentToolTarget.LandscapeInfo.Get(), CurrentToolTarget.CurrentLayerIndex)));
 
 			if (CurrentToolTarget.TargetType == ELandscapeToolTargetType::Heightmap)
 			{
@@ -2455,7 +2455,7 @@ void FEdModeLandscape::UpdateTargetList()
 
 			// Add visibility
 			FLandscapeInfoLayerSettings VisibilitySettings(ALandscapeProxy::VisibilityLayer, LandscapeProxy);
-			LandscapeTargetList.Add(MakeShareable(new FLandscapeTargetListInfo(LOCTEXT("Visibility", "Visibility"), ELandscapeToolTargetType::Visibility, VisibilitySettings, CurrentToolTarget.CurrentProceduralLayerIndex)));
+			LandscapeTargetList.Add(MakeShareable(new FLandscapeTargetListInfo(LOCTEXT("Visibility", "Visibility"), ELandscapeToolTargetType::Visibility, VisibilitySettings, CurrentToolTarget.CurrentLayerIndex)));
 
 			if (CurrentToolTarget.TargetType == ELandscapeToolTargetType::Visibility)
 			{
@@ -2506,7 +2506,7 @@ void FEdModeLandscape::UpdateTargetList()
 				}
 
 				// Add the layer
-				LandscapeTargetList.Add(MakeShareable(new FLandscapeTargetListInfo(FText::FromName(LayerName), ELandscapeToolTargetType::Weightmap, LayerSettings, CurrentToolTarget.CurrentProceduralLayerIndex)));
+				LandscapeTargetList.Add(MakeShareable(new FLandscapeTargetListInfo(FText::FromName(LayerName), ELandscapeToolTargetType::Weightmap, LayerSettings, CurrentToolTarget.CurrentLayerIndex)));
 			}
 
 			if (!bFoundSelected)
@@ -3350,7 +3350,7 @@ void FEdModeLandscape::ReimportData(const FLandscapeTargetListInfo& TargetInfo)
 	const FString& SourceFilePath = TargetInfo.ReimportFilePath();
 	if (SourceFilePath.Len())
 	{
-		FScopedSetLandscapeCurrentEditingProceduralLayer Scope(GetLandscape(), GetCurrentProceduralLayerGuid(), [&] { RequestProceduralContentUpdate(); });
+		FScopedSetLandscapeEditingLayer Scope(GetLandscape(), GetCurrentLayerGuid(), [&] { RequestLayersContentUpdate(); });
 		ImportData(TargetInfo, SourceFilePath);
 	}
 	else
@@ -3457,7 +3457,7 @@ void FEdModeLandscape::ImportData(const FLandscapeTargetListInfo& TargetInfo, co
 
 			{
 				ALandscape* Landscape = GetLandscape();
-				FScopedSetLandscapeCurrentEditingProceduralLayer Scope(Landscape, GetCurrentProceduralLayerGuid(), [&] { check(Landscape); Landscape->RequestProceduralContentUpdate(EProceduralContentUpdateFlag::Heightmap_All); });
+				FScopedSetLandscapeEditingLayer Scope(Landscape, GetCurrentLayerGuid(), [&] { check(Landscape); Landscape->RequestLayersContentUpdate(ELandscapeLayersContentUpdateFlag::Heightmap_All); });
 
 				TArray<uint16> Data;
 				if (ImportResolution != LandscapeResolution)
@@ -3573,7 +3573,7 @@ void FEdModeLandscape::ImportData(const FLandscapeTargetListInfo& TargetInfo, co
 
 			{
 				ALandscape* Landscape = GetLandscape();
-				FScopedSetLandscapeCurrentEditingProceduralLayer Scope(Landscape, GetCurrentProceduralLayerGuid(), [&] { check(Landscape); Landscape->RequestProceduralContentUpdate(EProceduralContentUpdateFlag::Weightmap_All); });
+				FScopedSetLandscapeEditingLayer Scope(Landscape, GetCurrentLayerGuid(), [&] { check(Landscape); Landscape->RequestLayersContentUpdate(ELandscapeLayersContentUpdateFlag::Weightmap_All); });
 
 				TArray<uint8> Data;
 				if (ImportResolution != LandscapeResolution)
@@ -3939,17 +3939,17 @@ ALandscape* FEdModeLandscape::ChangeComponentSetting(int32 NumComponentsX, int32
 			Landscape->bUsedForNavigation = OldLandscapeProxy->bUsedForNavigation;
 			Landscape->MaxPaintedLayersPerComponent = OldLandscapeProxy->MaxPaintedLayersPerComponent;
 
-			if (GetMutableDefault<UEditorExperimentalSettings>()->bProceduralLandscape)
+			if (GetMutableDefault<UEditorExperimentalSettings>()->bLandscapeLayerSystem)
 			{
 				ALandscape* OldLandscapeActor = Cast<ALandscape>(OldLandscapeProxy);
 				ALandscape* NewLandscapeActor = Cast<ALandscape>(Landscape);
 
-				NewLandscapeActor->PreviousExperimentalLandscapeProcedural = OldLandscapeActor->PreviousExperimentalLandscapeProcedural;
+				NewLandscapeActor->PreviousExperimentalLandscapeLayers = OldLandscapeActor->PreviousExperimentalLandscapeLayers;
 
 				// Moves specific config including brushes to new landscape actor
-				for (const FProceduralLayer& Layer : OldLandscapeActor->ProceduralLayers)
+				for (const FLandscapeLayer& Layer : OldLandscapeActor->LandscapeLayers)
 				{
-					FProceduralLayer* NewLayer = NewLandscapeActor->ProceduralLayers.FindByPredicate( [&](const FProceduralLayer& NewLayer)	{ return NewLayer.Name == Layer.Name; }	);
+					FLandscapeLayer* NewLayer = NewLandscapeActor->LandscapeLayers.FindByPredicate( [&](const FLandscapeLayer& NewLayer)	{ return NewLayer.Name == Layer.Name; }	);
 					check(NewLayer != nullptr);
 					*NewLayer = Layer;
 				}
@@ -4122,22 +4122,22 @@ ELandscapeEditingState FEdModeLandscape::GetEditingState() const
 	return ELandscapeEditingState::Enabled;
 }
 
-int32 FEdModeLandscape::GetProceduralLayerCount() const
+int32 FEdModeLandscape::GetLayerCount() const
 {
 	ALandscape* Landscape = GetLandscape();
-	return Landscape ? Landscape->ProceduralLayers.Num() : 0;
+	return Landscape ? Landscape->LandscapeLayers.Num() : 0;
 }
 
-void FEdModeLandscape::SetCurrentProceduralLayer(int32 InLayerIndex)
+void FEdModeLandscape::SetCurrentLayer(int32 InLayerIndex)
 {
-	CurrentToolTarget.CurrentProceduralLayerIndex = InLayerIndex;
+	CurrentToolTarget.CurrentLayerIndex = InLayerIndex;
 
 	RefreshDetailPanel();
 }
 
-int32 FEdModeLandscape::GetCurrentProceduralLayerIndex() const
+int32 FEdModeLandscape::GetCurrentLayerIndex() const
 {
-	return CurrentToolTarget.CurrentProceduralLayerIndex;
+	return CurrentToolTarget.CurrentLayerIndex;
 }
 
 ALandscape* FEdModeLandscape::GetLandscape() const
@@ -4145,27 +4145,27 @@ ALandscape* FEdModeLandscape::GetLandscape() const
 	return CurrentToolTarget.LandscapeInfo.IsValid() ? CurrentToolTarget.LandscapeInfo->LandscapeActor.Get() : nullptr;
 }
 
-FProceduralLayer* FEdModeLandscape::GetProceduralLayer(int32 InLayerIndex) const
+FLandscapeLayer* FEdModeLandscape::GetLayer(int32 InLayerIndex) const
 {
 	ALandscape* Landscape = GetLandscape();
-	return Landscape ? Landscape->GetProceduralLayer(InLayerIndex) : nullptr;
+	return Landscape ? Landscape->GetLayer(InLayerIndex) : nullptr;
 }
 
-FName FEdModeLandscape::GetProceduralLayerName(int32 InLayerIndex) const
+FName FEdModeLandscape::GetLayerName(int32 InLayerIndex) const
 {
-	FProceduralLayer* Layer = GetProceduralLayer(InLayerIndex);
+	FLandscapeLayer* Layer = GetLayer(InLayerIndex);
 	return Layer ? Layer->Name : NAME_None;
 }
 
-bool FEdModeLandscape::CanRenameProceduralLayerTo(int32 InLayerIndex, const FName& InNewName)
+bool FEdModeLandscape::CanRenameLayerTo(int32 InLayerIndex, const FName& InNewName)
 {
 	ALandscape* Landscape = GetLandscape();
 	if (Landscape)
 	{
-		int32 LayerCount = GetProceduralLayerCount();
+		int32 LayerCount = GetLayerCount();
 		for (int32 LayerIdx = 0; LayerIdx < LayerCount; ++LayerIdx)
 		{
-			if (LayerIdx != InLayerIndex && GetProceduralLayerName(LayerIdx) == InNewName)
+			if (LayerIdx != InLayerIndex && GetLayerName(LayerIdx) == InNewName)
 			{
 				return false;
 			}
@@ -4174,18 +4174,18 @@ bool FEdModeLandscape::CanRenameProceduralLayerTo(int32 InLayerIndex, const FNam
 	return true;
 }
 
-void FEdModeLandscape::SetProceduralLayerName(int32 InLayerIndex, const FName& InName)
+void FEdModeLandscape::SetLayerName(int32 InLayerIndex, const FName& InName)
 {
 	ALandscape* Landscape = GetLandscape();
 	if (Landscape)
 	{
-		Landscape->SetProceduralLayerName(InLayerIndex, InName);
+		Landscape->SetLayerName(InLayerIndex, InName);
 	}
 }
 
-float FEdModeLandscape::GetProceduralLayerAlpha(int32 InLayerIndex) const
+float FEdModeLandscape::GetLayerAlpha(int32 InLayerIndex) const
 {
-	FProceduralLayer* Layer = GetProceduralLayer(InLayerIndex);
+	FLandscapeLayer* Layer = GetLayer(InLayerIndex);
 	if (Layer)
 	{
 		return (CurrentToolTarget.TargetType == ELandscapeToolTargetType::Heightmap) ? Layer->HeightmapAlpha : Layer->WeightmapAlpha;
@@ -4193,316 +4193,170 @@ float FEdModeLandscape::GetProceduralLayerAlpha(int32 InLayerIndex) const
 	return 1.0f;
 }
 
-void FEdModeLandscape::SetProceduralLayerAlpha(int32 InLayerIndex, float InAlpha)
+void FEdModeLandscape::SetLayerAlpha(int32 InLayerIndex, float InAlpha)
 {
 	ALandscape* Landscape = GetLandscape();
 	if (Landscape)
 	{
-		Landscape->SetProceduralLayerAlpha(InLayerIndex, InAlpha, CurrentToolTarget.TargetType == ELandscapeToolTargetType::Heightmap);
+		Landscape->SetLayerAlpha(InLayerIndex, InAlpha, CurrentToolTarget.TargetType == ELandscapeToolTargetType::Heightmap);
 	}
 }
 
-bool FEdModeLandscape::IsProceduralLayerVisible(int32 InLayerIndex) const
+bool FEdModeLandscape::IsLayerVisible(int32 InLayerIndex) const
 {
-	FProceduralLayer* Layer = GetProceduralLayer(InLayerIndex);
+	FLandscapeLayer* Layer = GetLayer(InLayerIndex);
 	return Layer ? Layer->bVisible : false;
 }
 
-void FEdModeLandscape::SetProceduralLayerVisibility(bool bInVisible, int32 InLayerIndex)
+void FEdModeLandscape::SetLayerVisibility(bool bInVisible, int32 InLayerIndex)
 {
 	ALandscape* Landscape = GetLandscape();
 	if (Landscape)
 	{
-		Landscape->SetProceduralLayerVisibility(InLayerIndex, bInVisible);
+		Landscape->SetLayerVisibility(InLayerIndex, bInVisible);
 	}
 }
 
-bool FEdModeLandscape::IsProceduralLayerLocked(int32 InLayerIndex) const
+bool FEdModeLandscape::IsLayerLocked(int32 InLayerIndex) const
 {
-	FProceduralLayer* Layer = GetProceduralLayer(InLayerIndex);
+	FLandscapeLayer* Layer = GetLayer(InLayerIndex);
 	return Layer ? Layer->bLocked : false;
 }
 
-void FEdModeLandscape::SetProceduralLayerLocked(int32 InLayerIndex, bool bInLocked)
+void FEdModeLandscape::SetLayerLocked(int32 InLayerIndex, bool bInLocked)
 {
-	FProceduralLayer* Layer = GetProceduralLayer(InLayerIndex);
+	FLandscapeLayer* Layer = GetLayer(InLayerIndex);
 	if (Layer)
 	{
 		Layer->bLocked = bInLocked;
 	}
 }
 
-void FEdModeLandscape::AddBrushToCurrentProceduralLayer(int32 InTargetType, ALandscapeBlueprintCustomBrush* InBrush)
+void FEdModeLandscape::RequestLayersContentUpdate(bool InUpdateAllMaterials)
 {
 	ALandscape* Landscape = GetLandscape();
-	FProceduralLayer* Layer = GetCurrentProceduralLayer();
-	if (!Landscape || !Layer)
-	{
-		return;
-	}
-
-	int32 AddedIndex = Layer->Brushes.Add(FLandscapeProceduralLayerBrush(InBrush));
-
-	if (InTargetType == ELandscapeToolTargetType::Type::Heightmap)
-	{
-		Layer->HeightmapBrushOrderIndices.Add(AddedIndex);
-	}
-	else
+	if (Landscape)
 	{
-		Layer->WeightmapBrushOrderIndices.Add(AddedIndex);
+		Landscape->RequestLayersContentUpdate(ELandscapeLayersContentUpdateFlag::All_Render, InUpdateAllMaterials);
 	}
-
-	InBrush->SetOwningLandscape(Landscape);
-
-	Landscape->RequestProceduralContentUpdate(EProceduralContentUpdateFlag::All_Render);
 }
 
-void FEdModeLandscape::RequestProceduralContentUpdate(bool InUpdateAllMaterials)
+void FEdModeLandscape::AddBrushToCurrentLayer(int32 InTargetType, ALandscapeBlueprintCustomBrush* InBrush)
 {
 	ALandscape* Landscape = GetLandscape();
-	if (Landscape)
+	if (Landscape == nullptr)
 	{
-		Landscape->RequestProceduralContentUpdate(EProceduralContentUpdateFlag::All_Render, InUpdateAllMaterials);
+		return;
 	}
+
+	Landscape->AddBrushToLayer(GetCurrentLayerIndex(), InTargetType, InBrush);
 }
 
-void FEdModeLandscape::RemoveBrushFromCurrentProceduralLayer(int32 InTargetType, ALandscapeBlueprintCustomBrush* InBrush)
+void FEdModeLandscape::RemoveBrushFromCurrentLayer(int32 InTargetType, ALandscapeBlueprintCustomBrush* InBrush)
 {
 	ALandscape* Landscape = GetLandscape();
-	if (!Landscape)
+
+	if (Landscape == nullptr)
 	{
 		return;
 	}
-	for (FProceduralLayer& Layer : Landscape->ProceduralLayers)
-	{
-		int32 IndexToRemove = INDEX_NONE;
-		for (int32 i = 0; i < Layer.Brushes.Num(); ++i)
-		{
-			if (Layer.Brushes[i].BPCustomBrush == InBrush)
-			{
-				IndexToRemove = i;
-				break;
-			}
-		}
 
-		if (IndexToRemove != INDEX_NONE)
-		{
-			Layer.Brushes.RemoveAt(IndexToRemove);
-
-			for (int32 i = 0; i < Layer.HeightmapBrushOrderIndices.Num(); ++i)
-			{
-				if (Layer.HeightmapBrushOrderIndices[i] == IndexToRemove)
-				{
-					// Update the value of the index of all the one after the one we removed, so index still correctly match actual brushes list
-					for (int32 j = 0; j < Layer.HeightmapBrushOrderIndices.Num(); ++j)
-					{
-						if (Layer.HeightmapBrushOrderIndices[j] > IndexToRemove)
-						{
-							--Layer.HeightmapBrushOrderIndices[j];
-						}
-					}
-
-					Layer.HeightmapBrushOrderIndices.RemoveAt(i);
-					break;
-				}
-			}
-
-			for (int32 i = 0; i < Layer.WeightmapBrushOrderIndices.Num(); ++i)
-			{
-				if (Layer.WeightmapBrushOrderIndices[i] == IndexToRemove)
-				{
-					// Update the value of the index of all the one after the one we removed, so index still correctly match actual brushes list
-					for (int32 j = 0; j < Layer.WeightmapBrushOrderIndices.Num(); ++j)
-					{
-						if (Layer.WeightmapBrushOrderIndices[j] > IndexToRemove)
-						{
-							--Layer.WeightmapBrushOrderIndices[j];
-						}
-					}
-
-					Layer.WeightmapBrushOrderIndices.RemoveAt(i);
-					break;
-				}
-			}
-
-			InBrush->SetOwningLandscape(nullptr);
-			break;
-		}
-	}
-
-	Landscape->RequestProceduralContentUpdate(EProceduralContentUpdateFlag::All_Render);
+	Landscape->RemoveBrushFromLayer(GetCurrentLayerIndex(), InTargetType, InBrush);
 }
 
-bool FEdModeLandscape::AreAllBrushesCommitedToCurrentProceduralLayer(int32 InTargetType)
+bool FEdModeLandscape::AreAllBrushesCommitedToCurrentLayer(int32 InTargetType)
 {
-	FProceduralLayer* Layer = GetCurrentProceduralLayer();
+	ALandscape* Landscape = GetLandscape();
 
-	if (Layer == nullptr)
+	if (Landscape == nullptr)
 	{
 		return false;
 	}
 
-	for (FLandscapeProceduralLayerBrush& Brush : Layer->Brushes)
-	{
-		if (Brush.BPCustomBrush != nullptr && !Brush.BPCustomBrush->IsCommited()
-			&& ((InTargetType == ELandscapeToolTargetType::Type::Heightmap && Brush.BPCustomBrush->IsAffectingHeightmap()) || (InTargetType == ELandscapeToolTargetType::Type::Weightmap && Brush.BPCustomBrush->IsAffectingWeightmap())))
-		{
-			return false;
-		}
-	}
-
-	return true;	
+	return Landscape->AreAllBrushesCommitedToLayer(GetCurrentLayerIndex(), InTargetType);
 }
 
-void FEdModeLandscape::SetCurrentProceduralLayerBrushesCommitState(int32 InTargetType, bool InCommited)
+void FEdModeLandscape::SetBrushesCommitStateForCurrentLayer(int32 InTargetType, bool InCommited)
 {
-	FProceduralLayer* Layer = GetCurrentProceduralLayer();
+	ALandscape* Landscape = GetLandscape();
 
-	if (Layer == nullptr)
+	if (Landscape == nullptr)
 	{
 		return;
 	}
 
-	for (FLandscapeProceduralLayerBrush& Brush : Layer->Brushes)
-	{
-		if (Brush.BPCustomBrush != nullptr)
-		{
-			Brush.BPCustomBrush->SetCommitState(InCommited);
-		}
-	}
-
+	Landscape->SetBrushesCommitStateForLayer(GetCurrentLayerIndex(), InTargetType, InCommited);
+	
 	GEngine->BroadcastLevelActorListChanged();
 }
 
-TArray<int8>& FEdModeLandscape::GetBrushesOrderForCurrentProceduralLayer(int32 InTargetType) const
+TArray<int8>& FEdModeLandscape::GetBrushesOrderForCurrentLayer(int32 InTargetType) const
 {
-	FProceduralLayer* Layer = GetCurrentProceduralLayer();
-	check(Layer);
-
-	if (InTargetType == ELandscapeToolTargetType::Type::Heightmap)
-	{
-		return Layer->HeightmapBrushOrderIndices;
-	}
-	else
-	{
-		return Layer->WeightmapBrushOrderIndices;
-	}
+	ALandscape* Landscape = GetLandscape();
+	check(Landscape);
+	return Landscape->GetBrushesOrderForLayer(GetCurrentLayerIndex(), InTargetType);
 }
 
-ALandscapeBlueprintCustomBrush* FEdModeLandscape::GetBrushForCurrentProceduralLayer(int32 InTargetType, int8 InBrushIndex) const
+ALandscapeBlueprintCustomBrush* FEdModeLandscape::GetBrushForCurrentLayer(int32 InTargetType, int8 InBrushIndex) const
 {
-	FProceduralLayer* Layer = GetCurrentProceduralLayer();
+	ALandscape* Landscape = GetLandscape();
 
-	if (Layer == nullptr)
+	if (Landscape == nullptr)
 	{
 		return nullptr;
 	}
 
-	if (InTargetType == ELandscapeToolTargetType::Type::Heightmap)
-	{
-		if (Layer->HeightmapBrushOrderIndices.IsValidIndex(InBrushIndex))
-		{
-			int8 ActualBrushIndex = Layer->HeightmapBrushOrderIndices[InBrushIndex];
-			if (Layer->Brushes.IsValidIndex(ActualBrushIndex))
-			{
-				return Layer->Brushes[ActualBrushIndex].BPCustomBrush;
-			}
-		}
-	}
-	else
-	{
-		if (Layer->WeightmapBrushOrderIndices.IsValidIndex(InBrushIndex))
-		{
-			int8 ActualBrushIndex = Layer->WeightmapBrushOrderIndices[InBrushIndex];
-			if (Layer->Brushes.IsValidIndex(ActualBrushIndex))
-			{
-				return Layer->Brushes[ActualBrushIndex].BPCustomBrush;
-			}
-		}
-	}
-
-	return nullptr;
+	return Landscape->GetBrushForLayer(GetCurrentLayerIndex(), InTargetType, InBrushIndex);
 }
 
-TArray<ALandscapeBlueprintCustomBrush*> FEdModeLandscape::GetBrushesForCurrentProceduralLayer(int32 InTargetType)
+TArray<ALandscapeBlueprintCustomBrush*> FEdModeLandscape::GetBrushesForCurrentLayer(int32 InTargetType)
 {
+	ALandscape* Landscape = GetLandscape();
 	TArray<ALandscapeBlueprintCustomBrush*> Brushes;
 
-	FProceduralLayer* Layer = GetCurrentProceduralLayer();
-
-	if (Layer == nullptr)
+	if (Landscape == nullptr)
 	{
 		return Brushes;
 	}
 
-	Brushes.Reserve(Layer->Brushes.Num());
-
-	for (const FLandscapeProceduralLayerBrush& Brush : Layer->Brushes)
-	{
-		if (Brush.BPCustomBrush != nullptr && ((Brush.BPCustomBrush->IsAffectingHeightmap() && InTargetType == ELandscapeToolTargetType::Type::Heightmap)
-			|| (Brush.BPCustomBrush->IsAffectingWeightmap() && InTargetType == ELandscapeToolTargetType::Type::Weightmap)))
-		{
-			Brushes.Add(Brush.BPCustomBrush);
-		}
-	}
-
+	Brushes = Landscape->GetBrushesForLayer(GetCurrentLayerIndex(), InTargetType);
 	return Brushes;
 }
 
-bool FEdModeLandscape::IsCurrentProceduralLayerBlendSubstractive(const TWeakObjectPtr<ULandscapeLayerInfoObject>& InLayerInfoObj) const
+bool FEdModeLandscape::IsCurrentLayerBlendSubstractive(const TWeakObjectPtr<ULandscapeLayerInfoObject>& InLayerInfoObj) const
 {
-	FProceduralLayer* Layer = GetCurrentProceduralLayer();
-
-	if (Layer == nullptr)
-	{
-		return false;
-	}
-
-	bool* AllocationBlend = Layer->WeightmapLayerAllocationBlend.Find(InLayerInfoObj.Get());
+	ALandscape* Landscape = GetLandscape();
 
-	if (AllocationBlend != nullptr)
+	if (Landscape != nullptr)
 	{
-		return (*AllocationBlend);
+		return Landscape->IsLayerBlendSubstractive(GetCurrentLayerIndex(), InLayerInfoObj);
 	}
 
 	return false;
 }
 
-void FEdModeLandscape::SetCurrentProceduralLayerSubstractiveBlendStatus(bool InStatus, const TWeakObjectPtr<ULandscapeLayerInfoObject>& InLayerInfoObj)
+void FEdModeLandscape::SetCurrentLayerSubstractiveBlendStatus(bool InStatus, const TWeakObjectPtr<ULandscapeLayerInfoObject>& InLayerInfoObj)
 {
-	FProceduralLayer* Layer = GetCurrentProceduralLayer();
-
-	if (Layer == nullptr)
-	{
-		return;
-	}
-
-	bool* AllocationBlend = Layer->WeightmapLayerAllocationBlend.Find(InLayerInfoObj.Get());
+	ALandscape* Landscape = GetLandscape();
 
-	if (AllocationBlend == nullptr)
-	{
-		Layer->WeightmapLayerAllocationBlend.Add(InLayerInfoObj.Get(), InStatus);
-	}
-	else
+	if (Landscape != nullptr)
 	{
-		*AllocationBlend = InStatus;
+		return Landscape->SetLayerSubstractiveBlendStatus(GetCurrentLayerIndex(), InStatus, InLayerInfoObj);
 	}
-
-	RequestProceduralContentUpdate(true);
 }
 
-FProceduralLayer* FEdModeLandscape::GetCurrentProceduralLayer() const
+FLandscapeLayer* FEdModeLandscape::GetCurrentLayer() const
 {
-	return GetProceduralLayer(GetCurrentProceduralLayerIndex());
+	return GetLayer(GetCurrentLayerIndex());
 }
 
-FGuid FEdModeLandscape::GetCurrentProceduralLayerGuid() const
+FGuid FEdModeLandscape::GetCurrentLayerGuid() const
 {
-	FProceduralLayer* CurrentProceduralLayer = GetCurrentProceduralLayer();
-	return CurrentProceduralLayer ? CurrentProceduralLayer->Guid : FGuid();
+	FLandscapeLayer* CurrentLayer = GetCurrentLayer();
+	return CurrentLayer ? CurrentLayer->Guid : FGuid();
 }
 
-bool FEdModeLandscape::NeedToFillEmptyLayersForProcedural() const
+bool FEdModeLandscape::NeedToFillEmptyMaterialLayers() const
 {
 	if (!CurrentToolTarget.LandscapeInfo.IsValid() || !CurrentToolTarget.LandscapeInfo->LandscapeActor.IsValid())
 	{
@@ -4519,11 +4373,11 @@ bool FEdModeLandscape::NeedToFillEmptyLayersForProcedural() const
 
 	for (const ALandscapeProxy* LandscapeProxy : AllLandscapes)
 	{
-		for (const auto& ItLayerPair : LandscapeProxy->ProceduralLayersData)
+		for (const auto& ItLayerPair : LandscapeProxy->LandscapeLayersData)
 		{
-			const FProceduralLayerData& ProceduralLayerData = ItLayerPair.Value;
+			const FLandscapeLayerData& LayerData = ItLayerPair.Value;
 
-			for (const auto& ItWeightmapPair : ProceduralLayerData.WeightmapData)
+			for (const auto& ItWeightmapPair : LayerData.WeightmapData)
 			{
 				const FWeightmapLayerData& WeightmapData = ItWeightmapPair.Value;
 
@@ -4547,7 +4401,7 @@ void FEdModeLandscape::OnLevelActorAdded(AActor* InActor)
 
 	if (Brush != nullptr && Brush->GetTypedOuter<UPackage>() != GetTransientPackage())
 	{
-		AddBrushToCurrentProceduralLayer(CurrentToolTarget.TargetType, Brush);
+		AddBrushToCurrentLayer(CurrentToolTarget.TargetType, Brush);
 		RefreshDetailPanel();
 	}
 }
@@ -4558,7 +4412,7 @@ void FEdModeLandscape::OnLevelActorRemoved(AActor* InActor)
 
 	if (Brush != nullptr && Brush->GetTypedOuter<UPackage>() != GetTransientPackage())
 	{
-		RemoveBrushFromCurrentProceduralLayer(CurrentToolTarget.TargetType, Brush);
+		RemoveBrushFromCurrentLayer(CurrentToolTarget.TargetType, Brush);
 		RefreshDetailPanel();
 	}
 }
diff --git a/Engine/Source/Editor/LandscapeEditor/Private/LandscapeEdMode.h b/Engine/Source/Editor/LandscapeEditor/Private/LandscapeEdMode.h
--- a/Engine/Source/Editor/LandscapeEditor/Private/LandscapeEdMode.h
+++ b/Engine/Source/Editor/LandscapeEditor/Private/LandscapeEdMode.h
@@ -69,9 +69,9 @@ struct FLandscapeTargetListInfo
 	TWeakObjectPtr<class UMaterialInstanceConstant> ThumbnailMIC;	// ignored for heightmap
 	int32 DebugColorChannel;										// ignored for heightmap
 	uint32 bValid : 1;												// ignored for heightmap
-	int32 ProceduralLayerIndex;
+	int32 LayerIndex;
 
-	FLandscapeTargetListInfo(FText InTargetName, ELandscapeToolTargetType::Type InTargetType, const FLandscapeInfoLayerSettings& InLayerSettings, int32 InProceduralLayerIndex)
+	FLandscapeTargetListInfo(FText InTargetName, ELandscapeToolTargetType::Type InTargetType, const FLandscapeInfoLayerSettings& InLayerSettings, int32 InLayerIndex)
 		: TargetName(InTargetName)
 		, TargetType(InTargetType)
 		, LandscapeInfo(InLayerSettings.Owner->GetLandscapeInfo())
@@ -81,11 +81,11 @@ struct FLandscapeTargetListInfo
 		, ThumbnailMIC(InLayerSettings.ThumbnailMIC)
 		, DebugColorChannel(InLayerSettings.DebugColorChannel)
 		, bValid(InLayerSettings.bValid)
-		, ProceduralLayerIndex (InProceduralLayerIndex)
+		, LayerIndex(InLayerIndex)
 	{
 	}
 
-	FLandscapeTargetListInfo(FText InTargetName, ELandscapeToolTargetType::Type InTargetType, ULandscapeInfo* InLandscapeInfo, int32 InProceduralLayerIndex)
+	FLandscapeTargetListInfo(FText InTargetName, ELandscapeToolTargetType::Type InTargetType, ULandscapeInfo* InLandscapeInfo, int32 InLayerIndex)
 		: TargetName(InTargetName)
 		, TargetType(InTargetType)
 		, LandscapeInfo(InLandscapeInfo)
@@ -94,7 +94,7 @@ struct FLandscapeTargetListInfo
 		, Owner(NULL)
 		, ThumbnailMIC(NULL)
 		, bValid(true)
-		, ProceduralLayerIndex(InProceduralLayerIndex)
+		, LayerIndex(InLayerIndex)
 	{
 	}
 
@@ -489,36 +489,36 @@ public:
 	void OnLandscapeMaterialChangedDelegate();
 	void RefreshDetailPanel();
 
-	// Procedural Layers
-	int32 GetProceduralLayerCount() const;
-	void SetCurrentProceduralLayer(int32 InLayerIndex);
-	int32 GetCurrentProceduralLayerIndex() const;
+	// Layers
+	int32 GetLayerCount() const;
+	void SetCurrentLayer(int32 InLayerIndex);
+	int32 GetCurrentLayerIndex() const;
 	ALandscape* GetLandscape() const;
-	struct FProceduralLayer* GetProceduralLayer(int32 InLayerIndex) const;
-	FName GetProceduralLayerName(int32 InLayerIndex) const;
-	void SetProceduralLayerName(int32 InLayerIndex, const FName& InName);
-	bool CanRenameProceduralLayerTo(int32 InLayerIndex, const FName& InNewName);
-	float GetProceduralLayerAlpha(int32 InLayerIndex) const;
-	void SetProceduralLayerAlpha(int32 InLayerIndex, float InAlpha);
-	void SetProceduralLayerVisibility(bool InVisible, int32 InLayerIndex);
-	bool IsProceduralLayerVisible(int32 InLayerIndex) const;
-	bool IsProceduralLayerLocked(int32 InLayerIndex) const;
-	void SetProceduralLayerLocked(int32 InLayerIndex, bool bInLocked);
-	struct FProceduralLayer* GetCurrentProceduralLayer() const;
-	FGuid GetCurrentProceduralLayerGuid() const;
-	bool IsCurrentProceduralLayerBlendSubstractive(const TWeakObjectPtr<ULandscapeLayerInfoObject>& InLayerInfoObj) const;
-	void SetCurrentProceduralLayerSubstractiveBlendStatus(bool InStatus, const TWeakObjectPtr<ULandscapeLayerInfoObject>& InLayerInfoObj);
-
-	void AddBrushToCurrentProceduralLayer(int32 InTargetType, class ALandscapeBlueprintCustomBrush* InBrush);
-	void RemoveBrushFromCurrentProceduralLayer(int32 InTargetType, class ALandscapeBlueprintCustomBrush* InBrush);
-	bool AreAllBrushesCommitedToCurrentProceduralLayer(int32 InTargetType);
-	void SetCurrentProceduralLayerBrushesCommitState(int32 InTargetType, bool InCommited);
-	TArray<int8>& GetBrushesOrderForCurrentProceduralLayer(int32 InTargetType) const;
-	class ALandscapeBlueprintCustomBrush* GetBrushForCurrentProceduralLayer(int32 InTargetType, int8 BrushIndex) const;
-	TArray<class ALandscapeBlueprintCustomBrush*> GetBrushesForCurrentProceduralLayer(int32 InTargetType);
+	struct FLandscapeLayer* GetLayer(int32 InLayerIndex) const;
+	FName GetLayerName(int32 InLayerIndex) const;
+	void SetLayerName(int32 InLayerIndex, const FName& InName);
+	bool CanRenameLayerTo(int32 InLayerIndex, const FName& InNewName);
+	float GetLayerAlpha(int32 InLayerIndex) const;
+	void SetLayerAlpha(int32 InLayerIndex, float InAlpha);
+	void SetLayerVisibility(bool InVisible, int32 InLayerIndex);
+	bool IsLayerVisible(int32 InLayerIndex) const;
+	bool IsLayerLocked(int32 InLayerIndex) const;
+	void SetLayerLocked(int32 InLayerIndex, bool bInLocked);
+	struct FLandscapeLayer* GetCurrentLayer() const;
+	FGuid GetCurrentLayerGuid() const;
+	bool IsCurrentLayerBlendSubstractive(const TWeakObjectPtr<ULandscapeLayerInfoObject>& InLayerInfoObj) const;
+	void SetCurrentLayerSubstractiveBlendStatus(bool InStatus, const TWeakObjectPtr<ULandscapeLayerInfoObject>& InLayerInfoObj);
+
+	void AddBrushToCurrentLayer(int32 InTargetType, class ALandscapeBlueprintCustomBrush* InBrush);
+	void RemoveBrushFromCurrentLayer(int32 InTargetType, class ALandscapeBlueprintCustomBrush* InBrush);
+	bool AreAllBrushesCommitedToCurrentLayer(int32 InTargetType);
+	void SetBrushesCommitStateForCurrentLayer(int32 InTargetType, bool InCommited);
+	TArray<int8>& GetBrushesOrderForCurrentLayer(int32 InTargetType) const;
+	class ALandscapeBlueprintCustomBrush* GetBrushForCurrentLayer(int32 InTargetType, int8 BrushIndex) const;
+	TArray<class ALandscapeBlueprintCustomBrush*> GetBrushesForCurrentLayer(int32 InTargetType);
 	
-	bool NeedToFillEmptyLayersForProcedural() const;
-	void RequestProceduralContentUpdate(bool InUpdateAllMaterials = false);
+	bool NeedToFillEmptyMaterialLayers() const;
+	void RequestLayersContentUpdate(bool InUpdateAllMaterials = false);
 
 	void OnLevelActorAdded(AActor* InActor);
 	void OnLevelActorRemoved(AActor* InActor);
diff --git a/Engine/Source/Editor/LandscapeEditor/Private/LandscapeEdModeComponentTools.cpp b/Engine/Source/Editor/LandscapeEditor/Private/LandscapeEdModeComponentTools.cpp
--- a/Engine/Source/Editor/LandscapeEditor/Private/LandscapeEdModeComponentTools.cpp
+++ b/Engine/Source/Editor/LandscapeEditor/Private/LandscapeEdModeComponentTools.cpp
@@ -217,9 +217,9 @@ public:
 
 			ALandscape* Landscape = LandscapeInfo->LandscapeActor.Get();
 
-			if (Landscape != nullptr && Landscape->HasProceduralContent && !GetMutableDefault<UEditorExperimentalSettings>()->bProceduralLandscape)
+			if (Landscape != nullptr && Landscape->HasLayersContent && !GetMutableDefault<UEditorExperimentalSettings>()->bLandscapeLayerSystem)
 			{
-				FMessageLog("MapCheck").Warning()->AddToken(FTextToken::Create(LOCTEXT("LandscapeProcedural_ChangingDataWithoutSettings", "This map contains landscape procedural content, modifying the landscape data will result in data loss when the map is reopened with Landscape Procedural settings on. Please enable Landscape Procedural settings before modifying the data.")));
+				FMessageLog("MapCheck").Warning()->AddToken(FTextToken::Create(LOCTEXT("LandscapeLayers_ChangingDataWithoutSettings", "This map contains landscape layer system content, modifying the landscape data will result in data loss when the map is reopened with Landscape Layer System settings on. Please enable Landscape Layer System settings before modifying the data.")));
 				FMessageLog("MapCheck").Open(EMessageSeverity::Warning);
 			}
 
@@ -304,9 +304,9 @@ public:
 
 			ALandscape* Landscape = LandscapeInfo->LandscapeActor.Get();
 
-			if (Landscape != nullptr && Landscape->HasProceduralContent && !GetMutableDefault<UEditorExperimentalSettings>()->bProceduralLandscape)
+			if (Landscape != nullptr && Landscape->HasLayersContent && !GetMutableDefault<UEditorExperimentalSettings>()->bLandscapeLayerSystem)
 			{
-				FMessageLog("MapCheck").Warning()->AddToken(FTextToken::Create(LOCTEXT("LandscapeProcedural_ChangingDataWithoutSettings", "This map contains landscape procedural content, modifying the landscape data will result in data loss when the map is reopened with Landscape Procedural settings on. Please enable Landscape Procedural settings before modifying the data.")));
+				FMessageLog("MapCheck").Warning()->AddToken(FTextToken::Create(LOCTEXT("LandscapeLayers_ChangingDataWithoutSettings", "This map contains landscape layer system content, modifying the landscape data will result in data loss when the map is reopened with Landscape Layer System settings on. Please enable Landscape Layer System settings before modifying the data.")));
 				FMessageLog("MapCheck").Open(EMessageSeverity::Warning);
 			}
 
@@ -342,13 +342,13 @@ public:
 			return false;
 		}
 
-		if (GetMutableDefault<UEditorExperimentalSettings>()->bProceduralLandscape)
+		if (GetMutableDefault<UEditorExperimentalSettings>()->bLandscapeLayerSystem)
 		{
 			ALandscape* Landscape = this->EdMode->GetLandscape();
 			if (Landscape)
 			{
-				Landscape->RequestProceduralContentUpdate(EProceduralContentUpdateFlag::Weightmap_Render);
-				Landscape->SetCurrentEditingProceduralLayer(this->EdMode->GetCurrentProceduralLayerGuid());
+				Landscape->RequestLayersContentUpdate(ELandscapeLayersContentUpdateFlag::Weightmap_Render);
+				Landscape->SetEditingLayer(this->EdMode->GetCurrentLayerGuid());
 			}
 		}
 
@@ -359,25 +359,25 @@ public:
 	{
 		FLandscapeToolBase<FLandscapeToolStrokeVisibility>::Tick(ViewportClient, DeltaTime);
 
-		if (GetMutableDefault<UEditorExperimentalSettings>()->bProceduralLandscape && this->IsToolActive())
+		if (GetMutableDefault<UEditorExperimentalSettings>()->bLandscapeLayerSystem && this->IsToolActive())
 		{
 			ALandscape* Landscape = this->EdMode->GetLandscape();
 			if (Landscape)
 			{
-				Landscape->RequestProceduralContentUpdate(EProceduralContentUpdateFlag::Weightmap_Render);
+				Landscape->RequestLayersContentUpdate(ELandscapeLayersContentUpdateFlag::Weightmap_Render);
 			}
 		}
 	}
 
 	virtual void EndTool(FEditorViewportClient* ViewportClient) override
 	{
-		if (GetMutableDefault<UEditorExperimentalSettings>()->bProceduralLandscape)
+		if (GetMutableDefault<UEditorExperimentalSettings>()->bLandscapeLayerSystem)
 		{
 			ALandscape* Landscape = this->EdMode->GetLandscape();
 			if (Landscape)
 			{
-				Landscape->SetCurrentEditingProceduralLayer();
-				Landscape->RequestProceduralContentUpdate(EProceduralContentUpdateFlag::Weightmap_All);
+				Landscape->SetEditingLayer();
+				Landscape->RequestLayersContentUpdate(ELandscapeLayersContentUpdateFlag::Weightmap_All);
 			}
 		}
 
@@ -659,7 +659,7 @@ public:
 								for (auto& WeightmapUsagePair : LandscapeProxy->WeightmapUsageMap)
 								{
 									ULandscapeWeightmapUsage* TryWeightmapUsage = WeightmapUsagePair.Value;
-									if (TryWeightmapUsage->FreeChannelCount() >= TotalNeededChannels) // TODO: handle procedural layer
+									if (TryWeightmapUsage->FreeChannelCount() >= TotalNeededChannels) // TODO: handle layers
 									{
 										// See if this candidate is closer than any others we've found
 										for (int32 ChanIdx = 0; ChanIdx < 4; ChanIdx++)
@@ -958,9 +958,10 @@ public:
 				NewComponents[Idx]->RegisterComponent();
 			}
 
-			if (LandscapeInfo->LandscapeActor.IsValid() && LandscapeInfo->LandscapeActor.Get()->HasProceduralContent && !GetMutableDefault<UEditorExperimentalSettings>()->bProceduralLandscape)
+			if (LandscapeInfo->LandscapeActor.IsValid() && LandscapeInfo->LandscapeActor.Get()->HasLayersContent && !GetMutableDefault<UEditorExperimentalSettings>()->HasLayersContent)
 			{
-				FMessageLog("MapCheck").Warning()->AddToken(FTextToken::Create(LOCTEXT("LandscapeProcedural_ChangingDataWithoutSettings", "This map contains landscape procedural content, modifying the landscape data will result in data loss when the map is reopened with Landscape Procedural settings on. Please enable Landscape Procedural settings before modifying the data.")));
+				FMessageLog("MapCheck").Warning()->AddToken(FTextToken::Create(LOCTEXT("LandscapeLayers_ChangingDataWithoutSettings", "This map contains landscape layer system content, modifying the landscape data will result in data loss when the map is reopened with Landscape Layer System settings on. Please enable Landscape Layer System settings before modifying the data.")));
+
 				FMessageLog("MapCheck").Open(EMessageSeverity::Warning);
 			}
 
@@ -1782,9 +1783,9 @@ public:
 
 			ALandscape* Landscape = LandscapeInfo->LandscapeActor.Get();
 
-			if (Landscape != nullptr && Landscape->HasProceduralContent && !GetMutableDefault<UEditorExperimentalSettings>()->bProceduralLandscape)
+			if (Landscape != nullptr && Landscape->HasLayersContent && !GetMutableDefault<UEditorExperimentalSettings>()->bLandscapeLayerSystem)
 			{
-				FMessageLog("MapCheck").Warning()->AddToken(FTextToken::Create(LOCTEXT("LandscapeProcedural_ChangingDataWithoutSettings", "This map contains landscape procedural content, modifying the landscape data will result in data loss when the map is reopened with Landscape Procedural settings on. Please enable Landscape Procedural settings before modifying the data.")));
+				FMessageLog("MapCheck").Warning()->AddToken(FTextToken::Create(LOCTEXT("LandscapeLayers_ChangingDataWithoutSettings", "This map contains landscape layer system content, modifying the landscape data will result in data loss when the map is reopened with Landscape Layer System settings on. Please enable Landscape Layer System settings before modifying the data.")));
 				FMessageLog("MapCheck").Open(EMessageSeverity::Warning);
 			}
 
diff --git a/Engine/Source/Editor/LandscapeEditor/Private/LandscapeEdModeErosionTools.cpp b/Engine/Source/Editor/LandscapeEditor/Private/LandscapeEdModeErosionTools.cpp
--- a/Engine/Source/Editor/LandscapeEditor/Private/LandscapeEdModeErosionTools.cpp
+++ b/Engine/Source/Editor/LandscapeEditor/Private/LandscapeEdModeErosionTools.cpp
@@ -259,9 +259,9 @@ public:
 
 		ALandscape* Landscape = LandscapeInfo->LandscapeActor.Get();
 
-		if (Landscape != nullptr && Landscape->HasProceduralContent && !GetMutableDefault<UEditorExperimentalSettings>()->bProceduralLandscape)
+		if (Landscape != nullptr && Landscape->HasLayersContent && !GetMutableDefault<UEditorExperimentalSettings>()->bLandscapeLayerSystem)
 		{
-			FMessageLog("MapCheck").Warning()->AddToken(FTextToken::Create(LOCTEXT("LandscapeProcedural_ChangingDataWithoutSettings", "This map contains landscape procedural content, modifying the landscape data will result in data loss when the map is reopened with Landscape Procedural settings on. Please enable Landscape Procedural settings before modifying the data.")));
+			FMessageLog("MapCheck").Warning()->AddToken(FTextToken::Create(LOCTEXT("LandscapeLayers_ChangingDataWithoutSettings", "This map contains landscape layer system content, modifying the landscape data will result in data loss when the map is reopened with Landscape Layer System settings on. Please enable Landscape Layer System settings before modifying the data.")));
 			FMessageLog("MapCheck").Open(EMessageSeverity::Warning);
 		}
 
@@ -490,9 +490,9 @@ public:
 
 		ALandscape* Landscape = LandscapeInfo->LandscapeActor.Get();
 
-		if (Landscape != nullptr && Landscape->HasProceduralContent && !GetMutableDefault<UEditorExperimentalSettings>()->bProceduralLandscape)
+		if (Landscape != nullptr && Landscape->HasLayersContent && !GetMutableDefault<UEditorExperimentalSettings>()->bLandscapeLayerSystem)
 		{
-			FMessageLog("MapCheck").Warning()->AddToken(FTextToken::Create(LOCTEXT("LandscapeProcedural_ChangingDataWithoutSettings", "This map contains landscape procedural content, modifying the landscape data will result in data loss when the map is reopened with Landscape Procedural settings on. Please enable Landscape Procedural settings before modifying the data.")));
+			FMessageLog("MapCheck").Warning()->AddToken(FTextToken::Create(LOCTEXT("LandscapeLayers_ChangingDataWithoutSettings", "This map contains landscape layer system content, modifying the landscape data will result in data loss when the map is reopened with Landscape Layer System settings on. Please enable Landscape Layer System settings before modifying the data.")));
 			FMessageLog("MapCheck").Open(EMessageSeverity::Warning);
 		}
 
diff --git a/Engine/Source/Editor/LandscapeEditor/Private/LandscapeEdModePaintTools.cpp b/Engine/Source/Editor/LandscapeEditor/Private/LandscapeEdModePaintTools.cpp
--- a/Engine/Source/Editor/LandscapeEditor/Private/LandscapeEdModePaintTools.cpp
+++ b/Engine/Source/Editor/LandscapeEditor/Private/LandscapeEdModePaintTools.cpp
@@ -63,26 +63,26 @@ public:
 	{
 		FLandscapeToolBase<TStrokeClass>::Tick(ViewportClient, DeltaTime);
 
-		if (GetMutableDefault<UEditorExperimentalSettings>()->bProceduralLandscape && this->IsToolActive())
+		if (GetMutableDefault<UEditorExperimentalSettings>()->bLandscapeLayerSystem && this->IsToolActive())
 		{
 			ALandscape* Landscape = this->EdMode->CurrentToolTarget.LandscapeInfo->LandscapeActor.Get();
 			if (Landscape != nullptr)
 			{
-				Landscape->RequestProceduralContentUpdate(this->EdMode->CurrentToolTarget.TargetType == ELandscapeToolTargetType::Type::Heightmap ? EProceduralContentUpdateFlag::Heightmap_Render : EProceduralContentUpdateFlag::Weightmap_Render);
+				Landscape->RequestLayersContentUpdate(this->EdMode->CurrentToolTarget.TargetType == ELandscapeToolTargetType::Type::Heightmap ? ELandscapeLayersContentUpdateFlag::Heightmap_Render : ELandscapeLayersContentUpdateFlag::Weightmap_Render);
 			}
 		}
 	}
 
 	virtual bool BeginTool(FEditorViewportClient* ViewportClient, const FLandscapeToolTarget& InTarget, const FVector& InHitLocation) override
 	{
-		if (GetMutableDefault<UEditorExperimentalSettings>()->bProceduralLandscape)
+		if (GetMutableDefault<UEditorExperimentalSettings>()->bLandscapeLayerSystem)
 		{
 			ALandscape* Landscape = this->EdMode->GetLandscape();
 			if (Landscape)
 			{
 				bool bUpdateHeightmap = this->EdMode->CurrentToolTarget.TargetType == ELandscapeToolTargetType::Type::Heightmap;
-				Landscape->RequestProceduralContentUpdate(bUpdateHeightmap ? EProceduralContentUpdateFlag::Heightmap_Render : EProceduralContentUpdateFlag::Weightmap_Render);
-				Landscape->SetCurrentEditingProceduralLayer(this->EdMode->GetCurrentProceduralLayerGuid());
+				Landscape->RequestLayersContentUpdate(bUpdateHeightmap ? ELandscapeLayersContentUpdateFlag::Heightmap_Render : ELandscapeLayersContentUpdateFlag::Weightmap_Render);
+				Landscape->SetEditingLayer(this->EdMode->GetCurrentLayerGuid());
 			}
 		}
 
@@ -91,14 +91,14 @@ public:
 
 	virtual void EndTool(FEditorViewportClient* ViewportClient) override
 	{
-		if (GetMutableDefault<UEditorExperimentalSettings>()->bProceduralLandscape)
+		if (GetMutableDefault<UEditorExperimentalSettings>()->bLandscapeLayerSystem)
 		{
 			ALandscape* Landscape = this->EdMode->GetLandscape();
 			if (Landscape)
 			{
 				bool bUpdateHeightmap = this->EdMode->CurrentToolTarget.TargetType == ELandscapeToolTargetType::Type::Heightmap;
-				Landscape->SetCurrentEditingProceduralLayer();
-				Landscape->RequestProceduralContentUpdate(bUpdateHeightmap ? EProceduralContentUpdateFlag::Heightmap_All : EProceduralContentUpdateFlag::Weightmap_All);
+				Landscape->SetEditingLayer();
+				Landscape->RequestLayersContentUpdate(bUpdateHeightmap ? ELandscapeLayersContentUpdateFlag::Heightmap_All : ELandscapeLayersContentUpdateFlag::Weightmap_All);
 			}
 		}
 
@@ -303,9 +303,9 @@ public:
 
 		ALandscape* Landscape = LandscapeInfo->LandscapeActor.Get();
 
-		if (Landscape != nullptr && Landscape->HasProceduralContent && !GetMutableDefault<UEditorExperimentalSettings>()->bProceduralLandscape)
+		if (Landscape != nullptr && Landscape->HasLayersContent && !GetMutableDefault<UEditorExperimentalSettings>()->bLandscapeLayerSystem)
 		{
-			FMessageLog("MapCheck").Warning()->AddToken(FTextToken::Create(LOCTEXT("LandscapeProcedural_ChangingDataWithoutSettings", "This map contains landscape procedural content, modifying the landscape data will result in data loss when the map is reopened with Landscape Procedural settings on. Please enable Landscape Procedural settings before modifying the data.")));
+			FMessageLog("MapCheck").Warning()->AddToken(FTextToken::Create(LOCTEXT("LandscapeLayers_ChangingDataWithoutSettings", "This map contains landscape layer system content, modifying the landscape data will result in data loss when the map is reopened with Landscape Layer System settings on. Please enable Landscape Layer System settings before modifying the data.")));
 			FMessageLog("MapCheck").Open(EMessageSeverity::Warning);
 		}
 
@@ -546,9 +546,9 @@ public:
 
 		ALandscape* Landscape = LandscapeInfo->LandscapeActor.Get();
 
-		if (Landscape != nullptr && Landscape->HasProceduralContent && !GetMutableDefault<UEditorExperimentalSettings>()->bProceduralLandscape)
+		if (Landscape != nullptr && Landscape->HasLayersContent && !GetMutableDefault<UEditorExperimentalSettings>()->bLandscapeLayerSystem)
 		{
-			FMessageLog("MapCheck").Warning()->AddToken(FTextToken::Create(LOCTEXT("LandscapeProcedural_ChangingDataWithoutSettings", "This map contains landscape procedural content, modifying the landscape data will result in data loss when the map is reopened with Landscape Procedural settings on. Please enable Landscape Procedural settings before modifying the data.")));
+			FMessageLog("MapCheck").Warning()->AddToken(FTextToken::Create(LOCTEXT("LandscapeLayers_ChangingDataWithoutSettings", "This map contains landscape layer system content, modifying the landscape data will result in data loss when the map is reopened with Landscape Layer System settings on. Please enable Landscape Layer System settings before modifying the data.")));
 			FMessageLog("MapCheck").Open(EMessageSeverity::Warning);
 		}
 
@@ -680,9 +680,9 @@ public:
 
 		ALandscape* Landscape = this->LandscapeInfo->LandscapeActor.Get();
 
-		if (Landscape != nullptr && Landscape->HasProceduralContent && !GetMutableDefault<UEditorExperimentalSettings>()->bProceduralLandscape)
+		if (Landscape != nullptr && Landscape->HasLayersContent && !GetMutableDefault<UEditorExperimentalSettings>()->bLandscapeLayerSystem)
 		{
-			FMessageLog("MapCheck").Warning()->AddToken(FTextToken::Create(LOCTEXT("LandscapeProcedural_ChangingDataWithoutSettings", "This map contains landscape procedural content, modifying the landscape data will result in data loss when the map is reopened with Landscape Procedural settings on. Please enable Landscape Procedural settings before modifying the data.")));
+			FMessageLog("MapCheck").Warning()->AddToken(FTextToken::Create(LOCTEXT("LandscapeLayers_ChangingDataWithoutSettings", "This map contains landscape layer system content, modifying the landscape data will result in data loss when the map is reopened with Landscape Layer System settings on. Please enable Landscape Layer System settings before modifying the data.")));
 			FMessageLog("MapCheck").Open(EMessageSeverity::Warning);
 		}
 
@@ -922,9 +922,9 @@ public:
 
 		ALandscape* Landscape = this->LandscapeInfo->LandscapeActor.Get();
 
-		if (Landscape != nullptr && Landscape->HasProceduralContent && !GetMutableDefault<UEditorExperimentalSettings>()->bProceduralLandscape)
+		if (Landscape != nullptr && Landscape->HasLayersContent && !GetMutableDefault<UEditorExperimentalSettings>()->bLandscapeLayerSystem)
 		{
-			FMessageLog("MapCheck").Warning()->AddToken(FTextToken::Create(LOCTEXT("LandscapeProcedural_ChangingDataWithoutSettings", "This map contains landscape procedural content, modifying the landscape data will result in data loss when the map is reopened with Landscape Procedural settings on. Please enable Landscape Procedural settings before modifying the data.")));
+			FMessageLog("MapCheck").Warning()->AddToken(FTextToken::Create(LOCTEXT("LandscapeLayers_ChangingDataWithoutSettings", "This map contains landscape layer system content, modifying the landscape data will result in data loss when the map is reopened with Landscape Layer System settings on. Please enable Landscape Layer System settings before modifying the data.")));
 			FMessageLog("MapCheck").Open(EMessageSeverity::Warning);
 		}
 
@@ -1152,9 +1152,9 @@ public:
 
 		ALandscape* Landscape = this->LandscapeInfo->LandscapeActor.Get();
 
-		if (Landscape != nullptr && Landscape->HasProceduralContent && !GetMutableDefault<UEditorExperimentalSettings>()->bProceduralLandscape)
+		if (Landscape != nullptr && Landscape->HasLayersContent && !GetMutableDefault<UEditorExperimentalSettings>()->bLandscapeLayerSystem)
 		{
-			FMessageLog("MapCheck").Warning()->AddToken(FTextToken::Create(LOCTEXT("LandscapeProcedural_ChangingDataWithoutSettings", "This map contains landscape procedural content, modifying the landscape data will result in data loss when the map is reopened with Landscape Procedural settings on. Please enable Landscape Procedural settings before modifying the data.")));
+			FMessageLog("MapCheck").Warning()->AddToken(FTextToken::Create(LOCTEXT("LandscapeLayers_ChangingDataWithoutSettings", "This map contains landscape layer system content, modifying the landscape data will result in data loss when the map is reopened with Landscape Layer System settings on. Please enable Landscape Layer System settings before modifying the data.")));
 			FMessageLog("MapCheck").Open(EMessageSeverity::Warning);
 		}
 
diff --git a/Engine/Source/Editor/LandscapeEditor/Private/LandscapeEdModeXYOffset.cpp b/Engine/Source/Editor/LandscapeEditor/Private/LandscapeEdModeXYOffset.cpp
--- a/Engine/Source/Editor/LandscapeEditor/Private/LandscapeEdModeXYOffset.cpp
+++ b/Engine/Source/Editor/LandscapeEditor/Private/LandscapeEdModeXYOffset.cpp
@@ -543,9 +543,9 @@ public:
 
 		ALandscape* Landscape = LandscapeInfo->LandscapeActor.Get();
 
-		if (Landscape != nullptr && Landscape->HasProceduralContent && !GetMutableDefault<UEditorExperimentalSettings>()->bProceduralLandscape)
+		if (Landscape != nullptr && Landscape->HasLayersContent && !GetMutableDefault<UEditorExperimentalSettings>()->bLandscapeLayerSystem)
 		{
-			FMessageLog("MapCheck").Warning()->AddToken(FTextToken::Create(LOCTEXT("LandscapeProcedural_ChangingDataWithoutSettings", "This map contains landscape procedural content, modifying the landscape data will result in data loss when the map is reopened with Landscape Procedural settings on. Please enable Landscape Procedural settings before modifying the data.")));
+			FMessageLog("MapCheck").Warning()->AddToken(FTextToken::Create(LOCTEXT("LandscapeLayers_ChangingDataWithoutSettings", "This map contains landscape layer system content, modifying the landscape data will result in data loss when the map is reopened with Landscape Layer System settings on. Please enable Landscape Layer System settings before modifying the data.")));
 			FMessageLog("MapCheck").Open(EMessageSeverity::Warning);
 		}
 
diff --git a/Engine/Source/Editor/LandscapeEditor/Private/LandscapeEditorDetailCustomization_Layers.cpp b/Engine/Source/Editor/LandscapeEditor/Private/LandscapeEditorDetailCustomization_Layers.cpp
new file mode 100644
--- /dev/null
+++ b/Engine/Source/Editor/LandscapeEditor/Private/LandscapeEditorDetailCustomization_Layers.cpp
@@ -0,0 +1,537 @@
+// Copyright 1998-2019 Epic Games, Inc. All Rights Reserved.
+
+#include "LandscapeEditorDetailCustomization_Layers.h"
+#include "IDetailChildrenBuilder.h"
+#include "Framework/Commands/UIAction.h"
+#include "Widgets/Text/STextBlock.h"
+#include "Framework/MultiBox/MultiBoxBuilder.h"
+#include "Misc/MessageDialog.h"
+#include "Modules/ModuleManager.h"
+#include "Brushes/SlateColorBrush.h"
+#include "Layout/WidgetPath.h"
+#include "SlateOptMacros.h"
+#include "Framework/Application/MenuStack.h"
+#include "Framework/Application/SlateApplication.h"
+#include "Widgets/Images/SImage.h"
+#include "Widgets/Layout/SBox.h"
+#include "Widgets/Input/SButton.h"
+#include "Widgets/Notifications/SErrorText.h"
+#include "Widgets/Input/SComboButton.h"
+#include "Widgets/Input/SCheckBox.h"
+#include "EditorModeManager.h"
+#include "EditorModes.h"
+#include "LandscapeEditorModule.h"
+#include "LandscapeEditorObject.h"
+#include "Landscape.h"
+
+#include "DetailLayoutBuilder.h"
+#include "IDetailPropertyRow.h"
+#include "DetailCategoryBuilder.h"
+#include "PropertyCustomizationHelpers.h"
+
+#include "SLandscapeEditor.h"
+#include "Dialogs/DlgPickAssetPath.h"
+#include "ObjectTools.h"
+#include "ScopedTransaction.h"
+#include "DesktopPlatformModule.h"
+#include "AssetRegistryModule.h"
+
+#include "LandscapeRender.h"
+#include "Materials/MaterialExpressionLandscapeVisibilityMask.h"
+#include "LandscapeEdit.h"
+#include "IDetailGroup.h"
+#include "Widgets/SBoxPanel.h"
+#include "Editor/EditorStyle/Private/SlateEditorStyle.h"
+#include "LandscapeEditorDetailCustomization_TargetLayers.h"
+#include "Widgets/Input/SEditableText.h"
+#include "Widgets/Input/SNumericEntryBox.h"
+#include "Widgets/Text/SInlineEditableTextBlock.h"
+#include "LandscapeEditorCommands.h"
+
+#define LOCTEXT_NAMESPACE "LandscapeEditor.Layers"
+
+TSharedRef<IDetailCustomization> FLandscapeEditorDetailCustomization_Layers::MakeInstance()
+{
+	return MakeShareable(new FLandscapeEditorDetailCustomization_Layers);
+}
+
+BEGIN_SLATE_FUNCTION_BUILD_OPTIMIZATION
+void FLandscapeEditorDetailCustomization_Layers::CustomizeDetails(IDetailLayoutBuilder& DetailBuilder)
+{
+	IDetailCategoryBuilder& LayerCategory = DetailBuilder.EditCategory("Layers");
+
+	FEdModeLandscape* LandscapeEdMode = GetEditorMode();
+	if (LandscapeEdMode && LandscapeEdMode->CurrentToolMode != nullptr)
+	{
+		LayerCategory.AddCustomBuilder(MakeShareable(new FLandscapeEditorCustomNodeBuilder_Layers(DetailBuilder.GetThumbnailPool().ToSharedRef())));
+	}
+}
+END_SLATE_FUNCTION_BUILD_OPTIMIZATION
+
+//////////////////////////////////////////////////////////////////////////
+
+FEdModeLandscape* FLandscapeEditorCustomNodeBuilder_Layers::GetEditorMode()
+{
+	return (FEdModeLandscape*)GLevelEditorModeTools().GetActiveMode(FBuiltinEditorModes::EM_Landscape);
+}
+
+FLandscapeEditorCustomNodeBuilder_Layers::FLandscapeEditorCustomNodeBuilder_Layers(TSharedRef<FAssetThumbnailPool> InThumbnailPool)
+	: ThumbnailPool(InThumbnailPool)
+{
+}
+
+FLandscapeEditorCustomNodeBuilder_Layers::~FLandscapeEditorCustomNodeBuilder_Layers()
+{
+}
+
+void FLandscapeEditorCustomNodeBuilder_Layers::SetOnRebuildChildren(FSimpleDelegate InOnRegenerateChildren)
+{
+}
+
+void FLandscapeEditorCustomNodeBuilder_Layers::GenerateHeaderRowContent(FDetailWidgetRow& NodeRow)
+{
+	FEdModeLandscape* LandscapeEdMode = GetEditorMode();
+
+	if (LandscapeEdMode == NULL)
+	{
+		return;
+	}
+
+	NodeRow.NameWidget
+		[
+			SNew(STextBlock)
+			.Font(IDetailLayoutBuilder::GetDetailFont())
+		.Text(FText::FromString(TEXT("Layers")))
+		];
+}
+
+BEGIN_SLATE_FUNCTION_BUILD_OPTIMIZATION
+void FLandscapeEditorCustomNodeBuilder_Layers::GenerateChildContent(IDetailChildrenBuilder& ChildrenBuilder)
+{
+	FEdModeLandscape* LandscapeEdMode = GetEditorMode();
+	if (LandscapeEdMode != NULL)
+	{
+		TSharedPtr<SDragAndDropVerticalBox> LayerList = SNew(SDragAndDropVerticalBox)
+			.OnCanAcceptDrop(this, &FLandscapeEditorCustomNodeBuilder_Layers::HandleCanAcceptDrop)
+			.OnAcceptDrop(this, &FLandscapeEditorCustomNodeBuilder_Layers::HandleAcceptDrop)
+			.OnDragDetected(this, &FLandscapeEditorCustomNodeBuilder_Layers::HandleDragDetected);
+
+		LayerList->SetDropIndicator_Above(*FEditorStyle::GetBrush("LandscapeEditor.TargetList.DropZone.Above"));
+		LayerList->SetDropIndicator_Below(*FEditorStyle::GetBrush("LandscapeEditor.TargetList.DropZone.Below"));
+
+		ChildrenBuilder.AddCustomRow(FText::FromString(FString(TEXT("Layers"))))
+			.Visibility(EVisibility::Visible)
+			[
+				LayerList.ToSharedRef()
+			];
+
+		InlineTextBlocks.Empty();
+		InlineTextBlocks.AddDefaulted(LandscapeEdMode->GetLayerCount());
+		for (int32 i = 0; i < LandscapeEdMode->GetLayerCount(); ++i)
+		{
+			TSharedPtr<SWidget> GeneratedRowWidget = GenerateRow(i);
+
+			if (GeneratedRowWidget.IsValid())
+			{
+				LayerList->AddSlot()
+					.AutoHeight()
+					[
+						GeneratedRowWidget.ToSharedRef()
+					];
+			}
+		}
+	}
+}
+END_SLATE_FUNCTION_BUILD_OPTIMIZATION
+
+BEGIN_SLATE_FUNCTION_BUILD_OPTIMIZATION
+TSharedPtr<SWidget> FLandscapeEditorCustomNodeBuilder_Layers::GenerateRow(int32 InLayerIndex)
+{
+	FEdModeLandscape* LandscapeEdMode = GetEditorMode();
+	TSharedPtr<SWidget> RowWidget = SNew(SLandscapeEditorSelectableBorder)
+		.Padding(0)
+		.VAlign(VAlign_Center)
+		.OnContextMenuOpening(this, &FLandscapeEditorCustomNodeBuilder_Layers::OnLayerContextMenuOpening, InLayerIndex)
+		.OnSelected(this, &FLandscapeEditorCustomNodeBuilder_Layers::OnLayerSelectionChanged, InLayerIndex)
+		.IsSelected(TAttribute<bool>::Create(TAttribute<bool>::FGetter::CreateSP(this, &FLandscapeEditorCustomNodeBuilder_Layers::IsLayerSelected, InLayerIndex)))
+		.Visibility(EVisibility::Visible)
+		[
+			SNew(SHorizontalBox)
+
+			+ SHorizontalBox::Slot()
+		.AutoWidth()
+		.VAlign(VAlign_Center)
+		[
+			SNew(SButton)
+			.ButtonStyle(FEditorStyle::Get(), "NoBorder")
+		.OnClicked(this, &FLandscapeEditorCustomNodeBuilder_Layers::OnToggleLock, InLayerIndex)
+		.ToolTipText(LOCTEXT("FLandscapeEditorCustomNodeBuilder_LayerLock", "Locks the current layer"))
+		[
+			SNew(SImage)
+			.Image(this, &FLandscapeEditorCustomNodeBuilder_Layers::GetLockBrushForLayer, InLayerIndex)
+		]
+		]
+
+	+ SHorizontalBox::Slot()
+		.AutoWidth()
+		.VAlign(VAlign_Center)
+		[
+			SNew(SButton)
+			.ContentPadding(0)
+		.ButtonStyle(FEditorStyle::Get(), "NoBorder")
+		.OnClicked(this, &FLandscapeEditorCustomNodeBuilder_Layers::OnToggleVisibility, InLayerIndex)
+		.ToolTipText(LOCTEXT("FLandscapeEditorCustomNodeBuilder_LayerVisibility", "Toggle Layer Visibility"))
+		.HAlign(HAlign_Center)
+		.VAlign(VAlign_Center)
+		.Content()
+		[
+			SNew(SImage)
+			.Image(this, &FLandscapeEditorCustomNodeBuilder_Layers::GetVisibilityBrushForLayer, InLayerIndex)
+		]
+		]
+
+	+ SHorizontalBox::Slot()
+		.FillWidth(1.0)
+		.VAlign(VAlign_Center)
+		.Padding(4, 0)
+		[
+			SAssignNew(InlineTextBlocks[InLayerIndex], SInlineEditableTextBlock)
+			.Text(this, &FLandscapeEditorCustomNodeBuilder_Layers::GetLayerText, InLayerIndex)
+		.ToolTipText(LOCTEXT("FLandscapeEditorCustomNodeBuilder_Layers_tooltip", "Name of the Layer"))
+		.OnVerifyTextChanged(FOnVerifyTextChanged::CreateSP(this, &FLandscapeEditorCustomNodeBuilder_Layers::CanRenameLayerTo, InLayerIndex))
+		.OnTextCommitted(FOnTextCommitted::CreateSP(this, &FLandscapeEditorCustomNodeBuilder_Layers::SetLayerName, InLayerIndex))
+		]
+
+	+ SHorizontalBox::Slot()
+		.FillWidth(1.0f)
+		.VAlign(VAlign_Center)
+		.Padding(0, 2)
+		.HAlign(HAlign_Right)
+		[
+			SNew(SHorizontalBox)
+			+ SHorizontalBox::Slot()
+		.Padding(0)
+		.FillWidth(1.0f)
+		.VAlign(VAlign_Center)
+		.HAlign(HAlign_Left)
+		[
+			SNew(STextBlock)
+			.Text(LOCTEXT("FLandscapeEditorCustomNodeBuilder_LayerAlpha", "Alpha"))
+		]
+	+ SHorizontalBox::Slot()
+		.VAlign(VAlign_Center)
+		.Padding(0, 2)
+		.HAlign(HAlign_Left)
+		.FillWidth(1.0f)
+		[
+			SNew(SNumericEntryBox<float>)
+			.AllowSpin(true)
+		.MinValue(0.0f)
+		.MaxValue(100.0f)
+		.MaxSliderValue(100.0f)
+		.MinDesiredValueWidth(60.0f)
+		.Value(this, &FLandscapeEditorCustomNodeBuilder_Layers::GetLayerAlpha, InLayerIndex)
+		.OnValueChanged(this, &FLandscapeEditorCustomNodeBuilder_Layers::SetLayerAlpha, InLayerIndex)
+		.IsEnabled(true)
+		]
+		]
+		];
+
+	return RowWidget;
+}
+END_SLATE_FUNCTION_BUILD_OPTIMIZATION
+
+FText FLandscapeEditorCustomNodeBuilder_Layers::GetLayerText(int32 InLayerIndex) const
+{
+	FEdModeLandscape* LandscapeEdMode = GetEditorMode();
+	if (LandscapeEdMode)
+	{
+		return FText::FromName(LandscapeEdMode->GetLayerName(InLayerIndex));
+	}
+
+	return FText::FromString(TEXT("None"));
+}
+
+bool FLandscapeEditorCustomNodeBuilder_Layers::IsLayerSelected(int32 InLayerIndex)
+{
+	FEdModeLandscape* LandscapeEdMode = GetEditorMode();
+	if (LandscapeEdMode)
+	{
+		return LandscapeEdMode->GetCurrentLayerIndex() == InLayerIndex;
+	}
+
+	return false;
+}
+
+bool FLandscapeEditorCustomNodeBuilder_Layers::CanRenameLayerTo(const FText& InNewText, FText& OutErrorMessage, int32 InLayerIndex)
+{
+	FEdModeLandscape* LandscapeEdMode = GetEditorMode();
+	if (LandscapeEdMode)
+	{
+		if (!LandscapeEdMode->CanRenameLayerTo(InLayerIndex, *InNewText.ToString()))
+		{
+			OutErrorMessage = LOCTEXT("RenameFailed_AlreadyExists", "This layer already exists");
+			return false;
+		}
+	}
+	return true;
+}
+
+void FLandscapeEditorCustomNodeBuilder_Layers::SetLayerName(const FText& InText, ETextCommit::Type InCommitType, int32 InLayerIndex)
+{
+	FEdModeLandscape* LandscapeEdMode = GetEditorMode();
+	if (LandscapeEdMode)
+	{
+		const FScopedTransaction Transaction(LOCTEXT("Landscape_Layers_Rename", "Rename Layer"));
+		LandscapeEdMode->SetLayerName(InLayerIndex, *InText.ToString());
+	}
+}
+
+TSharedPtr<SWidget> FLandscapeEditorCustomNodeBuilder_Layers::OnLayerContextMenuOpening(int32 InLayerIndex)
+{
+	FEdModeLandscape* LandscapeEdMode = GetEditorMode();
+	ALandscape* Landscape = LandscapeEdMode ? LandscapeEdMode->GetLandscape() : nullptr;
+	if (LandscapeEdMode && Landscape)
+	{
+		FLandscapeLayer* Layer = LandscapeEdMode->GetLayer(InLayerIndex);
+		TSharedRef<FLandscapeEditorCustomNodeBuilder_Layers> SharedThis = AsShared();
+		FMenuBuilder MenuBuilder(true, NULL);
+		MenuBuilder.BeginSection("LandscapeEditorLayerActions", LOCTEXT("LandscapeEditorLayerActions.Heading", "Layers"));
+		{
+			// Create Layer
+			FUIAction CreateLayerAction = FUIAction(FExecuteAction::CreateLambda([SharedThis] { SharedThis->CreateLayer(); }));
+			MenuBuilder.AddMenuEntry(LOCTEXT("CreateLayer", "Create"), LOCTEXT("CreateLayerTooltip", "Create Layer"), FSlateIcon(), CreateLayerAction);
+
+			if (Layer)
+			{
+				// Rename Layer
+				FUIAction RenameLayerAction = FUIAction(FExecuteAction::CreateLambda([SharedThis, InLayerIndex] { SharedThis->RenameLayer(InLayerIndex); }));
+				MenuBuilder.AddMenuEntry(LOCTEXT("RenameLayer", "Rename..."), LOCTEXT("RenameLayerTooltip", "Rename Layer"), FSlateIcon(), RenameLayerAction);
+
+				if (!Layer->bLocked)
+				{
+					// Clear Layer
+					FUIAction ClearLayerAction = FUIAction(FExecuteAction::CreateLambda([SharedThis, InLayerIndex] { SharedThis->ClearLayer(InLayerIndex); }));
+					MenuBuilder.AddMenuEntry(LOCTEXT("ClearLayer", "Clear..."), LOCTEXT("ClearLayerTooltip", "Clear Layer"), FSlateIcon(), ClearLayerAction);
+
+					if (Landscape->LandscapeLayers.Num() > 1)
+					{
+						// Delete Layer
+						FUIAction DeleteLayerAction = FUIAction(FExecuteAction::CreateLambda([SharedThis, InLayerIndex] { SharedThis->DeleteLayer(InLayerIndex); }));
+						MenuBuilder.AddMenuEntry(LOCTEXT("DeleteLayer", "Delete..."), LOCTEXT("DeleteLayerTooltip", "Delete Layer"), FSlateIcon(), DeleteLayerAction);
+					}
+				}
+			}
+		}
+		MenuBuilder.EndSection();
+
+		MenuBuilder.BeginSection("LandscapeEditorLayerVisibility", LOCTEXT("LandscapeEditorLayerVisibility.Heading", "Visibility"));
+		{
+			if (Layer)
+			{
+				if (Layer->bVisible)
+				{
+					// Hide Selected Layer
+					FUIAction HideSelectedLayerAction = FUIAction(FExecuteAction::CreateLambda([SharedThis, InLayerIndex] { SharedThis->OnToggleVisibility(InLayerIndex); }));
+					MenuBuilder.AddMenuEntry(LOCTEXT("HideSelectedLayer", "Hide Selected"), LOCTEXT("HideSelectedLayerTooltip", "Hide Selected Layer"), FSlateIcon(), HideSelectedLayerAction);
+				}
+				else
+				{
+					// Show Selected Layer
+					FUIAction ShowSelectedLayerAction = FUIAction(FExecuteAction::CreateLambda([SharedThis, InLayerIndex] { SharedThis->OnToggleVisibility(InLayerIndex); }));
+					MenuBuilder.AddMenuEntry(LOCTEXT("ShowSelectedLayer", "Show Selected"), LOCTEXT("ShowSelectedLayerTooltip", "Show Selected Layer"), FSlateIcon(), ShowSelectedLayerAction);
+				}
+
+				// Show Only Selected Layer
+				FUIAction ShowOnlySelectedLayerAction = FUIAction(FExecuteAction::CreateLambda([SharedThis, InLayerIndex] { SharedThis->ShowOnlySelectedLayer(InLayerIndex); }));
+				MenuBuilder.AddMenuEntry(LOCTEXT("ShowOnlySelectedLayer", "Show Only Selected"), LOCTEXT("ShowOnlySelectedLayerTooltip", "Show Only Selected Layer"), FSlateIcon(), ShowOnlySelectedLayerAction);
+			}
+
+			// Show All Layers
+			FUIAction ShowAllLayersAction = FUIAction(FExecuteAction::CreateLambda([SharedThis] { SharedThis->ShowAllLayers(); }));
+			MenuBuilder.AddMenuEntry(LOCTEXT("ShowAllLayers", "Show All Layers"), LOCTEXT("ShowAllLayersTooltip", "Show All Layers"), FSlateIcon(), ShowAllLayersAction);
+		}
+		MenuBuilder.EndSection();
+
+		return MenuBuilder.MakeWidget();
+	}
+	return NULL;
+}
+
+void FLandscapeEditorCustomNodeBuilder_Layers::RenameLayer(int32 InLayerIndex)
+{
+	if (InlineTextBlocks.IsValidIndex(InLayerIndex) && InlineTextBlocks[InLayerIndex].IsValid())
+	{
+		InlineTextBlocks[InLayerIndex]->EnterEditingMode();
+	}
+}
+
+void FLandscapeEditorCustomNodeBuilder_Layers::ClearLayer(int32 InLayerIndex)
+{
+	FEdModeLandscape* LandscapeEdMode = GetEditorMode();
+	ALandscape* Landscape = LandscapeEdMode ? LandscapeEdMode->GetLandscape() : nullptr;
+	if (Landscape)
+	{
+		FLandscapeLayer* Layer = LandscapeEdMode->GetLayer(InLayerIndex);
+		if (Layer)
+		{
+			EAppReturnType::Type Result = FMessageDialog::Open(EAppMsgType::YesNo, FText::Format(LOCTEXT("LandscapeMode_Message", "The layer {0} content will be completely cleared.  Continue?"), FText::FromName(Layer->Name)));
+			if (Result == EAppReturnType::Yes)
+			{
+				const FScopedTransaction Transaction(LOCTEXT("Landscape_Layers_Clean", "Clear Layer"));
+				Landscape->ClearLayer(InLayerIndex);
+			}
+		}
+	}
+}
+
+void FLandscapeEditorCustomNodeBuilder_Layers::DeleteLayer(int32 InLayerIndex)
+{
+	FEdModeLandscape* LandscapeEdMode = GetEditorMode();
+	ALandscape* Landscape = LandscapeEdMode ? LandscapeEdMode->GetLandscape() : nullptr;
+	if (Landscape && Landscape->LandscapeLayers.Num() > 1)
+	{
+		FLandscapeLayer* Layer = LandscapeEdMode->GetLayer(InLayerIndex);
+		if (Layer)
+		{
+			EAppReturnType::Type Result = FMessageDialog::Open(EAppMsgType::YesNo, FText::Format(LOCTEXT("LandscapeMode_Message", "The layer {0} will be deleted.  Continue?"), FText::FromName(Layer->Name)));
+			if (Result == EAppReturnType::Yes)
+			{
+				const FScopedTransaction Transaction(LOCTEXT("Landscape_Layers_Delete", "Delete Layer"));
+				Landscape->DeleteLayer(InLayerIndex);
+				int32 NewLayerSelectionIndex = Landscape->GetLayer(InLayerIndex) ? InLayerIndex : 0;
+				OnLayerSelectionChanged(NewLayerSelectionIndex);
+				LandscapeEdMode->RefreshDetailPanel();
+			}
+		}
+	}
+}
+
+void FLandscapeEditorCustomNodeBuilder_Layers::ShowOnlySelectedLayer(int32 InLayerIndex)
+{
+	FEdModeLandscape* LandscapeEdMode = GetEditorMode();
+	ALandscape* Landscape = LandscapeEdMode ? LandscapeEdMode->GetLandscape() : nullptr;
+	if (Landscape)
+	{
+		const FScopedTransaction Transaction(LOCTEXT("ShowOnlySelectedLayer", "Show Only Selected Layer"));
+		Landscape->ShowOnlySelectedLayer(InLayerIndex);
+	}
+}
+
+void FLandscapeEditorCustomNodeBuilder_Layers::ShowAllLayers()
+{
+	FEdModeLandscape* LandscapeEdMode = GetEditorMode();
+	ALandscape* Landscape = LandscapeEdMode ? LandscapeEdMode->GetLandscape() : nullptr;
+	if (Landscape)
+	{
+		const FScopedTransaction Transaction(LOCTEXT("ShowAllLayers", "Show All Layers"));
+		Landscape->ShowAllLayers();
+	}
+}
+
+void FLandscapeEditorCustomNodeBuilder_Layers::CreateLayer()
+{
+	FEdModeLandscape* LandscapeEdMode = GetEditorMode();
+	ALandscape* Landscape = LandscapeEdMode ? LandscapeEdMode->GetLandscape() : nullptr;
+	if (Landscape)
+	{
+		{
+			const FScopedTransaction Transaction(LOCTEXT("Landscape_Layers_Create", "Create Layer"));
+			Landscape->CreateLayer();
+		}
+		LandscapeEdMode->RefreshDetailPanel();
+	}
+}
+
+void FLandscapeEditorCustomNodeBuilder_Layers::OnLayerSelectionChanged(int32 InLayerIndex)
+{
+	FEdModeLandscape* LandscapeEdMode = GetEditorMode();
+	if (LandscapeEdMode)
+	{
+		LandscapeEdMode->SetCurrentLayer(InLayerIndex);
+		LandscapeEdMode->UpdateTargetList();
+	}
+}
+
+TOptional<float> FLandscapeEditorCustomNodeBuilder_Layers::GetLayerAlpha(int32 InLayerIndex) const
+{
+	FEdModeLandscape* LandscapeEdMode = GetEditorMode();
+
+	if (LandscapeEdMode)
+	{
+		return LandscapeEdMode->GetLayerAlpha(InLayerIndex);
+	}
+
+	return 1.0f;
+}
+
+void FLandscapeEditorCustomNodeBuilder_Layers::SetLayerAlpha(float InAlpha, int32 InLayerIndex)
+{
+	FEdModeLandscape* LandscapeEdMode = GetEditorMode();
+
+	if (LandscapeEdMode)
+	{
+		const FScopedTransaction Transaction(LOCTEXT("Landscape_Layers_SetAlpha", "Set Layer Alpha"));
+		LandscapeEdMode->SetLayerAlpha(InLayerIndex, InAlpha);
+	}
+}
+
+FReply FLandscapeEditorCustomNodeBuilder_Layers::OnToggleVisibility(int32 InLayerIndex)
+{
+	FEdModeLandscape* LandscapeEdMode = GetEditorMode();
+	if (LandscapeEdMode)
+	{
+		LandscapeEdMode->SetLayerVisibility(!LandscapeEdMode->IsLayerVisible(InLayerIndex), InLayerIndex);
+	}
+	return FReply::Handled();
+}
+
+const FSlateBrush* FLandscapeEditorCustomNodeBuilder_Layers::GetVisibilityBrushForLayer(int32 InLayerIndex) const
+{
+	FEdModeLandscape* LandscapeEdMode = GetEditorMode();
+	bool bIsVisible = LandscapeEdMode && LandscapeEdMode->IsLayerVisible(InLayerIndex);
+	return bIsVisible ? FEditorStyle::GetBrush("Level.VisibleIcon16x") : FEditorStyle::GetBrush("Level.NotVisibleIcon16x");
+}
+
+FReply FLandscapeEditorCustomNodeBuilder_Layers::OnToggleLock(int32 InLayerIndex)
+{
+	FEdModeLandscape* LandscapeEdMode = GetEditorMode();
+	if (LandscapeEdMode)
+	{
+		LandscapeEdMode->SetLayerLocked(InLayerIndex, !LandscapeEdMode->IsLayerLocked(InLayerIndex));
+	}
+	return FReply::Handled();
+}
+
+const FSlateBrush* FLandscapeEditorCustomNodeBuilder_Layers::GetLockBrushForLayer(int32 InLayerIndex) const
+{
+	FEdModeLandscape* LandscapeEdMode = GetEditorMode();
+	bool bIsLocked = LandscapeEdMode && LandscapeEdMode->IsLayerLocked(InLayerIndex);
+	return bIsLocked ? FEditorStyle::GetBrush(TEXT("PropertyWindow.Locked")) : FEditorStyle::GetBrush(TEXT("PropertyWindow.Unlocked"));
+}
+
+FReply FLandscapeEditorCustomNodeBuilder_Layers::HandleDragDetected(const FGeometry& MyGeometry, const FPointerEvent& MouseEvent, int32 SlotIndex, SVerticalBox::FSlot* Slot)
+{
+	FEdModeLandscape* LandscapeEdMode = GetEditorMode();
+
+	if (LandscapeEdMode != nullptr)
+	{
+		// TODO: handle drag & drop
+	}
+
+	return FReply::Unhandled();
+}
+
+TOptional<SDragAndDropVerticalBox::EItemDropZone> FLandscapeEditorCustomNodeBuilder_Layers::HandleCanAcceptDrop(const FDragDropEvent& DragDropEvent, SDragAndDropVerticalBox::EItemDropZone DropZone, SVerticalBox::FSlot* Slot)
+{
+	// TODO: handle drag & drop
+	return TOptional<SDragAndDropVerticalBox::EItemDropZone>();
+}
+
+FReply FLandscapeEditorCustomNodeBuilder_Layers::HandleAcceptDrop(FDragDropEvent const& DragDropEvent, SDragAndDropVerticalBox::EItemDropZone DropZone, int32 SlotIndex, SVerticalBox::FSlot* Slot)
+{
+	// TODO: handle drag & drop
+	return FReply::Unhandled();
+}
+
+#undef LOCTEXT_NAMESPACE
