Fixed compilation error

diff --git a/Engine/Source/Editor/LandscapeEditor/Private/LandscapeEdModeComponentTools.cpp b/Engine/Source/Editor/LandscapeEditor/Private/LandscapeEdModeComponentTools.cpp
--- a/Engine/Source/Editor/LandscapeEditor/Private/LandscapeEdModeComponentTools.cpp
+++ b/Engine/Source/Editor/LandscapeEditor/Private/LandscapeEdModeComponentTools.cpp
@@ -958,7 +958,7 @@ public:
 				NewComponents[Idx]->RegisterComponent();
 			}
 
-			if (LandscapeInfo->LandscapeActor.IsValid() && LandscapeInfo->LandscapeActor.Get()->HasLayersContent && !GetMutableDefault<UEditorExperimentalSettings>()->HasLayersContent)
+			if (LandscapeInfo->LandscapeActor.IsValid() && LandscapeInfo->LandscapeActor.Get()->HasLayersContent && !GetMutableDefault<UEditorExperimentalSettings>()->bLandscapeLayerSystem)
 			{
 				FMessageLog("MapCheck").Warning()->AddToken(FTextToken::Create(LOCTEXT("LandscapeLayers_ChangingDataWithoutSettings", "This map contains landscape layer system content, modifying the landscape data will result in data loss when the map is reopened with Landscape Layer System settings on. Please enable Landscape Layer System settings before modifying the data.")));
 
