- Limited support of Flatten & Ramp Tool for Landscape Layer System (limited to first layer)
- Disabled Retopologize Landscape Tool for Landscape Layer System
- Erosion & HydroErosion Tool now supported Landscape Layer System
- Added message tip to show warnings of unsupported edition cases in the Landscape Layer System
- Fixed bug in Renaming of a Landscape Layer

diff --git a/Engine/Source/Editor/LandscapeEditor/Private/LandscapeEdMode.cpp b/Engine/Source/Editor/LandscapeEditor/Private/LandscapeEdMode.cpp
--- a/Engine/Source/Editor/LandscapeEditor/Private/LandscapeEdMode.cpp
+++ b/Engine/Source/Editor/LandscapeEditor/Private/LandscapeEdMode.cpp
@@ -1896,55 +1896,29 @@ bool FEdModeLandscape::InputKey(FEditorViewportClient* ViewportClient, FViewport
 					FVector HitLocation;
 					if (LandscapeMouseTrace(ViewportClient, HitLocation))
 					{
-						const FName CurrentToolName = CurrentTool->GetToolName();
-						ALandscape* Landscape = GetLandscape();
-						bool bLandscapeLayerSystem = GetMutableDefault<UEditorExperimentalSettings>()->bLandscapeLayerSystem;
-
-						if (bLandscapeLayerSystem && !GetCurrentLayer())
-						{
-							FMessageDialog::Open(EAppMsgType::Ok, NSLOCTEXT("UnrealEd", "LandscapeInvalidLayer", "No layer selected. You must first chose a layer to work on."));
-						}
-						else if (bLandscapeLayerSystem && GetCurrentLayer() && !GetCurrentLayer()->bVisible)
-						{
-							FMessageDialog::Open(EAppMsgType::Ok, NSLOCTEXT("UnrealEd", "LandscapeLayerHidden", "Painting in a hidden layer is not allowed."));
-						}
-						else if (bLandscapeLayerSystem && GetCurrentLayer() && GetCurrentLayer()->bLocked)
-						{
-							FMessageDialog::Open(EAppMsgType::Ok, NSLOCTEXT("UnrealEd", "LandscapeLayerLocked", "This layer is locked. You must unlock it before you can work on this layer."));
-						}
-						else if (bLandscapeLayerSystem && (CurrentTool != (FLandscapeTool*)SplinesTool) && GetCurrentLayer() && Landscape && (GetCurrentLayer() == Landscape->GetLandscapeSplinesReservedLayer()))
-						{
-							FMessageDialog::Open(EAppMsgType::Ok, NSLOCTEXT("UnrealEd", "LandscapeLayerReservedForSplines", "This layer is reserved for Landscape Splines."));
-						}
-						else if (CurrentToolTarget.TargetType == ELandscapeToolTargetType::Weightmap && CurrentToolTarget.LayerInfo == NULL && CurrentToolName != TEXT("BPCustom"))
+						FText Reason;
+						if (!CanEditLayer(&Reason))
 						{
-							FMessageDialog::Open(EAppMsgType::Ok,
-								NSLOCTEXT("UnrealEd", "LandscapeNeedToCreateLayerInfo", "This layer has no layer info assigned yet. You must create or assign a layer info before you can paint this layer."));
-							// TODO: FName to LayerInfo: do we want to create the layer info here?
-							//if (FMessageDialog::Open(EAppMsgType::YesNo, NSLOCTEXT("UnrealEd", "LandscapeNeedToCreateLayerInfo", "This layer has no layer info assigned yet. You must create or assign a layer info before you can paint this layer.")) == EAppReturnType::Yes)
-							//{
-							//	CurrentToolTarget.LandscapeInfo->LandscapeProxy->CreateLayerInfo(*CurrentToolTarget.PlaceholderLayerName.ToString());
-							//}
+							FMessageDialog::Open(EAppMsgType::Ok, Reason);
+							return true;
 						}
-						else
-						{
-							Viewport->CaptureMouse(true);
 
-							if (CurrentTool->CanToolBeActivated())
+						Viewport->CaptureMouse(true);
+
+						if (CurrentTool->CanToolBeActivated())
+						{
+							bool bToolActive = CurrentTool->BeginTool(ViewportClient, CurrentToolTarget, HitLocation);
+							if (bToolActive)
 							{
-								bool bToolActive = CurrentTool->BeginTool(ViewportClient, CurrentToolTarget, HitLocation);
-								if (bToolActive)
-								{
-									ToolActiveViewport = Viewport;
-								}
-								else
-								{
-									ToolActiveViewport = nullptr;
-									Viewport->CaptureMouse(false);
-								}
-								ViewportClient->Invalidate(false, false);
-								return bToolActive;
+								ToolActiveViewport = Viewport;
+							}
+							else
+							{
+								ToolActiveViewport = nullptr;
+								Viewport->CaptureMouse(false);
 							}
+							ViewportClient->Invalidate(false, false);
+							return bToolActive;
 						}
 					}
 				}
@@ -4382,6 +4356,80 @@ void FEdModeLandscape::AutoUpdateDirtyLandscapeSplines()
 	}
 }
 
+bool FEdModeLandscape::CanEditLayer(FText* Reason)
+{
+	if (GetMutableDefault<UEditorExperimentalSettings>()->bLandscapeLayerSystem)
+	{
+		ALandscape* Landscape = GetLandscape();
+		FLandscapeLayer* CurrentLayer = GetCurrentLayer();
+		if (!CurrentLayer)
+		{
+			if (Reason)
+			{
+				*Reason = NSLOCTEXT("UnrealEd", "LandscapeInvalidLayer", "No layer selected. You must first chose a layer to work on.");
+			}
+			return false;
+		}
+		else if (!CurrentLayer->bVisible)
+		{
+			if (Reason)
+			{
+				*Reason = NSLOCTEXT("UnrealEd", "LandscapeLayerHidden", "Painting or sculping in a hidden layer is not allowed.");
+			}
+			return false;
+		}
+		else if (CurrentLayer->bLocked)
+		{
+			if (Reason)
+			{
+				*Reason = NSLOCTEXT("UnrealEd", "LandscapeLayerLocked", "This layer is locked. You must unlock it before you can work on this layer.");
+			}
+			return false;
+		}
+		else if (CurrentTool && (CurrentTool != (FLandscapeTool*)SplinesTool) && Landscape && (CurrentLayer == Landscape->GetLandscapeSplinesReservedLayer()))
+		{
+			if (Reason)
+			{
+				*Reason = NSLOCTEXT("UnrealEd", "LandscapeLayerReservedForSplines", "This layer is reserved for Landscape Splines.");
+			}
+			return false;
+		}
+		else if (CurrentTool &&
+				 (GetCurrentLayerIndex() != 0) &&
+				 (CurrentTool->GetToolName() == FName("Ramp") || CurrentTool->GetToolName() == FName("Flatten")))
+		{
+			if (Reason)
+			{
+				FLandscapeLayer* Layer = GetLayer(0);
+				*Reason = FText::Format(NSLOCTEXT("UnrealEd", "LandscapeLayersToolAvailableOnlyOnFirstLayer", "{0} Tool is only available on the first layer {1}."), CurrentTool->GetDisplayName(), Layer ? FText::FromName(Layer->Name) : FText::GetEmpty());
+			}
+			return false;
+		}
+		else if (CurrentTool && CurrentTool->GetToolName() == FName("Retopologize"))
+		{
+			if (Reason)
+			{
+				*Reason = FText::Format(NSLOCTEXT("UnrealEd", "LandscapeLayersNoSupportForRetopologize", "{0} Tool is not available with the Landscape Layer System."), CurrentTool->GetDisplayName());
+			}
+			return false;
+		}
+	}
+	if (CurrentToolTarget.TargetType == ELandscapeToolTargetType::Weightmap && CurrentToolTarget.LayerInfo == NULL && CurrentTool->GetToolName() != FName("BPCustom"))
+	{
+		if (Reason)
+		{
+			*Reason = NSLOCTEXT("UnrealEd", "LandscapeNeedToCreateLayerInfo", "This layer has no layer info assigned yet. You must create or assign a layer info before you can paint this layer.");
+		}
+		return false;
+		// TODO: FName to LayerInfo: do we want to create the layer info here?
+		//if (FMessageDialog::Open(EAppMsgType::YesNo, NSLOCTEXT("UnrealEd", "LandscapeNeedToCreateLayerInfo", "This layer has no layer info assigned yet. You must create or assign a layer info before you can paint this layer.")) == EAppReturnType::Yes)
+		//{
+		//	CurrentToolTarget.LandscapeInfo->LandscapeProxy->CreateLayerInfo(*CurrentToolTarget.PlaceholderLayerName.ToString());
+		//}
+	}
+	return true;
+}
+
 void FEdModeLandscape::UpdateLandscapeSplines(bool bUpdateOnlySelected /* = false*/)
 {
 	if (GetMutableDefault<UEditorExperimentalSettings>()->bLandscapeLayerSystem)
diff --git a/Engine/Source/Editor/LandscapeEditor/Private/LandscapeEdMode.h b/Engine/Source/Editor/LandscapeEditor/Private/LandscapeEdMode.h
--- a/Engine/Source/Editor/LandscapeEditor/Private/LandscapeEdMode.h
+++ b/Engine/Source/Editor/LandscapeEditor/Private/LandscapeEdMode.h
@@ -511,6 +511,7 @@ public:
 	void SetCurrentLayerSubstractiveBlendStatus(bool InStatus, const TWeakObjectPtr<ULandscapeLayerInfoObject>& InLayerInfoObj);
 	void UpdateLandscapeSplines(bool bUpdateOnlySelected = false);
 	void AutoUpdateDirtyLandscapeSplines();
+	bool CanEditLayer(FText* Reason = nullptr);
 
 	void AddBrushToCurrentLayer(int32 InTargetType, class ALandscapeBlueprintCustomBrush* InBrush);
 	void RemoveBrushFromCurrentLayer(int32 InTargetType, class ALandscapeBlueprintCustomBrush* InBrush);
diff --git a/Engine/Source/Editor/LandscapeEditor/Private/LandscapeEdModeComponentTools.cpp b/Engine/Source/Editor/LandscapeEditor/Private/LandscapeEdModeComponentTools.cpp
--- a/Engine/Source/Editor/LandscapeEditor/Private/LandscapeEdModeComponentTools.cpp
+++ b/Engine/Source/Editor/LandscapeEditor/Private/LandscapeEdModeComponentTools.cpp
@@ -342,48 +342,9 @@ public:
 			return false;
 		}
 
-		if (GetMutableDefault<UEditorExperimentalSettings>()->bLandscapeLayerSystem)
-		{
-			ALandscape* Landscape = this->EdMode->GetLandscape();
-			if (Landscape)
-			{
-				Landscape->RequestLayersContentUpdate(ELandscapeLayersContentUpdateFlag::Weightmap_Render);
-				Landscape->SetEditingLayer(this->EdMode->GetCurrentLayerGuid());
-			}
-		}
-
 		return FLandscapeToolBase<FLandscapeToolStrokeVisibility>::BeginTool(ViewportClient, InTarget, InHitLocation);
 	}
 
-	virtual void Tick(FEditorViewportClient* ViewportClient, float DeltaTime) override
-	{
-		FLandscapeToolBase<FLandscapeToolStrokeVisibility>::Tick(ViewportClient, DeltaTime);
-
-		if (GetMutableDefault<UEditorExperimentalSettings>()->bLandscapeLayerSystem && this->IsToolActive())
-		{
-			ALandscape* Landscape = this->EdMode->GetLandscape();
-			if (Landscape)
-			{
-				Landscape->RequestLayersContentUpdate(ELandscapeLayersContentUpdateFlag::Weightmap_Render);
-			}
-		}
-	}
-
-	virtual void EndTool(FEditorViewportClient* ViewportClient) override
-	{
-		if (GetMutableDefault<UEditorExperimentalSettings>()->bLandscapeLayerSystem)
-		{
-			ALandscape* Landscape = this->EdMode->GetLandscape();
-			if (Landscape)
-			{
-				Landscape->SetEditingLayer();
-				Landscape->RequestLayersContentUpdate(ELandscapeLayersContentUpdateFlag::Weightmap_All);
-			}
-		}
-
-		FLandscapeToolBase<FLandscapeToolStrokeVisibility>::EndTool(ViewportClient);
-	}
-
 	virtual const TCHAR* GetToolName() override { return TEXT("Visibility"); }
 	virtual FText GetDisplayName() override { return NSLOCTEXT("UnrealEd", "LandscapeMode_Visibility", "Visibility"); };
 
diff --git a/Engine/Source/Editor/LandscapeEditor/Private/LandscapeEdModePaintTools.cpp b/Engine/Source/Editor/LandscapeEditor/Private/LandscapeEdModePaintTools.cpp
--- a/Engine/Source/Editor/LandscapeEditor/Private/LandscapeEdModePaintTools.cpp
+++ b/Engine/Source/Editor/LandscapeEditor/Private/LandscapeEdModePaintTools.cpp
@@ -58,52 +58,6 @@ public:
 	{
 		return ELandscapeToolTargetTypeMask::FromType(TToolTarget::TargetType);
 	}
-
-	virtual void Tick(FEditorViewportClient* ViewportClient, float DeltaTime) override
-	{
-		FLandscapeToolBase<TStrokeClass>::Tick(ViewportClient, DeltaTime);
-
-		if (GetMutableDefault<UEditorExperimentalSettings>()->bLandscapeLayerSystem && this->IsToolActive())
-		{
-			ALandscape* Landscape = this->EdMode->CurrentToolTarget.LandscapeInfo->LandscapeActor.Get();
-			if (Landscape != nullptr)
-			{
-				Landscape->RequestLayersContentUpdate(this->EdMode->CurrentToolTarget.TargetType == ELandscapeToolTargetType::Type::Heightmap ? ELandscapeLayersContentUpdateFlag::Heightmap_Render : ELandscapeLayersContentUpdateFlag::Weightmap_Render);
-			}
-		}
-	}
-
-	virtual bool BeginTool(FEditorViewportClient* ViewportClient, const FLandscapeToolTarget& InTarget, const FVector& InHitLocation) override
-	{
-		if (GetMutableDefault<UEditorExperimentalSettings>()->bLandscapeLayerSystem)
-		{
-			ALandscape* Landscape = this->EdMode->GetLandscape();
-			if (Landscape)
-			{
-				bool bUpdateHeightmap = this->EdMode->CurrentToolTarget.TargetType == ELandscapeToolTargetType::Type::Heightmap;
-				Landscape->RequestLayersContentUpdate(bUpdateHeightmap ? ELandscapeLayersContentUpdateFlag::Heightmap_Render : ELandscapeLayersContentUpdateFlag::Weightmap_Render);
-				Landscape->SetEditingLayer(this->EdMode->GetCurrentLayerGuid());
-			}
-		}
-
-		return FLandscapeToolBase<TStrokeClass>::BeginTool(ViewportClient, InTarget, InHitLocation);
-	}
-
-	virtual void EndTool(FEditorViewportClient* ViewportClient) override
-	{
-		if (GetMutableDefault<UEditorExperimentalSettings>()->bLandscapeLayerSystem)
-		{
-			ALandscape* Landscape = this->EdMode->GetLandscape();
-			if (Landscape)
-			{
-				bool bUpdateHeightmap = this->EdMode->CurrentToolTarget.TargetType == ELandscapeToolTargetType::Type::Heightmap;
-				Landscape->SetEditingLayer();
-				Landscape->RequestLayersContentUpdate(bUpdateHeightmap ? ELandscapeLayersContentUpdateFlag::Heightmap_All : ELandscapeLayersContentUpdateFlag::Weightmap_All);
-			}
-		}
-
-		FLandscapeToolBase<TStrokeClass>::EndTool(ViewportClient);
-	}
 };
 
 
diff --git a/Engine/Source/Editor/LandscapeEditor/Private/LandscapeEdModeRampTool.cpp b/Engine/Source/Editor/LandscapeEditor/Private/LandscapeEdModeRampTool.cpp
--- a/Engine/Source/Editor/LandscapeEditor/Private/LandscapeEdModeRampTool.cpp
+++ b/Engine/Source/Editor/LandscapeEditor/Private/LandscapeEdModeRampTool.cpp
@@ -19,6 +19,9 @@
 #include "LandscapeDataAccess.h"
 #include "LandscapeHeightfieldCollisionComponent.h"
 #include "Raster.h"
+#include "Landscape.h"
+#include "Settings/EditorExperimentalSettings.h"
+#include "Misc/MessageDialog.h"
 
 #define LOCTEXT_NAMESPACE "Landscape"
 
@@ -474,7 +477,16 @@ public:
 
 	virtual void ApplyRamp()
 	{
+		FText Reason;
+		if (!EdMode->CanEditLayer(&Reason))
+		{
+			FMessageDialog::Open(EAppMsgType::Ok, Reason);
+			return;
+		}
+
 		FScopedTransaction Transaction(LOCTEXT("Ramp_Apply", "Landscape Editing: Add ramp"));
+		ALandscape* Landscape = EdMode->GetLandscape();
+		FScopedSetLandscapeEditingLayer Scope(Landscape, EdMode->GetCurrentLayerGuid(), [&] { if (Landscape) { Landscape->RequestLayersContentUpdate(ELandscapeLayersContentUpdateFlag::Heightmap_All); } });
 
 		const ULandscapeInfo* LandscapeInfo = EdMode->CurrentToolTarget.LandscapeInfo.Get();
 		const ALandscapeProxy* LandscapeProxy = LandscapeInfo->GetLandscapeProxy();
@@ -570,17 +582,20 @@ public:
 			LandscapeEdit.SetHeightData(MinX, MinY, MaxX, MaxY, Data.GetData(), 0, true);
 			LandscapeEdit.Flush();
 
-			TSet<ULandscapeComponent*> Components;
-			if (LandscapeEdit.GetComponentsInRegion(MinX, MinY, MaxX, MaxY, &Components))
+			if (!GetMutableDefault<UEditorExperimentalSettings>()->bLandscapeLayerSystem)
 			{
-				for (ULandscapeComponent* Component : Components)
+				TSet<ULandscapeComponent*> Components;
+				if (LandscapeEdit.GetComponentsInRegion(MinX, MinY, MaxX, MaxY, &Components))
 				{
-					// Recreate collision for modified components and update the navmesh
-					ULandscapeHeightfieldCollisionComponent* CollisionComponent = Component->CollisionComponent.Get();
-					if (CollisionComponent)
+					for (ULandscapeComponent* Component : Components)
 					{
-						CollisionComponent->RecreateCollision();
-						FNavigationSystem::UpdateComponentData(*CollisionComponent);
+						// Recreate collision for modified components and update the navmesh
+						ULandscapeHeightfieldCollisionComponent* CollisionComponent = Component->CollisionComponent.Get();
+						if (CollisionComponent)
+						{
+							CollisionComponent->RecreateCollision();
+							FNavigationSystem::UpdateComponentData(*CollisionComponent);
+						}
 					}
 				}
 			}
@@ -589,7 +604,7 @@ public:
 
 	bool CanApplyRamp()
 	{
-		return NumPoints == 2;
+		return EdMode->CanEditLayer() && (NumPoints == 2);
 	}
 
 	void ResetRamp()
diff --git a/Engine/Source/Editor/LandscapeEditor/Private/LandscapeEdModeTools.h b/Engine/Source/Editor/LandscapeEditor/Private/LandscapeEdModeTools.h
--- a/Engine/Source/Editor/LandscapeEditor/Private/LandscapeEdModeTools.h
+++ b/Engine/Source/Editor/LandscapeEditor/Private/LandscapeEdModeTools.h
@@ -15,6 +15,8 @@
 #include "InstancedFoliageActor.h"
 #include "VREditorInteractor.h"
 #include "AI/NavigationSystemBase.h"
+#include "Settings/EditorExperimentalSettings.h"
+#include "Landscape.h"
 
 // VR Editor
 
@@ -1121,6 +1123,17 @@ public:
 
 	virtual bool BeginTool(FEditorViewportClient* ViewportClient, const FLandscapeToolTarget& InTarget, const FVector& InHitLocation) override
 	{
+		if (GetMutableDefault<UEditorExperimentalSettings>()->bLandscapeLayerSystem)
+		{
+			ALandscape* Landscape = this->EdMode->GetLandscape();
+			if (Landscape)
+			{
+				bool bUpdateHeightmap = this->EdMode->CurrentToolTarget.TargetType == ELandscapeToolTargetType::Type::Heightmap;
+				Landscape->RequestLayersContentUpdate(bUpdateHeightmap ? ELandscapeLayersContentUpdateFlag::Heightmap_Render : ELandscapeLayersContentUpdateFlag::Weightmap_Render);
+				Landscape->SetEditingLayer(this->EdMode->GetCurrentLayerGuid());
+			}
+		}
+
 		if (!ensure(InteractorPositions.Num() == 0))
 		{
 			InteractorPositions.Empty(1);
@@ -1164,11 +1177,32 @@ public:
 
 			// Prevent landscape from baking textures while tool stroke is active
 			EdMode->CurrentToolTarget.LandscapeInfo->PostponeTextureBaking();
+
+			if (GetMutableDefault<UEditorExperimentalSettings>()->bLandscapeLayerSystem)
+			{
+				ALandscape* Landscape = this->EdMode->CurrentToolTarget.LandscapeInfo->LandscapeActor.Get();
+				if (Landscape != nullptr)
+				{
+					bool bUpdateHeightmap = this->EdMode->CurrentToolTarget.TargetType == ELandscapeToolTargetType::Type::Heightmap;
+					Landscape->RequestLayersContentUpdate(bUpdateHeightmap ? ELandscapeLayersContentUpdateFlag::Heightmap_Render : ELandscapeLayersContentUpdateFlag::Weightmap_Render);
+				}
+			}
 		}
 	}
 
 	virtual void EndTool(FEditorViewportClient* ViewportClient) override
 	{
+		if (GetMutableDefault<UEditorExperimentalSettings>()->bLandscapeLayerSystem)
+		{
+			ALandscape* Landscape = this->EdMode->GetLandscape();
+			if (Landscape)
+			{
+				bool bUpdateHeightmap = this->EdMode->CurrentToolTarget.TargetType == ELandscapeToolTargetType::Type::Heightmap;
+				Landscape->RequestLayersContentUpdate(bUpdateHeightmap ? ELandscapeLayersContentUpdateFlag::Heightmap_All : ELandscapeLayersContentUpdateFlag::Weightmap_All);
+				Landscape->SetEditingLayer();
+			}
+		}
+
 		if (IsToolActive() && InteractorPositions.Num())
 		{
 			ToolStroke->Apply(ViewportClient, EdMode->CurrentBrush, EdMode->UISettings, InteractorPositions);
@@ -1257,4 +1291,4 @@ struct FToolFlattenCustomData
 
 	bool ActiveEyeDropperMode;
 	float EyeDropperModeHeight;
-};
+};
\ No newline at end of file
diff --git a/Engine/Source/Editor/LandscapeEditor/Private/LandscapeEditorDetailCustomization_Layers.cpp b/Engine/Source/Editor/LandscapeEditor/Private/LandscapeEditorDetailCustomization_Layers.cpp
--- a/Engine/Source/Editor/LandscapeEditor/Private/LandscapeEditorDetailCustomization_Layers.cpp
+++ b/Engine/Source/Editor/LandscapeEditor/Private/LandscapeEditorDetailCustomization_Layers.cpp
@@ -18,6 +18,7 @@
 #include "Widgets/Notifications/SErrorText.h"
 #include "Widgets/Input/SComboButton.h"
 #include "Widgets/Input/SCheckBox.h"
+#include "Widgets/Input/SEditableTextBox.h"
 #include "EditorModeManager.h"
 #include "EditorModes.h"
 #include "LandscapeEditorModule.h"
@@ -47,6 +48,8 @@
 #include "Widgets/Input/SNumericEntryBox.h"
 #include "Widgets/Text/SInlineEditableTextBlock.h"
 #include "LandscapeEditorCommands.h"
+#include "Settings/EditorExperimentalSettings.h"
+#include "LandscapeEditorDetailCustomization_LayersBrushStack.h"
 
 #define LOCTEXT_NAMESPACE "LandscapeEditor.Layers"
 
@@ -65,9 +68,39 @@ void FLandscapeEditorDetailCustomization_Layers::CustomizeDetails(IDetailLayoutB
 	{
 		LayerCategory.AddCustomBuilder(MakeShareable(new FLandscapeEditorCustomNodeBuilder_Layers(DetailBuilder.GetThumbnailPool().ToSharedRef())));
 	}
+
+	LayerCategory.AddCustomRow(FText())
+	.Visibility(TAttribute<EVisibility>::Create(TAttribute<EVisibility>::FGetter::CreateLambda([]() { return ShoudShowLayersErrorMessageTip() ? EVisibility::Visible : EVisibility::Collapsed; })))
+	[
+		SNew(SEditableTextBox)
+		.Font(DetailBuilder.GetDetailFontBold())
+		.BackgroundColor(TAttribute<FSlateColor>::Create(TAttribute<FSlateColor>::FGetter::CreateLambda([]() { return FEditorStyle::GetColor("ErrorReporting.WarningBackgroundColor"); })))
+		.Text(TAttribute<FText>::Create(TAttribute<FText>::FGetter::CreateStatic(&FLandscapeEditorDetailCustomization_Layers::GetLayersErrorMessageText)))
+	];
 }
 END_SLATE_FUNCTION_BUILD_OPTIMIZATION
 
+bool FLandscapeEditorDetailCustomization_Layers::ShoudShowLayersErrorMessageTip()
+{
+	FEdModeLandscape* LandscapeEdMode = GetEditorMode();
+	if (LandscapeEdMode)
+	{
+		return !LandscapeEdMode->CanEditLayer();
+	}
+	return false;
+}
+
+FText FLandscapeEditorDetailCustomization_Layers::GetLayersErrorMessageText()
+{
+	FEdModeLandscape* LandscapeEdMode = GetEditorMode();
+	FText Reason;
+	if (LandscapeEdMode && !LandscapeEdMode->CanEditLayer(&Reason))
+	{
+		return Reason;
+	}
+	return FText::GetEmpty();
+}
+
 //////////////////////////////////////////////////////////////////////////
 
 FEdModeLandscape* FLandscapeEditorCustomNodeBuilder_Layers::GetEditorMode()
@@ -77,6 +110,7 @@ FEdModeLandscape* FLandscapeEditorCustomNodeBuilder_Layers::GetEditorMode()
 
 FLandscapeEditorCustomNodeBuilder_Layers::FLandscapeEditorCustomNodeBuilder_Layers(TSharedRef<FAssetThumbnailPool> InThumbnailPool)
 	: ThumbnailPool(InThumbnailPool)
+	, CurrentEditingInlineTextBlock(INDEX_NONE)
 {
 }
 
@@ -200,6 +234,8 @@ TSharedPtr<SWidget> FLandscapeEditorCustomNodeBuilder_Layers::GenerateRow(int32
 				.Text(this, &FLandscapeEditorCustomNodeBuilder_Layers::GetLayerDisplayName, InLayerIndex)
 				.ToolTipText(LOCTEXT("LandscapeLayers_tooltip", "Name of the Layer"))
 				.OnVerifyTextChanged(FOnVerifyTextChanged::CreateSP(this, &FLandscapeEditorCustomNodeBuilder_Layers::CanRenameLayerTo, InLayerIndex))
+				.OnEnterEditingMode(this, &FLandscapeEditorCustomNodeBuilder_Layers::OnBeginNameTextEdit)
+				.OnExitEditingMode(this, &FLandscapeEditorCustomNodeBuilder_Layers::OnEndNameTextEdit)
 				.OnTextCommitted(FOnTextCommitted::CreateSP(this, &FLandscapeEditorCustomNodeBuilder_Layers::SetLayerName, InLayerIndex))
 			]
 
@@ -253,8 +289,8 @@ FText FLandscapeEditorCustomNodeBuilder_Layers::GetLayerDisplayName(int32 InLaye
 	{
 		const FLandscapeLayer* Layer = LandscapeEdMode->GetLayer(InLayerIndex);
 		const FLandscapeLayer* ReservedLayer = Landscape->GetLandscapeSplinesReservedLayer();
-		bool bIsReserved = Layer && Layer == ReservedLayer;
-		return FText::Format(LOCTEXT("DisplayName", "{0} {1}"), FText::FromName(LandscapeEdMode->GetLayerName(InLayerIndex)), bIsReserved ? LOCTEXT("ReservedForSplines", "(Reserved for Splines)") : FText::GetEmpty());
+		bool bIsReserved = Layer && Layer == ReservedLayer && InLayerIndex != CurrentEditingInlineTextBlock;
+		return FText::Format(LOCTEXT("LayerDisplayName", "{0}{1}"), FText::FromName(LandscapeEdMode->GetLayerName(InLayerIndex)), bIsReserved ? LOCTEXT("ReservedForSplines", " (Reserved for Splines)") : FText::GetEmpty());
 	}
 
 	return FText::FromString(TEXT("None"));
@@ -271,6 +307,17 @@ bool FLandscapeEditorCustomNodeBuilder_Layers::IsLayerSelected(int32 InLayerInde
 	return false;
 }
 
+void FLandscapeEditorCustomNodeBuilder_Layers::OnBeginNameTextEdit()
+{
+	FEdModeLandscape* LandscapeEdMode = GetEditorMode();
+	CurrentEditingInlineTextBlock = LandscapeEdMode ? LandscapeEdMode->GetCurrentLayerIndex() : INDEX_NONE;
+}
+
+void FLandscapeEditorCustomNodeBuilder_Layers::OnEndNameTextEdit()
+{
+	CurrentEditingInlineTextBlock = INDEX_NONE;
+}
+
 bool FLandscapeEditorCustomNodeBuilder_Layers::CanRenameLayerTo(const FText& InNewText, FText& OutErrorMessage, int32 InLayerIndex)
 {
 	FEdModeLandscape* LandscapeEdMode = GetEditorMode();
@@ -599,4 +646,4 @@ FReply FLandscapeEditorCustomNodeBuilder_Layers::HandleAcceptDrop(FDragDropEvent
 	return FReply::Unhandled();
 }
 
-#undef LOCTEXT_NAMESPACE
\ No newline at end of file
+#undef LOCTEXT_NAMESPACE
diff --git a/Engine/Source/Editor/LandscapeEditor/Private/LandscapeEditorDetailCustomization_Layers.h b/Engine/Source/Editor/LandscapeEditor/Private/LandscapeEditorDetailCustomization_Layers.h
--- a/Engine/Source/Editor/LandscapeEditor/Private/LandscapeEditorDetailCustomization_Layers.h
+++ b/Engine/Source/Editor/LandscapeEditor/Private/LandscapeEditorDetailCustomization_Layers.h
@@ -36,6 +36,10 @@ public:
 
 	/** IDetailCustomization interface */
 	virtual void CustomizeDetails(IDetailLayoutBuilder& DetailBuilder) override;
+
+protected:
+	static bool ShoudShowLayersErrorMessageTip();
+	static FText GetLayersErrorMessageText();
 };
 
 class FLandscapeEditorCustomNodeBuilder_Layers : public IDetailCustomNodeBuilder, public TSharedFromThis<FLandscapeEditorCustomNodeBuilder_Layers>
@@ -74,6 +78,8 @@ protected:
 	void ShowOnlySelectedLayer(int32 InLayerIndex);
 	void ShowAllLayers();
 	bool CanRenameLayerTo(const FText& NewText, FText& OutErrorMessage, int32 InLayerIndex);
+	void OnBeginNameTextEdit();
+	void OnEndNameTextEdit();
 	void SetLayerName(const FText& InText, ETextCommit::Type InCommitType, int32 InLayerIndex);
 	FText GetLayerText(int32 InLayerIndex) const;
 	void SetLandscapeSplinesReservedLayer(int32 InLayerIndex);
@@ -94,4 +100,6 @@ private:
 
 	/** Widgets for displaying and editing the layer name */
 	TArray< TSharedPtr< SInlineEditableTextBlock > > InlineTextBlocks;
+
+	int32 CurrentEditingInlineTextBlock;
 };
\ No newline at end of file
