- Support reorder of landscape layers

diff --git a/Engine/Source/Editor/LandscapeEditor/Private/LandscapeEditorDetailCustomization_Layers.cpp b/Engine/Source/Editor/LandscapeEditor/Private/LandscapeEditorDetailCustomization_Layers.cpp
--- a/Engine/Source/Editor/LandscapeEditor/Private/LandscapeEditorDetailCustomization_Layers.cpp
+++ b/Engine/Source/Editor/LandscapeEditor/Private/LandscapeEditorDetailCustomization_Layers.cpp
@@ -49,7 +49,6 @@
 #include "Widgets/Text/SInlineEditableTextBlock.h"
 #include "LandscapeEditorCommands.h"
 #include "Settings/EditorExperimentalSettings.h"
-#include "LandscapeEditorDetailCustomization_LayersBrushStack.h"
 
 #define LOCTEXT_NAMESPACE "LandscapeEditor.Layers"
 
@@ -625,25 +624,78 @@ const FSlateBrush* FLandscapeEditorCustomNodeBuilder_Layers::GetLockBrushForLaye
 FReply FLandscapeEditorCustomNodeBuilder_Layers::HandleDragDetected(const FGeometry& MyGeometry, const FPointerEvent& MouseEvent, int32 SlotIndex, SVerticalBox::FSlot* Slot)
 {
 	FEdModeLandscape* LandscapeEdMode = GetEditorMode();
-
-	if (LandscapeEdMode != nullptr)
+	if (LandscapeEdMode)
 	{
-		// TODO: handle drag & drop
+		FLandscapeLayer* Layer = LandscapeEdMode->GetLayer(SlotIndex);
+		if (Layer && !Layer->bLocked)
+		{
+			TSharedPtr<SWidget> Row = GenerateRow(SlotIndex);
+			if (Row.IsValid())
+			{
+				return FReply::Handled().BeginDragDrop(FLandscapeListElementDragDropOp::New(SlotIndex, Slot, Row));
+			}
+		}
 	}
-
 	return FReply::Unhandled();
 }
 
 TOptional<SDragAndDropVerticalBox::EItemDropZone> FLandscapeEditorCustomNodeBuilder_Layers::HandleCanAcceptDrop(const FDragDropEvent& DragDropEvent, SDragAndDropVerticalBox::EItemDropZone DropZone, SVerticalBox::FSlot* Slot)
 {
-	// TODO: handle drag & drop
+	TSharedPtr<FLandscapeListElementDragDropOp> DragDropOperation = DragDropEvent.GetOperationAs<FLandscapeListElementDragDropOp>();
+	if (DragDropOperation.IsValid())
+	{
+		return DropZone;
+	}
 	return TOptional<SDragAndDropVerticalBox::EItemDropZone>();
 }
 
 FReply FLandscapeEditorCustomNodeBuilder_Layers::HandleAcceptDrop(FDragDropEvent const& DragDropEvent, SDragAndDropVerticalBox::EItemDropZone DropZone, int32 SlotIndex, SVerticalBox::FSlot* Slot)
 {
-	// TODO: handle drag & drop
+	TSharedPtr<FLandscapeListElementDragDropOp> DragDropOperation = DragDropEvent.GetOperationAs<FLandscapeListElementDragDropOp>();
+
+	if (DragDropOperation.IsValid())
+	{
+		FEdModeLandscape* LandscapeEdMode = GetEditorMode();
+		ALandscape* Landscape = LandscapeEdMode ? LandscapeEdMode->GetLandscape() : nullptr;
+		if (Landscape)
+		{
+			int32 StartingLayerIndex = DragDropOperation->SlotIndexBeingDragged;
+			int32 DestinationLayerIndex = SlotIndex;
+			if (Landscape->ReorderLayer(StartingLayerIndex, DestinationLayerIndex))
+			{
+				LandscapeEdMode->SetCurrentLayer(DestinationLayerIndex);
+				LandscapeEdMode->RefreshDetailPanel();
+				LandscapeEdMode->RequestLayersContentUpdate();
+				return FReply::Handled();
+			}
+		}
+	}
+
 	return FReply::Unhandled();
 }
 
-#undef LOCTEXT_NAMESPACE
+TSharedRef<FLandscapeListElementDragDropOp> FLandscapeListElementDragDropOp::New(int32 InSlotIndexBeingDragged, SVerticalBox::FSlot* InSlotBeingDragged, TSharedPtr<SWidget> WidgetToShow)
+{
+	TSharedRef<FLandscapeListElementDragDropOp> Operation = MakeShareable(new FLandscapeListElementDragDropOp);
+
+	Operation->MouseCursor = EMouseCursor::GrabHandClosed;
+	Operation->SlotIndexBeingDragged = InSlotIndexBeingDragged;
+	Operation->SlotBeingDragged = InSlotBeingDragged;
+	Operation->WidgetToShow = WidgetToShow;
+
+	Operation->Construct();
+
+	return Operation;
+}
+
+TSharedPtr<SWidget> FLandscapeListElementDragDropOp::GetDefaultDecorator() const
+{
+	return SNew(SBorder)
+		.BorderImage(FEditorStyle::GetBrush("ContentBrowser.AssetDragDropTooltipBackground"))
+		.Content()
+		[
+			WidgetToShow.ToSharedRef()
+		];
+}
+
+#undef LOCTEXT_NAMESPACE
\ No newline at end of file
diff --git a/Engine/Source/Editor/LandscapeEditor/Private/LandscapeEditorDetailCustomization_Layers.h b/Engine/Source/Editor/LandscapeEditor/Private/LandscapeEditorDetailCustomization_Layers.h
--- a/Engine/Source/Editor/LandscapeEditor/Private/LandscapeEditorDetailCustomization_Layers.h
+++ b/Engine/Source/Editor/LandscapeEditor/Private/LandscapeEditorDetailCustomization_Layers.h
@@ -102,4 +102,19 @@ private:
 	TArray< TSharedPtr< SInlineEditableTextBlock > > InlineTextBlocks;
 
 	int32 CurrentEditingInlineTextBlock;
+};
+
+class FLandscapeListElementDragDropOp : public FDragAndDropVerticalBoxOp
+{
+public:
+	DRAG_DROP_OPERATOR_TYPE(FLandscapeListElementDragDropOp, FDragAndDropVerticalBoxOp)
+
+	TSharedPtr<SWidget> WidgetToShow;
+
+	static TSharedRef<FLandscapeListElementDragDropOp> New(int32 InSlotIndexBeingDragged, SVerticalBox::FSlot* InSlotBeingDragged, TSharedPtr<SWidget> InWidgetToShow);
+
+public:
+	virtual ~FLandscapeListElementDragDropOp() {}
+
+	virtual TSharedPtr<SWidget> GetDefaultDecorator() const override;
 };
\ No newline at end of file
diff --git a/Engine/Source/Editor/LandscapeEditor/Private/LandscapeEditorDetailCustomization_LayersBrushStack.cpp b/Engine/Source/Editor/LandscapeEditor/Private/LandscapeEditorDetailCustomization_LayersBrushStack.cpp
--- a/Engine/Source/Editor/LandscapeEditor/Private/LandscapeEditorDetailCustomization_LayersBrushStack.cpp
+++ b/Engine/Source/Editor/LandscapeEditor/Private/LandscapeEditorDetailCustomization_LayersBrushStack.cpp
@@ -20,6 +20,7 @@
 #include "IDetailPropertyRow.h"
 #include "DetailCategoryBuilder.h"
 #include "PropertyCustomizationHelpers.h"
+#include "LandscapeEditorDetailCustomization_Layers.h"
 
 #include "ScopedTransaction.h"
 
@@ -316,7 +317,7 @@ FReply FLandscapeEditorCustomNodeBuilder_LayersBrushStack::HandleDragDetected(co
 
 			if (Row.IsValid())
 			{
-				return FReply::Handled().BeginDragDrop(FLandscapeBrushDragDropOp::New(SlotIndex, Slot, Row));
+				return FReply::Handled().BeginDragDrop(FLandscapeListElementDragDropOp::New(SlotIndex, Slot, Row));
 			}
 		}
 	}
@@ -326,7 +327,7 @@ FReply FLandscapeEditorCustomNodeBuilder_LayersBrushStack::HandleDragDetected(co
 
 TOptional<SDragAndDropVerticalBox::EItemDropZone> FLandscapeEditorCustomNodeBuilder_LayersBrushStack::HandleCanAcceptDrop(const FDragDropEvent& DragDropEvent, SDragAndDropVerticalBox::EItemDropZone DropZone, SVerticalBox::FSlot* Slot)
 {
-	TSharedPtr<FLandscapeBrushDragDropOp> DragDropOperation = DragDropEvent.GetOperationAs<FLandscapeBrushDragDropOp>();
+	TSharedPtr<FLandscapeListElementDragDropOp> DragDropOperation = DragDropEvent.GetOperationAs<FLandscapeListElementDragDropOp>();
 
 	if (DragDropOperation.IsValid())
 	{
@@ -338,7 +339,7 @@ TOptional<SDragAndDropVerticalBox::EItemDropZone> FLandscapeEditorCustomNodeBuil
 
 FReply FLandscapeEditorCustomNodeBuilder_LayersBrushStack::HandleAcceptDrop(FDragDropEvent const& DragDropEvent, SDragAndDropVerticalBox::EItemDropZone DropZone, int32 SlotIndex, SVerticalBox::FSlot* Slot)
 {
-	TSharedPtr<FLandscapeBrushDragDropOp> DragDropOperation = DragDropEvent.GetOperationAs<FLandscapeBrushDragDropOp>();
+	TSharedPtr<FLandscapeListElementDragDropOp> DragDropOperation = DragDropEvent.GetOperationAs<FLandscapeListElementDragDropOp>();
 
 	if (DragDropOperation.IsValid())
 	{
@@ -372,33 +373,4 @@ FReply FLandscapeEditorCustomNodeBuilder_LayersBrushStack::HandleAcceptDrop(FDra
 	return FReply::Unhandled();
 }
 
-TSharedRef<FLandscapeBrushDragDropOp> FLandscapeBrushDragDropOp::New(int32 InSlotIndexBeingDragged, SVerticalBox::FSlot* InSlotBeingDragged, TSharedPtr<SWidget> WidgetToShow)
-{
-	TSharedRef<FLandscapeBrushDragDropOp> Operation = MakeShareable(new FLandscapeBrushDragDropOp);
-
-	Operation->MouseCursor = EMouseCursor::GrabHandClosed;
-	Operation->SlotIndexBeingDragged = InSlotIndexBeingDragged;
-	Operation->SlotBeingDragged = InSlotBeingDragged;
-	Operation->WidgetToShow = WidgetToShow;
-
-	Operation->Construct();
-
-	return Operation;
-}
-
-FLandscapeBrushDragDropOp::~FLandscapeBrushDragDropOp()
-{
-}
-
-TSharedPtr<SWidget> FLandscapeBrushDragDropOp::GetDefaultDecorator() const
-{
-	return SNew(SBorder)
-		.BorderImage(FEditorStyle::GetBrush("ContentBrowser.AssetDragDropTooltipBackground"))
-		.Content()
-		[
-			WidgetToShow.ToSharedRef()
-		];
-
-}
-
 #undef LOCTEXT_NAMESPACE
\ No newline at end of file
diff --git a/Engine/Source/Editor/LandscapeEditor/Private/LandscapeEditorDetailCustomization_LayersBrushStack.h b/Engine/Source/Editor/LandscapeEditor/Private/LandscapeEditorDetailCustomization_LayersBrushStack.h
--- a/Engine/Source/Editor/LandscapeEditor/Private/LandscapeEditorDetailCustomization_LayersBrushStack.h
+++ b/Engine/Source/Editor/LandscapeEditor/Private/LandscapeEditorDetailCustomization_LayersBrushStack.h
@@ -73,18 +73,3 @@ protected:
 	TOptional<SDragAndDropVerticalBox::EItemDropZone> HandleCanAcceptDrop(const FDragDropEvent& DragDropEvent, SDragAndDropVerticalBox::EItemDropZone DropZone, SVerticalBox::FSlot* Slot);
 	FReply HandleAcceptDrop(FDragDropEvent const& DragDropEvent, SDragAndDropVerticalBox::EItemDropZone DropZone, int32 SlotIndex, SVerticalBox::FSlot* Slot);
 };
-
-class FLandscapeBrushDragDropOp : public FDragAndDropVerticalBoxOp
-{
-public:
-	DRAG_DROP_OPERATOR_TYPE(FLandscapeBrushDragDropOp, FDragAndDropVerticalBoxOp)
-
-	TSharedPtr<SWidget> WidgetToShow;
-
-	static TSharedRef<FLandscapeBrushDragDropOp> New(int32 InSlotIndexBeingDragged, SVerticalBox::FSlot* InSlotBeingDragged, TSharedPtr<SWidget> InWidgetToShow);
-
-public:
-	virtual ~FLandscapeBrushDragDropOp();
-
-	virtual TSharedPtr<SWidget> GetDefaultDecorator() const override;
-};
diff --git a/Engine/Source/Runtime/Landscape/Classes/Landscape.h b/Engine/Source/Runtime/Landscape/Classes/Landscape.h
--- a/Engine/Source/Runtime/Landscape/Classes/Landscape.h
+++ b/Engine/Source/Runtime/Landscape/Classes/Landscape.h
@@ -243,6 +243,7 @@ public:
 #if WITH_EDITOR
 	LANDSCAPE_API void RequestLayersContentUpdate(uint32 InDataFlags, bool InUpdateAllMaterials = false);
 	LANDSCAPE_API void CreateLayer(FName InName = NAME_None, bool bInUpdateLayersContent = true);
+	LANDSCAPE_API bool ReorderLayer(int32 InStartingLayerIndex, int32 InDestinationLayerIndex);
 	LANDSCAPE_API bool IsLayerNameUnique(const FName& InName) const;
 	LANDSCAPE_API void SetLayerName(int32 InLayerIndex, const FName& InName);
 	LANDSCAPE_API void SetLayerAlpha(int32 InLayerIndex, const float InAlpha, bool bInHeightmap);
diff --git a/Engine/Source/Runtime/Landscape/Private/LandscapeEditLayers.cpp b/Engine/Source/Runtime/Landscape/Private/LandscapeEditLayers.cpp
--- a/Engine/Source/Runtime/Landscape/Private/LandscapeEditLayers.cpp
+++ b/Engine/Source/Runtime/Landscape/Private/LandscapeEditLayers.cpp
@@ -4768,6 +4768,21 @@ void ALandscape::CreateLayer(FName InName, bool bInUpdateLayersContent)
 	}
 }
 
+bool ALandscape::ReorderLayer(int32 InStartingLayerIndex, int32 InDestinationLayerIndex)
+{
+	if (InStartingLayerIndex != InDestinationLayerIndex &&
+		LandscapeLayers.IsValidIndex(InStartingLayerIndex) && 
+		LandscapeLayers.IsValidIndex(InDestinationLayerIndex))
+	{
+		Modify();
+		FLandscapeLayer Layer = LandscapeLayers[InStartingLayerIndex];
+		LandscapeLayers.RemoveAt(InStartingLayerIndex);
+		LandscapeLayers.Insert(Layer, InDestinationLayerIndex);
+		return true;
+	}
+	return false;
+}
+
 FName ALandscape::GenerateUniqueLayerName(FName InName) const
 {
 	FString BaseName = InName == NAME_None ? "Layer" : InName.ToString();
