- Fixed LandscapeEdMode to automatically set the Target Landscape based on selected spline(s) segment(s) and or control point(s).

diff --git a/Engine/Source/Editor/LandscapeEditor/Private/LandscapeEdMode.cpp b/Engine/Source/Editor/LandscapeEditor/Private/LandscapeEdMode.cpp
--- a/Engine/Source/Editor/LandscapeEditor/Private/LandscapeEdMode.cpp
+++ b/Engine/Source/Editor/LandscapeEditor/Private/LandscapeEdMode.cpp
@@ -2413,6 +2413,40 @@ int32 FEdModeLandscape::UpdateLandscapeList()
 	return CurrentIndex;
 }
 
+void FEdModeLandscape::SetTargetLandscape(const TWeakObjectPtr<ULandscapeInfo>& InLandscapeInfo)
+{
+	if ((CurrentToolTarget.LandscapeInfo == InLandscapeInfo) || !InLandscapeInfo.IsValid())
+	{
+		return;
+	}
+
+	// Unregister from old one
+	if (CurrentToolTarget.LandscapeInfo.IsValid())
+	{
+		ALandscapeProxy* LandscapeProxy = CurrentToolTarget.LandscapeInfo->GetLandscapeProxy();
+		LandscapeProxy->OnMaterialChangedDelegate().RemoveAll(this);
+	}
+
+	CurrentToolTarget.LandscapeInfo = InLandscapeInfo.Get();
+	UpdateTargetList();
+	// force a Leave and Enter the current tool, in case it has something about the current landscape cached
+	SetCurrentTool(CurrentToolIndex);
+	if (CurrentGizmoActor.IsValid())
+	{
+		CurrentGizmoActor->SetTargetLandscape(CurrentToolTarget.LandscapeInfo.Get());
+	}
+
+	// register to new one
+	if (CurrentToolTarget.LandscapeInfo.IsValid())
+	{
+		ALandscapeProxy* LandscapeProxy = CurrentToolTarget.LandscapeInfo->GetLandscapeProxy();
+		LandscapeProxy->OnMaterialChangedDelegate().AddRaw(this, &FEdModeLandscape::OnLandscapeMaterialChangedDelegate);
+	}
+
+	UpdateTargetList();
+	UpdateShownLayerList();
+}
+
 void FEdModeLandscape::UpdateTargetList()
 {
 	LandscapeTargetList.Empty();
diff --git a/Engine/Source/Editor/LandscapeEditor/Private/LandscapeEdMode.h b/Engine/Source/Editor/LandscapeEditor/Private/LandscapeEdMode.h
--- a/Engine/Source/Editor/LandscapeEditor/Private/LandscapeEdMode.h
+++ b/Engine/Source/Editor/LandscapeEditor/Private/LandscapeEdMode.h
@@ -477,7 +477,8 @@ public:
 
 	int32 UpdateLandscapeList();
 	void UpdateTargetList();
-	
+	void SetTargetLandscape(const TWeakObjectPtr<ULandscapeInfo>& InLandscapeInfo);
+
 	/** Update Display order list */
 	void UpdateTargetLayerDisplayOrder(ELandscapeLayerDisplayMode InTargetDisplayOrder);
 	void MoveTargetLayerDisplayOrder(int32 IndexToMove, int32 IndexToDestination);
diff --git a/Engine/Source/Editor/LandscapeEditor/Private/LandscapeEdModeSplineTools.cpp b/Engine/Source/Editor/LandscapeEditor/Private/LandscapeEdModeSplineTools.cpp
--- a/Engine/Source/Editor/LandscapeEditor/Private/LandscapeEdModeSplineTools.cpp
+++ b/Engine/Source/Editor/LandscapeEditor/Private/LandscapeEdModeSplineTools.cpp
@@ -1093,6 +1093,21 @@ public:
 	{
 	}
 
+	template <typename T>
+	void SetTargetLandscapeBasedOnSelection(typename T* Selection)
+	{
+		check(Selection);
+		if (ALandscapeProxy* LandscapeProxy = Selection->GetTypedOuter<ALandscapeProxy>())
+		{
+			ALandscape* NewLandscapeActor = LandscapeProxy->GetLandscapeActor();
+			if (NewLandscapeActor && (NewLandscapeActor != EdMode->GetLandscape()))
+			{
+				EdMode->SetTargetLandscape(LandscapeProxy->GetLandscapeInfo());
+				ClearSelection();
+			}
+		}
+	}
+
 	virtual bool HandleClick(HHitProxy* HitProxy, const FViewportClick& Click) override
 	{
 		if ((!HitProxy || !HitProxy->IsA(HWidgetAxis::StaticGetType()))
@@ -1152,6 +1167,7 @@ public:
 				}
 				else
 				{
+					SetTargetLandscapeBasedOnSelection(ClickedControlPoint);
 					SelectControlPoint(ClickedControlPoint);
 				}
 				GEditor->SelectNone(true, true);
@@ -1169,6 +1185,7 @@ public:
 				}
 				else
 				{
+					SetTargetLandscapeBasedOnSelection(ClickedSplineSegment);
 					SelectSegment(ClickedSplineSegment);
 				}
 				GEditor->SelectNone(true, true);
diff --git a/Engine/Source/Editor/LandscapeEditor/Private/LandscapeEditorDetails.cpp b/Engine/Source/Editor/LandscapeEditor/Private/LandscapeEditorDetails.cpp
--- a/Engine/Source/Editor/LandscapeEditor/Private/LandscapeEditorDetails.cpp
+++ b/Engine/Source/Editor/LandscapeEditor/Private/LandscapeEditorDetails.cpp
@@ -231,31 +231,7 @@ void FLandscapeEditorDetails::OnChangeTargetLandscape(TWeakObjectPtr<ULandscapeI
 	FEdModeLandscape* LandscapeEdMode = GetEditorMode();
 	if (LandscapeEdMode)
 	{
-		// Unregister from old one
-		if (LandscapeEdMode->CurrentToolTarget.LandscapeInfo.IsValid())
-		{
-			ALandscapeProxy* LandscapeProxy = LandscapeEdMode->CurrentToolTarget.LandscapeInfo->GetLandscapeProxy();
-			LandscapeProxy->OnMaterialChangedDelegate().RemoveAll(LandscapeEdMode);
-		}
-
-		LandscapeEdMode->CurrentToolTarget.LandscapeInfo = LandscapeInfo.Get();
-		LandscapeEdMode->UpdateTargetList();
-		// force a Leave and Enter the current tool, in case it has something about the current landscape cached
-		LandscapeEdMode->SetCurrentTool(LandscapeEdMode->CurrentToolIndex);
-		if (LandscapeEdMode->CurrentGizmoActor.IsValid())
-		{
-			LandscapeEdMode->CurrentGizmoActor->SetTargetLandscape(LandscapeEdMode->CurrentToolTarget.LandscapeInfo.Get());
-		}
-
-		// register to new one
-		if (LandscapeEdMode->CurrentToolTarget.LandscapeInfo.IsValid())
-		{
-			ALandscapeProxy* LandscapeProxy = LandscapeEdMode->CurrentToolTarget.LandscapeInfo->GetLandscapeProxy();
-			LandscapeProxy->OnMaterialChangedDelegate().AddRaw(LandscapeEdMode, &FEdModeLandscape::OnLandscapeMaterialChangedDelegate);
-		}
-
-		LandscapeEdMode->UpdateTargetList();
-		LandscapeEdMode->UpdateShownLayerList();
+		LandscapeEdMode->SetTargetLandscape(LandscapeInfo);
 	}
 }
 
