Fixed clang compilation error

diff --git a/Engine/Source/Editor/LandscapeEditor/Private/LandscapeEdModeSplineTools.cpp b/Engine/Source/Editor/LandscapeEditor/Private/LandscapeEdModeSplineTools.cpp
--- a/Engine/Source/Editor/LandscapeEditor/Private/LandscapeEdModeSplineTools.cpp
+++ b/Engine/Source/Editor/LandscapeEditor/Private/LandscapeEdModeSplineTools.cpp
@@ -1094,7 +1094,7 @@ public:
 	}
 
 	template <typename T>
-	void SetTargetLandscapeBasedOnSelection(typename T* Selection)
+	void SetTargetLandscapeBasedOnSelection(T* Selection)
 	{
 		check(Selection);
 		if (ALandscapeProxy* LandscapeProxy = Selection->GetTypedOuter<ALandscapeProxy>())
