Fixed clang template compilation error

diff --git a/Engine/Source/Editor/LandscapeEditor/Private/LandscapeEdModeSplineTools.cpp b/Engine/Source/Editor/LandscapeEditor/Private/LandscapeEdModeSplineTools.cpp
--- a/Engine/Source/Editor/LandscapeEditor/Private/LandscapeEdModeSplineTools.cpp
+++ b/Engine/Source/Editor/LandscapeEditor/Private/LandscapeEdModeSplineTools.cpp
@@ -1097,7 +1097,7 @@ public:
 	void SetTargetLandscapeBasedOnSelection(T* Selection)
 	{
 		check(Selection);
-		if (ALandscapeProxy* LandscapeProxy = Selection->GetTypedOuter<ALandscapeProxy>())
+		if (ALandscapeProxy* LandscapeProxy = Selection->template GetTypedOuter<ALandscapeProxy>())
 		{
 			ALandscape* NewLandscapeActor = LandscapeProxy->GetLandscapeActor();
 			if (NewLandscapeActor && (NewLandscapeActor != EdMode->GetLandscape()))
