Fixed incorrect landscape spline end falloff when deleting a landscape control point or a segment

diff --git a/Engine/Source/Editor/LandscapeEditor/Private/LandscapeEdModeSplineTools.cpp b/Engine/Source/Editor/LandscapeEditor/Private/LandscapeEdModeSplineTools.cpp
--- a/Engine/Source/Editor/LandscapeEditor/Private/LandscapeEdModeSplineTools.cpp
+++ b/Engine/Source/Editor/LandscapeEditor/Private/LandscapeEdModeSplineTools.cpp
@@ -532,16 +532,8 @@ public:
 
 		SplinesComponent->Segments.Remove(ToDelete);
 
-		// Control points' points are currently based on connected segments, so need to be updated.
-		if (ToDelete->Connections[0].ControlPoint->Mesh != NULL)
-		{
-			ToDelete->Connections[0].ControlPoint->UpdateSplinePoints();
-		}
-		if (ToDelete->Connections[1].ControlPoint->Mesh != NULL)
-		{
-			ToDelete->Connections[1].ControlPoint->UpdateSplinePoints();
-		}
-
+		ToDelete->Connections[0].ControlPoint->UpdateSplinePoints();
+		ToDelete->Connections[1].ControlPoint->UpdateSplinePoints();
 
 		EdMode->AutoUpdateDirtyLandscapeSplines();
 		SplinesComponent->MarkRenderStateDirty();
@@ -611,6 +603,8 @@ public:
 			ULandscapeSplineControlPoint* OtherEnd = Connection.GetFarConnection().ControlPoint;
 			OtherEnd->Modify();
 			OtherEnd->ConnectedSegments.Remove(FLandscapeSplineConnection(Connection.Segment, 1 - Connection.End));
+			OtherEnd->UpdateSplinePoints();
+
 			SplinesComponent->Segments.Remove(Connection.Segment);
 
 			if (bInDeleteLooseEnds)
