- Fixed dragging control points not triggering landscape spline update using Ctrl+LMB or Ctrl+LMB+RMB.

diff --git a/Engine/Source/Editor/LandscapeEditor/Private/LandscapeEdModeSplineTools.cpp b/Engine/Source/Editor/LandscapeEditor/Private/LandscapeEdModeSplineTools.cpp
--- a/Engine/Source/Editor/LandscapeEditor/Private/LandscapeEdModeSplineTools.cpp
+++ b/Engine/Source/Editor/LandscapeEditor/Private/LandscapeEdModeSplineTools.cpp
@@ -1521,6 +1521,12 @@ public:
 			}
 		}
 
+		// To avoid updating while Ctrl+LMB / Ctrl+RMB+LMB, handle the case one button(s) are released
+		if (InKey == EKeys::RightMouseButton && IsCtrlDown(InViewport) && InEvent == IE_Released && SelectedSplineControlPoints.Num() > 0)
+		{
+			EdMode->AutoUpdateDirtyLandscapeSplines();
+		}
+
 		return false;
 	}
 
