- Removed Landscape Layers section when there's Landscape Editor doesn't have any selected Landscape
- Landscape Layers Warning TextBlock now supports multlines/auto-wrap

diff --git a/Engine/Source/Editor/LandscapeEditor/Private/LandscapeEditorDetailCustomization_Layers.cpp b/Engine/Source/Editor/LandscapeEditor/Private/LandscapeEditorDetailCustomization_Layers.cpp
--- a/Engine/Source/Editor/LandscapeEditor/Private/LandscapeEditorDetailCustomization_Layers.cpp
+++ b/Engine/Source/Editor/LandscapeEditor/Private/LandscapeEditorDetailCustomization_Layers.cpp
@@ -18,7 +18,7 @@
 #include "Widgets/Notifications/SErrorText.h"
 #include "Widgets/Input/SComboButton.h"
 #include "Widgets/Input/SCheckBox.h"
-#include "Widgets/Input/SEditableTextBox.h"
+#include "Widgets/Input/SMultiLineEditableTextBox.h"
 #include "EditorModeManager.h"
 #include "EditorModes.h"
 #include "LandscapeEditorModule.h"
@@ -65,17 +65,21 @@ void FLandscapeEditorDetailCustomization_Layers::CustomizeDetails(IDetailLayoutB
 	FEdModeLandscape* LandscapeEdMode = GetEditorMode();
 	if (LandscapeEdMode && LandscapeEdMode->CurrentToolMode != nullptr)
 	{
-		LayerCategory.AddCustomBuilder(MakeShareable(new FLandscapeEditorCustomNodeBuilder_Layers(DetailBuilder.GetThumbnailPool().ToSharedRef())));
-	}
+		if (LandscapeEdMode->GetLandscape())
+		{
+			LayerCategory.AddCustomBuilder(MakeShareable(new FLandscapeEditorCustomNodeBuilder_Layers(DetailBuilder.GetThumbnailPool().ToSharedRef())));
 
-	LayerCategory.AddCustomRow(FText())
-	.Visibility(TAttribute<EVisibility>::Create(TAttribute<EVisibility>::FGetter::CreateLambda([]() { return ShoudShowLayersErrorMessageTip() ? EVisibility::Visible : EVisibility::Collapsed; })))
-	[
-		SNew(SEditableTextBox)
-		.Font(DetailBuilder.GetDetailFontBold())
-		.BackgroundColor(TAttribute<FSlateColor>::Create(TAttribute<FSlateColor>::FGetter::CreateLambda([]() { return FEditorStyle::GetColor("ErrorReporting.WarningBackgroundColor"); })))
-		.Text(TAttribute<FText>::Create(TAttribute<FText>::FGetter::CreateStatic(&FLandscapeEditorDetailCustomization_Layers::GetLayersErrorMessageText)))
-	];
+			LayerCategory.AddCustomRow(FText())
+				.Visibility(TAttribute<EVisibility>::Create(TAttribute<EVisibility>::FGetter::CreateLambda([]() { return ShoudShowLayersErrorMessageTip() ? EVisibility::Visible : EVisibility::Collapsed; })))
+				[
+					SNew(SMultiLineEditableTextBox)
+					.Font(DetailBuilder.GetDetailFontBold())
+					.BackgroundColor(TAttribute<FSlateColor>::Create(TAttribute<FSlateColor>::FGetter::CreateLambda([]() { return FEditorStyle::GetColor("ErrorReporting.WarningBackgroundColor"); })))
+					.Text(TAttribute<FText>::Create(TAttribute<FText>::FGetter::CreateStatic(&FLandscapeEditorDetailCustomization_Layers::GetLayersErrorMessageText)))
+					.AutoWrapText(true)
+				];
+		}
+	}
 }
 END_SLATE_FUNCTION_BUILD_OPTIMIZATION
 
