- Fixed Landscape Copy/Paste Tool not copying height when option "Gizmo copy/paste all layers" is unchecked.

diff --git a/Engine/Source/Editor/LandscapeEditor/Private/LandscapeEdModeComponentTools.cpp b/Engine/Source/Editor/LandscapeEditor/Private/LandscapeEdModeComponentTools.cpp
--- a/Engine/Source/Editor/LandscapeEditor/Private/LandscapeEdModeComponentTools.cpp
+++ b/Engine/Source/Editor/LandscapeEditor/Private/LandscapeEdModeComponentTools.cpp
@@ -1258,7 +1258,7 @@ public:
 						{
 							// Don't try to copy data for null layers
 							if ((bApplyToAll && i >= 0 && !LandscapeInfo->Layers[i].LayerInfoObj) ||
-								(!bApplyToAll && !EdMode->CurrentToolTarget.LayerInfo.Get()))
+								(!bApplyToAll && (EdMode->CurrentToolTarget.TargetType != ELandscapeToolTargetType::Heightmap) && !EdMode->CurrentToolTarget.LayerInfo.Get()))
 							{
 								continue;
 							}
