- Support Landscape Copy/Paste Tool with Landscape Layer System activated
- Copies from selected layer and Pastes into selected layer (can be different layers)

diff --git a/Engine/Source/Editor/LandscapeEditor/Private/LandscapeEdModeComponentTools.cpp b/Engine/Source/Editor/LandscapeEditor/Private/LandscapeEdModeComponentTools.cpp
--- a/Engine/Source/Editor/LandscapeEditor/Private/LandscapeEdModeComponentTools.cpp
+++ b/Engine/Source/Editor/LandscapeEditor/Private/LandscapeEdModeComponentTools.cpp
@@ -1809,6 +1809,12 @@ public:
 		return ELandscapeToolTargetTypeMask::FromType(ToolTarget::TargetType);
 	}
 
+	virtual ELandscapeLayersContentUpdateFlag GetBeginToolContentUpdateFlag() const override { return ELandscapeLayersContentUpdateFlag::All_Render; }
+
+	virtual ELandscapeLayersContentUpdateFlag GetTickToolContentUpdateFlag() const override { return ELandscapeLayersContentUpdateFlag::All_Render; }
+
+	virtual ELandscapeLayersContentUpdateFlag GetEndToolContentUpdateFlag() const override { return ELandscapeLayersContentUpdateFlag::All; }
+
 	virtual bool BeginTool(FEditorViewportClient* ViewportClient, const FLandscapeToolTarget& InTarget, const FVector& InHitLocation) override
 	{
 		this->EdMode->GizmoBrush->Tick(ViewportClient, 0.1f);
diff --git a/Engine/Source/Editor/LandscapeEditor/Private/LandscapeEdModeTools.h b/Engine/Source/Editor/LandscapeEditor/Private/LandscapeEdModeTools.h
--- a/Engine/Source/Editor/LandscapeEditor/Private/LandscapeEdModeTools.h
+++ b/Engine/Source/Editor/LandscapeEditor/Private/LandscapeEdModeTools.h
@@ -1121,6 +1121,23 @@ public:
 	{
 	}
 
+	virtual ELandscapeLayersContentUpdateFlag GetBeginToolContentUpdateFlag() const
+	{
+		bool bUpdateHeightmap = this->EdMode->CurrentToolTarget.TargetType == ELandscapeToolTargetType::Type::Heightmap; 
+		return bUpdateHeightmap ? ELandscapeLayersContentUpdateFlag::Heightmap_Render : ELandscapeLayersContentUpdateFlag::Weightmap_Render;
+	}
+
+	virtual ELandscapeLayersContentUpdateFlag GetTickToolContentUpdateFlag() const
+	{
+		return GetBeginToolContentUpdateFlag();
+	}
+
+	virtual ELandscapeLayersContentUpdateFlag GetEndToolContentUpdateFlag() const
+	{
+		bool bUpdateHeightmap = this->EdMode->CurrentToolTarget.TargetType == ELandscapeToolTargetType::Type::Heightmap; 
+		return bUpdateHeightmap ? ELandscapeLayersContentUpdateFlag::Heightmap_All : ELandscapeLayersContentUpdateFlag::Weightmap_All;
+	}
+
 	virtual bool BeginTool(FEditorViewportClient* ViewportClient, const FLandscapeToolTarget& InTarget, const FVector& InHitLocation) override
 	{
 		if (GetMutableDefault<UEditorExperimentalSettings>()->bLandscapeLayerSystem)
@@ -1128,8 +1145,7 @@ public:
 			ALandscape* Landscape = this->EdMode->GetLandscape();
 			if (Landscape)
 			{
-				bool bUpdateHeightmap = this->EdMode->CurrentToolTarget.TargetType == ELandscapeToolTargetType::Type::Heightmap;
-				Landscape->RequestLayersContentUpdate(bUpdateHeightmap ? ELandscapeLayersContentUpdateFlag::Heightmap_Render : ELandscapeLayersContentUpdateFlag::Weightmap_Render);
+				Landscape->RequestLayersContentUpdate(GetBeginToolContentUpdateFlag());
 				Landscape->SetEditingLayer(this->EdMode->GetCurrentLayerGuid());
 			}
 		}
@@ -1183,8 +1199,7 @@ public:
 				ALandscape* Landscape = this->EdMode->CurrentToolTarget.LandscapeInfo->LandscapeActor.Get();
 				if (Landscape != nullptr)
 				{
-					bool bUpdateHeightmap = this->EdMode->CurrentToolTarget.TargetType == ELandscapeToolTargetType::Type::Heightmap;
-					Landscape->RequestLayersContentUpdate(bUpdateHeightmap ? ELandscapeLayersContentUpdateFlag::Heightmap_Render : ELandscapeLayersContentUpdateFlag::Weightmap_Render);
+					Landscape->RequestLayersContentUpdate(GetTickToolContentUpdateFlag());
 				}
 			}
 		}
@@ -1197,8 +1212,7 @@ public:
 			ALandscape* Landscape = this->EdMode->GetLandscape();
 			if (Landscape)
 			{
-				bool bUpdateHeightmap = this->EdMode->CurrentToolTarget.TargetType == ELandscapeToolTargetType::Type::Heightmap;
-				Landscape->RequestLayersContentUpdate(bUpdateHeightmap ? ELandscapeLayersContentUpdateFlag::Heightmap_All : ELandscapeLayersContentUpdateFlag::Weightmap_All);
+				Landscape->RequestLayersContentUpdate(GetEndToolContentUpdateFlag());
 				Landscape->SetEditingLayer();
 			}
 		}
