- Support Landscape Mirror Tool with Landscape Layer System activated (mirrors selected layer)
- Decoupled layer locking mechanism from reserved for splines to allow manual locking of reserved layer
- Blocked Pasting in locked landscape layer

diff --git a/Engine/Source/Editor/LandscapeEditor/Private/LandscapeEdMode.cpp b/Engine/Source/Editor/LandscapeEditor/Private/LandscapeEdMode.cpp
--- a/Engine/Source/Editor/LandscapeEditor/Private/LandscapeEdMode.cpp
+++ b/Engine/Source/Editor/LandscapeEditor/Private/LandscapeEdMode.cpp
@@ -1596,6 +1596,20 @@ bool FEdModeLandscape::ProcessEditPaste()
 
 	bool Result = false;
 
+	
+	FLandscapeLayer* SplinesLayer = nullptr;
+	if (CurrentTool == (FLandscapeTool*)SplinesTool)
+	{
+		ALandscape* Landscape = GetLandscape();
+		SplinesLayer = Landscape ? Landscape->GetLandscapeSplinesReservedLayer() : nullptr;
+	}
+	FText Reason;
+	if (!CanEditLayer(&Reason, SplinesLayer))
+	{
+		FMessageDialog::Open(EAppMsgType::Ok, Reason);
+		return Result;
+	}
+	
 	if (NewLandscapePreviewMode == ENewLandscapePreviewMode::None)
 	{
 		if (CurrentTool != NULL)
@@ -4243,10 +4257,8 @@ void FEdModeLandscape::SetLayerVisibility(bool bInVisible, int32 InLayerIndex)
 
 bool FEdModeLandscape::IsLayerLocked(int32 InLayerIndex) const
 {
-	ALandscape* Landscape = GetLandscape();
 	const FLandscapeLayer* Layer = GetLayer(InLayerIndex);
-	const FLandscapeLayer* LayerReservedForSplines = Landscape ? Landscape->GetLandscapeSplinesReservedLayer() : nullptr;
-	return Layer ? (Layer->bLocked || Layer == LayerReservedForSplines) : false;
+	return Layer && Layer->bLocked;
 }
 
 void FEdModeLandscape::SetLayerLocked(int32 InLayerIndex, bool bInLocked)
@@ -4390,13 +4402,13 @@ void FEdModeLandscape::AutoUpdateDirtyLandscapeSplines()
 	}
 }
 
-bool FEdModeLandscape::CanEditLayer(FText* Reason)
+bool FEdModeLandscape::CanEditLayer(FText* Reason /*=nullptr*/, FLandscapeLayer* InLayer /*= nullptr*/)
 {
 	if (GetMutableDefault<UEditorExperimentalSettings>()->bLandscapeLayerSystem)
 	{
 		ALandscape* Landscape = GetLandscape();
-		FLandscapeLayer* CurrentLayer = GetCurrentLayer();
-		if (!CurrentLayer)
+		FLandscapeLayer* TargetLayer = InLayer ? InLayer : GetCurrentLayer();
+		if (!TargetLayer)
 		{
 			if (Reason)
 			{
@@ -4404,7 +4416,7 @@ bool FEdModeLandscape::CanEditLayer(FText* Reason)
 			}
 			return false;
 		}
-		else if (!CurrentLayer->bVisible)
+		else if (!TargetLayer->bVisible)
 		{
 			if (Reason)
 			{
@@ -4412,7 +4424,7 @@ bool FEdModeLandscape::CanEditLayer(FText* Reason)
 			}
 			return false;
 		}
-		else if (CurrentLayer->bLocked)
+		else if (TargetLayer->bLocked)
 		{
 			if (Reason)
 			{
@@ -4420,34 +4432,38 @@ bool FEdModeLandscape::CanEditLayer(FText* Reason)
 			}
 			return false;
 		}
-		else if (CurrentTool && (CurrentTool != (FLandscapeTool*)SplinesTool) && Landscape && (CurrentLayer == Landscape->GetLandscapeSplinesReservedLayer()))
+		else if (CurrentTool)
 		{
-			if (Reason)
+			int32 TargetLayerIndex = Landscape->LandscapeLayers.IndexOfByPredicate([TargetLayeyGuid = TargetLayer->Guid](const FLandscapeLayer& OtherLayer) { return OtherLayer.Guid == TargetLayeyGuid; });
+
+			if ((CurrentTool != (FLandscapeTool*)SplinesTool) && Landscape && (TargetLayer == Landscape->GetLandscapeSplinesReservedLayer()))
 			{
-				*Reason = NSLOCTEXT("UnrealEd", "LandscapeLayerReservedForSplines", "This layer is reserved for Landscape Splines.");
+				if (Reason)
+				{
+					*Reason = NSLOCTEXT("UnrealEd", "LandscapeLayerReservedForSplines", "This layer is reserved for Landscape Splines.");
+				}
+				return false;
 			}
-			return false;
-		}
-		else if (CurrentTool &&
-				 (GetCurrentLayerIndex() != 0) &&
-				 (CurrentTool->GetToolName() == FName("Ramp") || CurrentTool->GetToolName() == FName("Flatten")))
-		{
-			if (Reason)
+			else if ((TargetLayerIndex > 0) && (CurrentTool->GetToolName() == FName("Ramp") || CurrentTool->GetToolName() == FName("Flatten")))
 			{
-				FLandscapeLayer* Layer = GetLayer(0);
-				*Reason = FText::Format(NSLOCTEXT("UnrealEd", "LandscapeLayersToolAvailableOnlyOnFirstLayer", "{0} Tool is only available on the first layer {1}."), CurrentTool->GetDisplayName(), Layer ? FText::FromName(Layer->Name) : FText::GetEmpty());
+				if (Reason)
+				{
+					FLandscapeLayer* Layer = GetLayer(0);
+					*Reason = FText::Format(NSLOCTEXT("UnrealEd", "LandscapeLayersToolAvailableOnlyOnFirstLayer", "{0} Tool is only available on the first layer {1}."), CurrentTool->GetDisplayName(), Layer ? FText::FromName(Layer->Name) : FText::GetEmpty());
+				}
+				return false;
 			}
-			return false;
-		}
-		else if (CurrentTool && CurrentTool->GetToolName() == FName("Retopologize"))
-		{
-			if (Reason)
+			else if (CurrentTool->GetToolName() == FName("Retopologize"))
 			{
-				*Reason = FText::Format(NSLOCTEXT("UnrealEd", "LandscapeLayersNoSupportForRetopologize", "{0} Tool is not available with the Landscape Layer System."), CurrentTool->GetDisplayName());
+				if (Reason)
+				{
+					*Reason = FText::Format(NSLOCTEXT("UnrealEd", "LandscapeLayersNoSupportForRetopologize", "{0} Tool is not available with the Landscape Layer System."), CurrentTool->GetDisplayName());
+				}
+				return false;
 			}
-			return false;
 		}
 	}
+
 	if (CurrentToolTarget.TargetType == ELandscapeToolTargetType::Weightmap && CurrentToolTarget.LayerInfo == NULL && CurrentTool->GetToolName() != FName("BPCustom"))
 	{
 		if (Reason)
diff --git a/Engine/Source/Editor/LandscapeEditor/Private/LandscapeEdMode.h b/Engine/Source/Editor/LandscapeEditor/Private/LandscapeEdMode.h
--- a/Engine/Source/Editor/LandscapeEditor/Private/LandscapeEdMode.h
+++ b/Engine/Source/Editor/LandscapeEditor/Private/LandscapeEdMode.h
@@ -512,7 +512,7 @@ public:
 	void SetCurrentLayerSubstractiveBlendStatus(bool InStatus, const TWeakObjectPtr<ULandscapeLayerInfoObject>& InLayerInfoObj);
 	void UpdateLandscapeSplines(bool bUpdateOnlySelected = false);
 	void AutoUpdateDirtyLandscapeSplines();
-	bool CanEditLayer(FText* Reason = nullptr);
+	bool CanEditLayer(FText* Reason = nullptr, FLandscapeLayer* InLayer = nullptr);
 
 	void AddBrushToCurrentLayer(int32 InTargetType, class ALandscapeBlueprintCustomBrush* InBrush);
 	void RemoveBrushFromCurrentLayer(int32 InTargetType, class ALandscapeBlueprintCustomBrush* InBrush);
diff --git a/Engine/Source/Editor/LandscapeEditor/Private/LandscapeEdModeMirrorTool.cpp b/Engine/Source/Editor/LandscapeEditor/Private/LandscapeEdModeMirrorTool.cpp
--- a/Engine/Source/Editor/LandscapeEditor/Private/LandscapeEdModeMirrorTool.cpp
+++ b/Engine/Source/Editor/LandscapeEditor/Private/LandscapeEdModeMirrorTool.cpp
@@ -18,7 +18,8 @@
 #include "LandscapeDataAccess.h"
 #include "LandscapeRender.h"
 #include "LandscapeHeightfieldCollisionComponent.h"
-//#include "LandscapeDataAccess.h"
+#include "Landscape.h"
+#include "Misc/MessageDialog.h"
 
 #define LOCTEXT_NAMESPACE "Landscape"
 
@@ -500,7 +501,14 @@ protected:
 public:
 	virtual void ApplyMirror()
 	{
+		FText Reason;
+		if (!EdMode->CanEditLayer(&Reason))
+		{
+			FMessageDialog::Open(EAppMsgType::Ok, Reason);
+			return;
+		}
 		FScopedTransaction Transaction(LOCTEXT("Mirror_Apply", "Landscape Editing: Mirror Landscape"));
+		FScopedSetLandscapeEditingLayer Scope(EdMode->GetLandscape(), EdMode->GetCurrentLayerGuid(), [&] { EdMode->RequestLayersContentUpdate(true); });
 
 		const ULandscapeInfo* const LandscapeInfo = EdMode->CurrentToolTarget.LandscapeInfo.Get();
 		const ALandscapeProxy* const LandscapeProxy = LandscapeInfo->GetLandscapeProxy();
diff --git a/Engine/Source/Editor/LandscapeEditor/Private/LandscapeEditorDetailCustomization_Layers.cpp b/Engine/Source/Editor/LandscapeEditor/Private/LandscapeEditorDetailCustomization_Layers.cpp
--- a/Engine/Source/Editor/LandscapeEditor/Private/LandscapeEditorDetailCustomization_Layers.cpp
+++ b/Engine/Source/Editor/LandscapeEditor/Private/LandscapeEditorDetailCustomization_Layers.cpp
@@ -614,8 +614,10 @@ EVisibility FLandscapeEditorCustomNodeBuilder_Layers::GetLayerAlphaVisibility(in
 bool FLandscapeEditorCustomNodeBuilder_Layers::IsLayerEditionEnabled(int32 InLayerIndex) const
 {
 	FEdModeLandscape* LandscapeEdMode = GetEditorMode();
-	bool bIsLocked = LandscapeEdMode && LandscapeEdMode->IsLayerLocked(InLayerIndex);
-	return !bIsLocked;
+	ALandscape* Landscape = LandscapeEdMode ? LandscapeEdMode->GetLandscape() : nullptr;
+	const FLandscapeLayer* Layer = LandscapeEdMode ? LandscapeEdMode->GetLayer(InLayerIndex) : nullptr;
+	const FLandscapeLayer* LayerReservedForSplines = Landscape ? Landscape->GetLandscapeSplinesReservedLayer() : nullptr;
+	return Layer && !Layer->bLocked && (Layer != LayerReservedForSplines);
 }
 
 const FSlateBrush* FLandscapeEditorCustomNodeBuilder_Layers::GetLockBrushForLayer(int32 InLayerIndex) const
diff --git a/Engine/Source/Editor/LandscapeEditor/Private/LandscapeEditorDetailCustomization_MiscTools.cpp b/Engine/Source/Editor/LandscapeEditor/Private/LandscapeEditorDetailCustomization_MiscTools.cpp
--- a/Engine/Source/Editor/LandscapeEditor/Private/LandscapeEditorDetailCustomization_MiscTools.cpp
+++ b/Engine/Source/Editor/LandscapeEditor/Private/LandscapeEditorDetailCustomization_MiscTools.cpp
@@ -216,7 +216,7 @@ void FLandscapeEditorDetailCustomization_MiscTools::CustomizeDetails(IDetailLayo
 				.Padding(3, 0, 0, 0)
 				[
 					SNew(SButton)
-					//.IsEnabled_Static(&FLandscapeEditorDetailCustomization_MiscTools::GetApplyMirrorButtonIsEnabled)
+					.IsEnabled_Lambda([]() { FEdModeLandscape* LandscapeEdMode = GetEditorMode(); return LandscapeEdMode && LandscapeEdMode->CanEditLayer(); })
 					.Text(LOCTEXT("Mirror.Apply", "Apply"))
 					.HAlign(HAlign_Center)
 					.OnClicked_Static(&FLandscapeEditorDetailCustomization_MiscTools::OnApplyMirrorButtonClicked)
