- Landscape Subtractive Blend label appeared twice in LayerInfo when activating Layer Debug landscape visualizer

diff --git a/Engine/Source/Editor/LandscapeEditor/Private/LandscapeEditorDetailCustomization_TargetLayers.cpp b/Engine/Source/Editor/LandscapeEditor/Private/LandscapeEditorDetailCustomization_TargetLayers.cpp
--- a/Engine/Source/Editor/LandscapeEditor/Private/LandscapeEditorDetailCustomization_TargetLayers.cpp
+++ b/Engine/Source/Editor/LandscapeEditor/Private/LandscapeEditorDetailCustomization_TargetLayers.cpp
@@ -712,7 +712,7 @@ TSharedPtr<SWidget> FLandscapeEditorCustomNodeBuilder_TargetLayers::GenerateRow(
 							.OnCheckStateChanged_Static(&FLandscapeEditorCustomNodeBuilder_TargetLayers::OnLayersSubstractiveBlendChanged, Target)
 							[
 								SNew(STextBlock)
-								.Text(LOCTEXT("ViewMode.Debug_None", "Substractive Blend"))
+								.Text(LOCTEXT("SubtractiveBlend", "Subtractive Blend"))
 							]
 						]
 					]
