- Landscape Layer System: Undo/Redo Fixes ( Layers: rename, reorder, visibility, locked, alpha, new landscape component size )

#rnx

diff --git a/Engine/Source/Editor/LandscapeEditor/Private/LandscapeEdMode.cpp b/Engine/Source/Editor/LandscapeEditor/Private/LandscapeEdMode.cpp
--- a/Engine/Source/Editor/LandscapeEditor/Private/LandscapeEdMode.cpp
+++ b/Engine/Source/Editor/LandscapeEditor/Private/LandscapeEdMode.cpp
@@ -4263,10 +4263,10 @@ bool FEdModeLandscape::IsLayerLocked(int32 InLayerIndex) const
 
 void FEdModeLandscape::SetLayerLocked(int32 InLayerIndex, bool bInLocked)
 {
-	FLandscapeLayer* Layer = GetLayer(InLayerIndex);
-	if (Layer)
+	ALandscape* Landscape = GetLandscape();
+	if (Landscape)
 	{
-		Layer->bLocked = bInLocked;
+		Landscape->SetLayerLocked(InLayerIndex, bInLocked);
 	}
 }
 
diff --git a/Engine/Source/Editor/LandscapeEditor/Private/LandscapeEditorDetailCustomization_Layers.cpp b/Engine/Source/Editor/LandscapeEditor/Private/LandscapeEditorDetailCustomization_Layers.cpp
--- a/Engine/Source/Editor/LandscapeEditor/Private/LandscapeEditorDetailCustomization_Layers.cpp
+++ b/Engine/Source/Editor/LandscapeEditor/Private/LandscapeEditorDetailCustomization_Layers.cpp
@@ -114,6 +114,7 @@ FEdModeLandscape* FLandscapeEditorCustomNodeBuilder_Layers::GetEditorMode()
 FLandscapeEditorCustomNodeBuilder_Layers::FLandscapeEditorCustomNodeBuilder_Layers(TSharedRef<FAssetThumbnailPool> InThumbnailPool)
 	: ThumbnailPool(InThumbnailPool)
 	, CurrentEditingInlineTextBlock(INDEX_NONE)
+	, CurrentSlider(INDEX_NONE)
 {
 }
 
@@ -275,7 +276,18 @@ TSharedPtr<SWidget> FLandscapeEditorCustomNodeBuilder_Layers::GenerateRow(int32
 					.IsEnabled(this, &FLandscapeEditorCustomNodeBuilder_Layers::IsLayerEditionEnabled, InLayerIndex)
 					.Visibility(this, &FLandscapeEditorCustomNodeBuilder_Layers::GetLayerAlphaVisibility, InLayerIndex)
 					.Value(this, &FLandscapeEditorCustomNodeBuilder_Layers::GetLayerAlpha, InLayerIndex)
-					.OnValueChanged(this, &FLandscapeEditorCustomNodeBuilder_Layers::SetLayerAlpha, InLayerIndex)
+					.OnValueChanged_Lambda([=](float InValue) { SetLayerAlpha(InValue, InLayerIndex, false); })
+					.OnValueCommitted_Lambda([=](float InValue, ETextCommit::Type InCommitType) { SetLayerAlpha(InValue, InLayerIndex, true); })
+					.OnBeginSliderMovement_Lambda([=]()
+					{
+						CurrentSlider = InLayerIndex;
+						GEditor->BeginTransaction(LOCTEXT("Landscape_Layers_SetAlpha", "Set Layer Alpha"));
+					})
+					.OnEndSliderMovement_Lambda([=](double)
+					{
+						GEditor->EndTransaction();
+						CurrentSlider = INDEX_NONE;
+					})
 				]		
 			]
 		];	
@@ -566,13 +578,19 @@ TOptional<float> FLandscapeEditorCustomNodeBuilder_Layers::GetLayerAlpha(int32 I
 	return 1.0f;
 }
 
-void FLandscapeEditorCustomNodeBuilder_Layers::SetLayerAlpha(float InAlpha, int32 InLayerIndex)
+void FLandscapeEditorCustomNodeBuilder_Layers::SetLayerAlpha(float InAlpha, int32 InLayerIndex, bool bCommit)
 {
 	FEdModeLandscape* LandscapeEdMode = GetEditorMode();
-
 	if (LandscapeEdMode)
 	{
-		const FScopedTransaction Transaction(LOCTEXT("Landscape_Layers_SetAlpha", "Set Layer Alpha"));
+		// We get multiple commits when editing through the text box
+		if (LandscapeEdMode->GetLayerAlpha(InLayerIndex) == InAlpha)
+		{
+			return;
+		}
+
+		FScopedTransaction Transaction(LOCTEXT("Landscape_Layers_SetAlpha", "Set Layer Alpha"), CurrentSlider == INDEX_NONE && bCommit);
+		// Set Value when using slider or when committing text
 		LandscapeEdMode->SetLayerAlpha(InLayerIndex, InAlpha);
 	}
 }
@@ -582,6 +600,7 @@ FReply FLandscapeEditorCustomNodeBuilder_Layers::OnToggleVisibility(int32 InLaye
 	FEdModeLandscape* LandscapeEdMode = GetEditorMode();
 	if (LandscapeEdMode)
 	{
+		const FScopedTransaction Transaction(LOCTEXT("Landscape_Layers_SetVisibility", "Set Layer Visibility"));
 		LandscapeEdMode->SetLayerVisibility(!LandscapeEdMode->IsLayerVisible(InLayerIndex), InLayerIndex);
 	}
 	return FReply::Handled();
@@ -599,6 +618,7 @@ FReply FLandscapeEditorCustomNodeBuilder_Layers::OnToggleLock(int32 InLayerIndex
 	FEdModeLandscape* LandscapeEdMode = GetEditorMode();
 	if (LandscapeEdMode)
 	{
+		const FScopedTransaction Transaction(LOCTEXT("Landscape_Layers_Locked", "Set Layer Locked"));
 		LandscapeEdMode->SetLayerLocked(InLayerIndex, !LandscapeEdMode->IsLayerLocked(InLayerIndex));
 	}
 	return FReply::Handled();
@@ -667,6 +687,7 @@ FReply FLandscapeEditorCustomNodeBuilder_Layers::HandleAcceptDrop(FDragDropEvent
 		{
 			int32 StartingLayerIndex = DragDropOperation->SlotIndexBeingDragged;
 			int32 DestinationLayerIndex = SlotIndex;
+			const FScopedTransaction Transaction(LOCTEXT("Landscape_Layers_Reorder", "Reorder Layer"));
 			if (Landscape->ReorderLayer(StartingLayerIndex, DestinationLayerIndex))
 			{
 				LandscapeEdMode->SetCurrentLayer(DestinationLayerIndex);
diff --git a/Engine/Source/Editor/LandscapeEditor/Private/LandscapeEditorDetailCustomization_Layers.h b/Engine/Source/Editor/LandscapeEditor/Private/LandscapeEditorDetailCustomization_Layers.h
--- a/Engine/Source/Editor/LandscapeEditor/Private/LandscapeEditorDetailCustomization_Layers.h
+++ b/Engine/Source/Editor/LandscapeEditor/Private/LandscapeEditorDetailCustomization_Layers.h
@@ -88,8 +88,8 @@ protected:
 	EVisibility GetLayerAlphaVisibility(int32 InLayerIndex) const;
 
 	TOptional<float> GetLayerAlpha(int32 InLayerIndex) const;
-	void SetLayerAlpha(float InAlpha, int32 InLayerIndex);
-
+	void SetLayerAlpha(float InAlpha, int32 InLayerIndex, bool bCommit);
+	
 	FReply OnToggleVisibility(int32 InLayerIndex);
 	const FSlateBrush* GetVisibilityBrushForLayer(int32 InLayerIndex) const;
 	
@@ -102,6 +102,7 @@ private:
 	TArray< TSharedPtr< SInlineEditableTextBlock > > InlineTextBlocks;
 
 	int32 CurrentEditingInlineTextBlock;
+	int32 CurrentSlider;
 };
 
 class FLandscapeListElementDragDropOp : public FDragAndDropVerticalBoxOp
diff --git a/Engine/Source/Editor/LandscapeEditor/Private/LandscapeEditorDetailCustomization_NewLandscape.cpp b/Engine/Source/Editor/LandscapeEditor/Private/LandscapeEditorDetailCustomization_NewLandscape.cpp
--- a/Engine/Source/Editor/LandscapeEditor/Private/LandscapeEditorDetailCustomization_NewLandscape.cpp
+++ b/Engine/Source/Editor/LandscapeEditor/Private/LandscapeEditorDetailCustomization_NewLandscape.cpp
@@ -392,8 +392,24 @@ void FLandscapeEditorDetailCustomization_NewLandscape::CustomizeDetails(IDetailL
 			//.MinSliderValue(TAttribute<TOptional<int32> >(this, &FLandscapeEditorDetailCustomization_NewLandscape::GetMinLandscapeResolution))
 			//.MaxSliderValue(TAttribute<TOptional<int32> >(this, &FLandscapeEditorDetailCustomization_NewLandscape::GetMaxLandscapeResolution))
 			.Value(this, &FLandscapeEditorDetailCustomization_NewLandscape::GetLandscapeResolutionX)
-			.OnValueChanged(this, &FLandscapeEditorDetailCustomization_NewLandscape::OnChangeLandscapeResolutionX)
-			.OnValueCommitted(this, &FLandscapeEditorDetailCustomization_NewLandscape::OnCommitLandscapeResolutionX)
+			.OnValueChanged_Lambda([=](int32 NewValue)
+			{
+				OnChangeLandscapeResolutionX(NewValue, false);
+			})
+			.OnValueCommitted_Lambda([=](int32 NewValue, ETextCommit::Type)
+			{
+				OnChangeLandscapeResolutionX(NewValue, true);
+			})
+			.OnBeginSliderMovement_Lambda([=]()
+			{
+				bUsingSlider = true;
+				GEditor->BeginTransaction(LOCTEXT("ChangeResolutionX_Transaction", "Change Landscape Resolution X"));
+			})
+			.OnEndSliderMovement_Lambda([=](double)
+			{
+				GEditor->EndTransaction();
+				bUsingSlider = false;
+			})
 		]
 		+ SHorizontalBox::Slot()
 		.AutoWidth()
@@ -418,8 +434,24 @@ void FLandscapeEditorDetailCustomization_NewLandscape::CustomizeDetails(IDetailL
 			//.MinSliderValue(TAttribute<TOptional<int32> >(this, &FLandscapeEditorDetailCustomization_NewLandscape::GetMinLandscapeResolution))
 			//.MaxSliderValue(TAttribute<TOptional<int32> >(this, &FLandscapeEditorDetailCustomization_NewLandscape::GetMaxLandscapeResolution))
 			.Value(this, &FLandscapeEditorDetailCustomization_NewLandscape::GetLandscapeResolutionY)
-			.OnValueChanged(this, &FLandscapeEditorDetailCustomization_NewLandscape::OnChangeLandscapeResolutionY)
-			.OnValueCommitted(this, &FLandscapeEditorDetailCustomization_NewLandscape::OnCommitLandscapeResolutionY)
+			.OnValueChanged_Lambda([=](int32 NewValue)
+			{
+				OnChangeLandscapeResolutionY(NewValue, false);
+			})
+			.OnValueCommitted_Lambda([=](int32 NewValue, ETextCommit::Type)
+			{
+				OnChangeLandscapeResolutionY(NewValue, true);
+			})
+			.OnBeginSliderMovement_Lambda([=]()
+			{
+				bUsingSlider = true;
+				GEditor->BeginTransaction(LOCTEXT("ChangeResolutionY_Transaction", "Change Landscape Resolution Y"));
+			})
+			.OnEndSliderMovement_Lambda([=](double)
+			{
+				GEditor->EndTransaction();
+				bUsingSlider = false;
+			})
 		]
 	];
 
@@ -617,37 +649,21 @@ TOptional<int32> FLandscapeEditorDetailCustomization_NewLandscape::GetLandscapeR
 	return 0;
 }
 
-void FLandscapeEditorDetailCustomization_NewLandscape::OnChangeLandscapeResolutionX(int32 NewValue)
+void FLandscapeEditorDetailCustomization_NewLandscape::OnChangeLandscapeResolutionX(int32 NewValue, bool bCommit)
 {
 	FEdModeLandscape* LandscapeEdMode = GetEditorMode();
 	if (LandscapeEdMode != nullptr)
 	{
 		int32 NewComponentCountX = LandscapeEdMode->UISettings->CalcComponentsCount(NewValue);
-		if (NewComponentCountX != LandscapeEdMode->UISettings->NewLandscape_ComponentCount.X)
+		if (NewComponentCountX == LandscapeEdMode->UISettings->NewLandscape_ComponentCount.X)
 		{
-			if (!GEditor->IsTransactionActive())
-			{
-				GEditor->BeginTransaction(LOCTEXT("ChangeResolutionX_Transaction", "Change Landscape Resolution X"));
-			}
-
-			LandscapeEdMode->UISettings->Modify();
-			LandscapeEdMode->UISettings->NewLandscape_ComponentCount.X = NewComponentCountX;
+			return;
 		}
-	}
-}
 
-void FLandscapeEditorDetailCustomization_NewLandscape::OnCommitLandscapeResolutionX(int32 NewValue, ETextCommit::Type CommitInfo)
-{
-	FEdModeLandscape* LandscapeEdMode = GetEditorMode();
-	if (LandscapeEdMode != nullptr)
-	{
-		if (!GEditor->IsTransactionActive())
-		{
-			GEditor->BeginTransaction(LOCTEXT("ChangeResolutionX_Transaction", "Change Landscape Resolution X"));
-		}
+		FScopedTransaction Transaction(LOCTEXT("ChangeResolutionX_Transaction", "Change Landscape Resolution X"), !bUsingSlider && bCommit);
+		
 		LandscapeEdMode->UISettings->Modify();
-		LandscapeEdMode->UISettings->NewLandscape_ComponentCount.X = LandscapeEdMode->UISettings->CalcComponentsCount(NewValue);
-		GEditor->EndTransaction();
+		LandscapeEdMode->UISettings->NewLandscape_ComponentCount.X = NewComponentCountX;
 	}
 }
 
@@ -662,37 +678,21 @@ TOptional<int32> FLandscapeEditorDetailCustomization_NewLandscape::GetLandscapeR
 	return 0;
 }
 
-void FLandscapeEditorDetailCustomization_NewLandscape::OnChangeLandscapeResolutionY(int32 NewValue)
+void FLandscapeEditorDetailCustomization_NewLandscape::OnChangeLandscapeResolutionY(int32 NewValue, bool bCommit)
 {
 	FEdModeLandscape* LandscapeEdMode = GetEditorMode();
 	if (LandscapeEdMode != nullptr)
 	{
 		int32 NewComponentCountY = LandscapeEdMode->UISettings->CalcComponentsCount(NewValue);
-		if (NewComponentCountY != LandscapeEdMode->UISettings->NewLandscape_ComponentCount.Y)
+		if (NewComponentCountY == LandscapeEdMode->UISettings->NewLandscape_ComponentCount.Y)
 		{
-			if (!GEditor->IsTransactionActive())
-			{
-				GEditor->BeginTransaction(LOCTEXT("ChangeResolutionY_Transaction", "Change Landscape Resolution Y"));
-			}
-			
-			LandscapeEdMode->UISettings->Modify();
-			LandscapeEdMode->UISettings->NewLandscape_ComponentCount.Y = NewComponentCountY;
+			return;
 		}
-	}
-}
 
-void FLandscapeEditorDetailCustomization_NewLandscape::OnCommitLandscapeResolutionY(int32 NewValue, ETextCommit::Type CommitInfo)
-{
-	FEdModeLandscape* LandscapeEdMode = GetEditorMode();
-	if (LandscapeEdMode != nullptr)
-	{
-		if (!GEditor->IsTransactionActive())
-		{
-			GEditor->BeginTransaction(LOCTEXT("ChangeResolutionY_Transaction", "Change Landscape Resolution Y"));
-		}
+		FScopedTransaction Transaction(LOCTEXT("ChangeResolutionX_Transaction", "Change Landscape Resolution X"), !bUsingSlider && bCommit);
+	
 		LandscapeEdMode->UISettings->Modify();
-		LandscapeEdMode->UISettings->NewLandscape_ComponentCount.Y = LandscapeEdMode->UISettings->CalcComponentsCount(NewValue);
-		GEditor->EndTransaction();
+		LandscapeEdMode->UISettings->NewLandscape_ComponentCount.Y = NewComponentCountY;
 	}
 }
 
diff --git a/Engine/Source/Editor/LandscapeEditor/Private/LandscapeEditorDetailCustomization_NewLandscape.h b/Engine/Source/Editor/LandscapeEditor/Private/LandscapeEditorDetailCustomization_NewLandscape.h
--- a/Engine/Source/Editor/LandscapeEditor/Private/LandscapeEditorDetailCustomization_NewLandscape.h
+++ b/Engine/Source/Editor/LandscapeEditor/Private/LandscapeEditorDetailCustomization_NewLandscape.h
@@ -23,6 +23,10 @@ class IDetailLayoutBuilder;
 class FLandscapeEditorDetailCustomization_NewLandscape : public FLandscapeEditorDetailCustomization_Base
 {
 public:
+	FLandscapeEditorDetailCustomization_NewLandscape()
+		: bUsingSlider(false)
+	{}
+
 	/** Makes a new instance of this detail layout class for a specific detail view requesting it */
 	static TSharedRef<IDetailCustomization> MakeInstance();
 
@@ -41,13 +45,11 @@ public:
 	static FText GetSectionsPerComponent(TSharedRef<IPropertyHandle> PropertyHandle);
 
 	TOptional<int32> GetLandscapeResolutionX() const;
-	void OnChangeLandscapeResolutionX(int32 NewValue);
-	void OnCommitLandscapeResolutionX(int32 NewValue, ETextCommit::Type CommitInfo);
-
+	void OnChangeLandscapeResolutionX(int32 NewValue, bool bCommit);
+	
 	TOptional<int32> GetLandscapeResolutionY() const;
-	void OnChangeLandscapeResolutionY(int32 NewValue);
-	void OnCommitLandscapeResolutionY(int32 NewValue, ETextCommit::Type CommitInfo);
-
+	void OnChangeLandscapeResolutionY(int32 NewValue, bool bCommit);
+	
 	TOptional<int32> GetMinLandscapeResolution() const;
 	TOptional<int32> GetMaxLandscapeResolution() const;
 
@@ -82,6 +84,7 @@ public:
 
 protected:
 	TArray<FLandscapeFileResolution> ImportResolutions;
+	bool bUsingSlider;
 };
 
 class FLandscapeEditorStructCustomization_FLandscapeImportLayer : public FLandscapeEditorStructCustomization_Base
diff --git a/Engine/Source/Runtime/Landscape/Classes/Landscape.h b/Engine/Source/Runtime/Landscape/Classes/Landscape.h
--- a/Engine/Source/Runtime/Landscape/Classes/Landscape.h
+++ b/Engine/Source/Runtime/Landscape/Classes/Landscape.h
@@ -257,6 +257,7 @@ public:
 	LANDSCAPE_API void SetLayerName(int32 InLayerIndex, const FName& InName);
 	LANDSCAPE_API void SetLayerAlpha(int32 InLayerIndex, const float InAlpha, bool bInHeightmap);
 	LANDSCAPE_API void SetLayerVisibility(int32 InLayerIndex, bool bInVisible);
+	LANDSCAPE_API void SetLayerLocked(int32 InLayerIndex, bool bLocked);
 	LANDSCAPE_API struct FLandscapeLayer* GetLayer(int32 InLayerIndex);
 	LANDSCAPE_API const struct FLandscapeLayer* GetLayer(int32 InLayerIndex) const;
 	LANDSCAPE_API const struct FLandscapeLayer* GetLayer(const FGuid& InLayerGuid) const;
diff --git a/Engine/Source/Runtime/Landscape/Private/LandscapeEditLayers.cpp b/Engine/Source/Runtime/Landscape/Private/LandscapeEditLayers.cpp
--- a/Engine/Source/Runtime/Landscape/Private/LandscapeEditLayers.cpp
+++ b/Engine/Source/Runtime/Landscape/Private/LandscapeEditLayers.cpp
@@ -4404,6 +4404,7 @@ void ALandscape::SetLayerName(int32 InLayerIndex, const FName& InName)
 		return;
 	}
 
+	Modify();
 	LandscapeLayers[InLayerIndex].Name = InName;
 }
 
@@ -4420,6 +4421,7 @@ void ALandscape::SetLayerAlpha(int32 InLayerIndex, const float InAlpha, bool bIn
 		return;
 	}
 
+	Modify();
 	LayerAlpha = InAlpha;
 	RequestLayersContentUpdate(ELandscapeLayersContentUpdateFlag::All, true);
 }
@@ -4432,10 +4434,23 @@ void ALandscape::SetLayerVisibility(int32 InLayerIndex, bool bInVisible)
 		return;
 	}
 
+	Modify();
 	Layer->bVisible = bInVisible;
 	RequestLayersContentUpdate(ELandscapeLayersContentUpdateFlag::All, true);
 }
 
+void ALandscape::SetLayerLocked(int32 InLayerIndex, bool bLocked)
+{
+	FLandscapeLayer* Layer = GetLayer(InLayerIndex);
+	if (!Layer || Layer->bLocked == bLocked)
+	{
+		return;
+	}
+
+	Modify();
+	Layer->bLocked = bLocked;
+}
+
 FLandscapeLayer* ALandscape::GetLayer(int32 InLayerIndex)
 {
 	if (LandscapeLayers.IsValidIndex(InLayerIndex))
