- Landscape Layer System: Current Layer Selection now Transactional

diff --git a/Engine/Source/Editor/LandscapeEditor/Private/LandscapeEdMode.cpp b/Engine/Source/Editor/LandscapeEditor/Private/LandscapeEdMode.cpp
--- a/Engine/Source/Editor/LandscapeEditor/Private/LandscapeEdMode.cpp
+++ b/Engine/Source/Editor/LandscapeEditor/Private/LandscapeEdMode.cpp
@@ -2476,7 +2476,7 @@ void FEdModeLandscape::UpdateTargetList()
 			bool bFoundSelected = false;
 
 			// Add heightmap
-			LandscapeTargetList.Add(MakeShareable(new FLandscapeTargetListInfo(LOCTEXT("Heightmap", "Heightmap"), ELandscapeToolTargetType::Heightmap, CurrentToolTarget.LandscapeInfo.Get(), CurrentToolTarget.CurrentLayerIndex)));
+			LandscapeTargetList.Add(MakeShareable(new FLandscapeTargetListInfo(LOCTEXT("Heightmap", "Heightmap"), ELandscapeToolTargetType::Heightmap, CurrentToolTarget.LandscapeInfo.Get(), UISettings->CurrentLayerIndex)));
 
 			if (CurrentToolTarget.TargetType == ELandscapeToolTargetType::Heightmap)
 			{
@@ -2485,7 +2485,7 @@ void FEdModeLandscape::UpdateTargetList()
 
 			// Add visibility
 			FLandscapeInfoLayerSettings VisibilitySettings(ALandscapeProxy::VisibilityLayer, LandscapeProxy);
-			LandscapeTargetList.Add(MakeShareable(new FLandscapeTargetListInfo(LOCTEXT("Visibility", "Visibility"), ELandscapeToolTargetType::Visibility, VisibilitySettings, CurrentToolTarget.CurrentLayerIndex)));
+			LandscapeTargetList.Add(MakeShareable(new FLandscapeTargetListInfo(LOCTEXT("Visibility", "Visibility"), ELandscapeToolTargetType::Visibility, VisibilitySettings, UISettings->CurrentLayerIndex)));
 
 			if (CurrentToolTarget.TargetType == ELandscapeToolTargetType::Visibility)
 			{
@@ -2536,7 +2536,7 @@ void FEdModeLandscape::UpdateTargetList()
 				}
 
 				// Add the layer
-				LandscapeTargetList.Add(MakeShareable(new FLandscapeTargetListInfo(FText::FromName(LayerName), ELandscapeToolTargetType::Weightmap, LayerSettings, CurrentToolTarget.CurrentLayerIndex)));
+				LandscapeTargetList.Add(MakeShareable(new FLandscapeTargetListInfo(FText::FromName(LayerName), ELandscapeToolTargetType::Weightmap, LayerSettings, UISettings->CurrentLayerIndex)));
 			}
 
 			if (!bFoundSelected)
@@ -4160,14 +4160,15 @@ int32 FEdModeLandscape::GetLayerCount() const
 
 void FEdModeLandscape::SetCurrentLayer(int32 InLayerIndex)
 {
-	CurrentToolTarget.CurrentLayerIndex = InLayerIndex;
+	UISettings->Modify();
+	UISettings->CurrentLayerIndex = InLayerIndex;
 
 	RefreshDetailPanel();
 }
 
 int32 FEdModeLandscape::GetCurrentLayerIndex() const
 {
-	return CurrentToolTarget.CurrentLayerIndex;
+	return UISettings ? UISettings->CurrentLayerIndex : INDEX_NONE;
 }
 
 ALandscape* FEdModeLandscape::GetLandscape() const
diff --git a/Engine/Source/Editor/LandscapeEditor/Private/LandscapeEditorDetailCustomization_Layers.cpp b/Engine/Source/Editor/LandscapeEditor/Private/LandscapeEditorDetailCustomization_Layers.cpp
--- a/Engine/Source/Editor/LandscapeEditor/Private/LandscapeEditorDetailCustomization_Layers.cpp
+++ b/Engine/Source/Editor/LandscapeEditor/Private/LandscapeEditorDetailCustomization_Layers.cpp
@@ -561,6 +561,7 @@ void FLandscapeEditorCustomNodeBuilder_Layers::OnLayerSelectionChanged(int32 InL
 	FEdModeLandscape* LandscapeEdMode = GetEditorMode();
 	if (LandscapeEdMode)
 	{
+		FScopedTransaction Transaction(LOCTEXT("Landscape_Layers_SetCurrentLayer", "Set Current Layer"));
 		LandscapeEdMode->SetCurrentLayer(InLayerIndex);
 		LandscapeEdMode->UpdateTargetList();
 	}
diff --git a/Engine/Source/Editor/LandscapeEditor/Private/LandscapeEditorObject.cpp b/Engine/Source/Editor/LandscapeEditor/Private/LandscapeEditorObject.cpp
--- a/Engine/Source/Editor/LandscapeEditor/Private/LandscapeEditorObject.cpp
+++ b/Engine/Source/Editor/LandscapeEditor/Private/LandscapeEditorObject.cpp
@@ -100,6 +100,7 @@ ULandscapeEditorObject::ULandscapeEditorObject(const FObjectInitializer& ObjectI
 	, BrushComponentSize(1)
 	, TargetDisplayOrder(ELandscapeLayerDisplayMode::Default)
 	, ShowUnusedLayers(true)
+	, CurrentLayerIndex(INDEX_NONE)
 {
 	// Structure to hold one-time initialization
 	struct FConstructorStatics
diff --git a/Engine/Source/Editor/LandscapeEditor/Public/LandscapeEditorObject.h b/Engine/Source/Editor/LandscapeEditor/Public/LandscapeEditorObject.h
--- a/Engine/Source/Editor/LandscapeEditor/Public/LandscapeEditorObject.h
+++ b/Engine/Source/Editor/LandscapeEditor/Public/LandscapeEditorObject.h
@@ -609,6 +609,9 @@ public:
 	UPROPERTY(Category = "Target Layers", EditAnywhere)
 	bool ShowUnusedLayers;	
 
+	UPROPERTY(Transient)
+	int32 CurrentLayerIndex;
+
 #if WITH_EDITOR
 	//~ Begin UObject Interface
 	void PostEditChangeProperty(FPropertyChangedEvent& PropertyChangedEvent) override;
diff --git a/Engine/Source/Editor/LandscapeEditor/Public/LandscapeToolInterface.h b/Engine/Source/Editor/LandscapeEditor/Public/LandscapeToolInterface.h
--- a/Engine/Source/Editor/LandscapeEditor/Public/LandscapeToolInterface.h
+++ b/Engine/Source/Editor/LandscapeEditor/Public/LandscapeToolInterface.h
@@ -176,14 +176,12 @@ struct FLandscapeToolTarget
 	ELandscapeToolTargetType::Type TargetType;
 	TWeakObjectPtr<ULandscapeLayerInfoObject> LayerInfo;
 	FName LayerName;
-	int32 CurrentLayerIndex;
-
+	
 	FLandscapeToolTarget()
 		: LandscapeInfo()
 		, TargetType(ELandscapeToolTargetType::Heightmap)
 		, LayerInfo()
 		, LayerName(NAME_None)
-		, CurrentLayerIndex(INDEX_NONE)
 	{
 	}
 };
