Fixed static analysis of potential null ptr

diff --git a/Engine/Source/Editor/LandscapeEditor/Private/LandscapeEdMode.cpp b/Engine/Source/Editor/LandscapeEditor/Private/LandscapeEdMode.cpp
--- a/Engine/Source/Editor/LandscapeEditor/Private/LandscapeEdMode.cpp
+++ b/Engine/Source/Editor/LandscapeEditor/Private/LandscapeEdMode.cpp
@@ -4435,7 +4435,7 @@ bool FEdModeLandscape::CanEditLayer(FText* Reason /*=nullptr*/, FLandscapeLayer*
 		}
 		else if (CurrentTool)
 		{
-			int32 TargetLayerIndex = Landscape->LandscapeLayers.IndexOfByPredicate([TargetLayeyGuid = TargetLayer->Guid](const FLandscapeLayer& OtherLayer) { return OtherLayer.Guid == TargetLayeyGuid; });
+			int32 TargetLayerIndex = Landscape ? Landscape->LandscapeLayers.IndexOfByPredicate([TargetLayeyGuid = TargetLayer->Guid](const FLandscapeLayer& OtherLayer) { return OtherLayer.Guid == TargetLayeyGuid; }) : INDEX_NONE;
 
 			if ((CurrentTool != (FLandscapeTool*)SplinesTool) && Landscape && (TargetLayer == Landscape->GetLandscapeSplinesReservedLayer()))
 			{
