- Landscape Layer System: Fixes
Subtractive Flag Undo/Redo support
LayerInfo NoWeightBlend properly reflected in details panel
Missing fixup in ReplaceLayer for Subtractive Flag

#rnx

diff --git a/Engine/Source/Editor/LandscapeEditor/Private/LandscapeEditorDetailCustomization_TargetLayers.cpp b/Engine/Source/Editor/LandscapeEditor/Private/LandscapeEditorDetailCustomization_TargetLayers.cpp
--- a/Engine/Source/Editor/LandscapeEditor/Private/LandscapeEditorDetailCustomization_TargetLayers.cpp
+++ b/Engine/Source/Editor/LandscapeEditor/Private/LandscapeEditorDetailCustomization_TargetLayers.cpp
@@ -620,7 +620,7 @@ TSharedPtr<SWidget> FLandscapeEditorCustomNodeBuilder_TargetLayers::GenerateRow(
 						.HAlign(HAlign_Right)
 						[
 							SNew(STextBlock)
-							.Visibility((Target->LayerInfoObj.IsValid() && Target->LayerInfoObj->bNoWeightBlend) ? EVisibility::Visible : EVisibility::Collapsed)
+							.Visibility_Lambda([=] { return (Target->LayerInfoObj.IsValid() && Target->LayerInfoObj->bNoWeightBlend) ? EVisibility::Visible : EVisibility::Collapsed; })
 							.Font(IDetailLayoutBuilder::GetDetailFont())
 							.Text(LOCTEXT("NoWeightBlend", "No Weight-Blend"))
 							.ShadowOffset(FVector2D::UnitVector)
@@ -1451,6 +1451,7 @@ void FLandscapeEditorCustomNodeBuilder_TargetLayers::OnLayersSubstractiveBlendCh
 	FEdModeLandscape* LandscapeEdMode = GetEditorMode();
 	if (LandscapeEdMode)
 	{
+		FScopedTransaction Transaction(LOCTEXT("Undo_SubtractiveBlend", "Set Subtractive Blend Layer"));
 		LandscapeEdMode->SetCurrentLayerSubstractiveBlendStatus(NewCheckedState == ECheckBoxState::Checked, Target.Get().LayerInfoObj);
 	}
 }
diff --git a/Engine/Source/Runtime/Landscape/Classes/Landscape.h b/Engine/Source/Runtime/Landscape/Classes/Landscape.h
--- a/Engine/Source/Runtime/Landscape/Classes/Landscape.h
+++ b/Engine/Source/Runtime/Landscape/Classes/Landscape.h
@@ -262,7 +262,7 @@ public:
 	LANDSCAPE_API struct FLandscapeLayer* GetLayer(int32 InLayerIndex);
 	LANDSCAPE_API const struct FLandscapeLayer* GetLayer(int32 InLayerIndex) const;
 	LANDSCAPE_API const struct FLandscapeLayer* GetLayer(const FGuid& InLayerGuid) const;
-	LANDSCAPE_API void ForEachLayer(TFunctionRef<void(const struct FLandscapeLayer&)> Fn);
+	LANDSCAPE_API void ForEachLayer(TFunctionRef<void(struct FLandscapeLayer&)> Fn);
 	LANDSCAPE_API void ClearLayer(int32 InLayerIndex, bool bInUpdateCollision = true);
 	LANDSCAPE_API void ClearLayer(const FGuid& InLayerGuid, bool bInUpdateCollision = true);
 	LANDSCAPE_API void DeleteLayer(int32 InLayerIndex);
diff --git a/Engine/Source/Runtime/Landscape/Private/LandscapeEditInterface.cpp b/Engine/Source/Runtime/Landscape/Private/LandscapeEditInterface.cpp
--- a/Engine/Source/Runtime/Landscape/Private/LandscapeEditInterface.cpp
+++ b/Engine/Source/Runtime/Landscape/Private/LandscapeEditInterface.cpp
@@ -2266,8 +2266,15 @@ void FLandscapeEditDataInterface::ReplaceLayer(ULandscapeLayerInfoObject* FromLa
 	{
 		if (LandscapeInfo->LandscapeActor)
 		{
-			LandscapeInfo->LandscapeActor->ForEachLayer([&](const FLandscapeLayer& CurrentLayer)
+			LandscapeInfo->LandscapeActor->Modify();
+			LandscapeInfo->LandscapeActor->ForEachLayer([&](FLandscapeLayer& CurrentLayer)
 			{
+				bool OutValue;
+				if (CurrentLayer.WeightmapLayerAllocationBlend.RemoveAndCopyValue(FromLayerInfo, OutValue))
+				{
+					CurrentLayer.WeightmapLayerAllocationBlend.Add(ToLayerInfo, OutValue);
+				}
+
 				FScopedSetLandscapeEditingLayer Scope(LandscapeInfo->LandscapeActor.Get(), CurrentLayer.Guid);
 				DoReplace();
 			});
diff --git a/Engine/Source/Runtime/Landscape/Private/LandscapeEditLayers.cpp b/Engine/Source/Runtime/Landscape/Private/LandscapeEditLayers.cpp
--- a/Engine/Source/Runtime/Landscape/Private/LandscapeEditLayers.cpp
+++ b/Engine/Source/Runtime/Landscape/Private/LandscapeEditLayers.cpp
@@ -4504,7 +4504,7 @@ const FLandscapeLayer* ALandscape::GetLayer(const FGuid& InLayerGuid) const
 	return LandscapeLayers.FindByPredicate([InLayerGuid](const FLandscapeLayer& Other) { return Other.Guid == InLayerGuid; });
 }
 
-void ALandscape::ForEachLayer(TFunctionRef<void(const struct FLandscapeLayer&)> Fn)
+void ALandscape::ForEachLayer(TFunctionRef<void(struct FLandscapeLayer&)> Fn)
 {
 	for (FLandscapeLayer& Layer : LandscapeLayers)
 	{
@@ -4902,6 +4902,7 @@ void ALandscape::SetLayerSubstractiveBlendStatus(int32 InLayerIndex, bool InStat
 		return;
 	}
 
+	Modify();
 	bool* AllocationBlend = Layer->WeightmapLayerAllocationBlend.Find(InLayerInfoObj.Get());
 
 	if (AllocationBlend == nullptr)
