- Landscape Layer System: Prevent editing of LandscapeProxy if LandscapeActor isn't loaded

#rnx

diff --git a/Engine/Source/Editor/LandscapeEditor/Private/LandscapeEdMode.cpp b/Engine/Source/Editor/LandscapeEditor/Private/LandscapeEdMode.cpp
--- a/Engine/Source/Editor/LandscapeEditor/Private/LandscapeEdMode.cpp
+++ b/Engine/Source/Editor/LandscapeEditor/Private/LandscapeEdMode.cpp
@@ -2356,6 +2356,15 @@ int32 FEdModeLandscape::UpdateLandscapeList()
 				ALandscapeProxy* LandscapeProxy = LandscapeInfo->GetLandscapeProxy();
 				if (LandscapeProxy)
 				{
+					// In Layer System for now disable Landscape Editing if Actor is not loaded
+					if (GetMutableDefault<UEditorExperimentalSettings>()->bLandscapeLayerSystem)
+					{
+						if (!LandscapeProxy->GetLandscapeActor())
+						{
+							continue;
+						}
+					}
+
 					if (CurrentToolTarget.LandscapeInfo == LandscapeInfo)
 					{
 						CurrentIndex = Index;
