Support Delete Component Tool in Landscape Layer System

diff --git a/Engine/Source/Editor/LandscapeEditor/Private/LandscapeEdMode.cpp b/Engine/Source/Editor/LandscapeEditor/Private/LandscapeEdMode.cpp
--- a/Engine/Source/Editor/LandscapeEditor/Private/LandscapeEdMode.cpp
+++ b/Engine/Source/Editor/LandscapeEditor/Private/LandscapeEdMode.cpp
@@ -102,15 +102,20 @@ ENGINE_API extern bool GDisableAutomaticTextureMaterialUpdateDependencies;
 void ALandscape::SplitHeightmap(ULandscapeComponent* Comp, bool bMoveToCurrentLevel, FMaterialUpdateContext* InOutUpdateContext, TArray<FComponentRecreateRenderStateContext>* InOutRecreateRenderStateContext, bool InReregisterComponent)
 {
 	ULandscapeInfo* Info = Comp->GetLandscapeInfo();
-	ALandscape* Landscape = Info->LandscapeActor.Get();
+	
+	// Make sure the heightmap UVs are powers of two.
 	int32 ComponentSizeVerts = Comp->NumSubsections * (Comp->SubsectionSizeQuads + 1);
-	// make sure the heightmap UVs are powers of two.
 	int32 HeightmapSizeU = (1 << FMath::CeilLogTwo(ComponentSizeVerts));
 	int32 HeightmapSizeV = (1 << FMath::CeilLogTwo(ComponentSizeVerts));
 
-	UTexture2D* HeightmapTexture = NULL;
-	TArray<FColor*> HeightmapTextureMipData;
-	// Scope for FLandscapeEditDataInterface
+	ALandscapeProxy* Proxy = Comp->GetLandscapeProxy();
+	Proxy->Modify();
+
+	UObject* TextureOuter = bMoveToCurrentLevel ? Comp->GetWorld()->GetCurrentLevel()->GetOutermost() : nullptr;
+	
+	UTexture2D* OldHeightmapTexture = Comp->GetHeightmap(false);
+	UTexture2D* NewHeightmapTexture = NULL;
+	
 	{
 		// Read old data and split
 		FLandscapeEditDataInterface LandscapeEdit(Info);
@@ -121,141 +126,14 @@ void ALandscape::SplitHeightmap(ULandscapeComponent* Comp, bool bMoveToCurrentLe
 		NormalData.AddZeroed((1 + Comp->ComponentSizeQuads)*(1 + Comp->ComponentSizeQuads)*sizeof(uint16));
 		LandscapeEdit.GetHeightDataFast(Comp->GetSectionBase().X, Comp->GetSectionBase().Y, Comp->GetSectionBase().X + Comp->ComponentSizeQuads, Comp->GetSectionBase().Y + Comp->ComponentSizeQuads, (uint16*)HeightData.GetData(), 0, (uint16*)NormalData.GetData());
 
-		// Construct the heightmap textures
-		UObject* TextureOuter = bMoveToCurrentLevel ? Comp->GetWorld()->GetCurrentLevel()->GetOutermost() : nullptr;
-		HeightmapTexture = Comp->GetLandscapeProxy()->CreateLandscapeTexture(HeightmapSizeU, HeightmapSizeV, TEXTUREGROUP_Terrain_Heightmap, TSF_BGRA8, TextureOuter);
-
-		int32 MipSubsectionSizeQuads = Comp->SubsectionSizeQuads;
-		int32 MipSizeU = HeightmapSizeU;
-		int32 MipSizeV = HeightmapSizeV;
-		while (MipSizeU > 1 && MipSizeV > 1 && MipSubsectionSizeQuads >= 1)
-		{
-			FColor* HeightmapTextureData = (FColor*)HeightmapTexture->Source.LockMip(HeightmapTextureMipData.Num());
-			FMemory::Memzero(HeightmapTextureData, MipSizeU*MipSizeV*sizeof(FColor));
-			HeightmapTextureMipData.Add(HeightmapTextureData);
-
-			MipSizeU >>= 1;
-			MipSizeV >>= 1;
-
-			MipSubsectionSizeQuads = ((MipSubsectionSizeQuads + 1) >> 1) - 1;
-		}
-
+		// Create the new heightmap texture
+		NewHeightmapTexture = Proxy->CreateLandscapeTexture(HeightmapSizeU, HeightmapSizeV, TEXTUREGROUP_Terrain_Heightmap, TSF_BGRA8, TextureOuter);
+		ULandscapeComponent::CreateEmptyTextureMips(NewHeightmapTexture, true);
 		Comp->HeightmapScaleBias = FVector4(1.0f / (float)HeightmapSizeU, 1.0f / (float)HeightmapSizeV, 0.0f, 0.0f);
-
-#if WITH_EDITORONLY_DATA
-		// TODO: There is a bug here with updating the mip regions
-		/*if (GetMutableDefault<UEditorExperimentalSettings>()->bLandscapeLayerSystem)
-		{
-			// Add component to new heightmap
-			FRenderDataPerHeightmap* RenderData = Landscape->RenderDataPerHeightmap.Find(HeightmapTexture);
-
-			if (RenderData == nullptr)
-			{
-				FRenderDataPerHeightmap NewData;
-				NewData.Components.Add(Comp);
-				NewData.TopLeftSectionBase = Comp->GetSectionBase();
-				NewData.HeightmapKey = HeightmapTexture;
-				NewData.HeightmapsCPUReadBack = new FLandscapeLayersTexture2DCPUReadBackResource(HeightmapTexture->Source.GetSizeX(), HeightmapTexture->Source.GetSizeY(), Comp->GetHeightmap(false)->GetPixelFormat(), HeightmapTexture->Source.GetNumMips());
-				BeginInitResource(NewData.HeightmapsCPUReadBack);
-
-				Landscape->RenderDataPerHeightmap.Add(HeightmapTexture, NewData);
-			}
-			else
-			{
-				for (ULandscapeComponent* Component : RenderData->Components)
-				{
-					RenderData->TopLeftSectionBase.X = FMath::Min(RenderData->TopLeftSectionBase.X, Component->GetSectionBase().X);
-					RenderData->TopLeftSectionBase.Y = FMath::Min(RenderData->TopLeftSectionBase.Y, Component->GetSectionBase().Y);
-				}
-
-				RenderData->Components.Add(Comp);
-			}
-
-			// Remove component from old heightmap
-			FRenderDataPerHeightmap* OldRenderData = Landscape->RenderDataPerHeightmap.Find(Comp->GetHeightmap(false));
-			OldRenderData->Components.Remove(Comp);
-
-			bool RemoveLayerHeightmap = false;
-
-			// No component use this heightmap anymore, so clean up
-			if (OldRenderData->Components.Num() == 0)
-			{
-				ReleaseResourceAndFlush(OldRenderData->HeightmapsCPUReadBack);
-				delete OldRenderData->HeightmapsCPUReadBack;
-				OldRenderData->HeightmapsCPUReadBack = nullptr;
-				Landscape->RenderDataPerHeightmap.Remove(Comp->GetHeightmap(false));
-				RemoveLayerHeightmap = true;
-			}
-
-			// Move layer content to new layer heightmap
-			for (FLandscapeLayer& Layer : Landscape->LandscapeLayers)
-			{
-				UTexture2D** OldLayerHeightmap = Layer.HeightmapTarget.Heightmaps.Find(Comp->GetHeightmap(false));
-
-				if (OldLayerHeightmap != nullptr)
-				{
-					// If we remove the main heightmap, remove the layers one, as they exist on a 1 to 1
-					if (RemoveLayerHeightmap)
-					{
-						Layer.HeightmapTarget.Heightmaps.Remove(*OldLayerHeightmap);
-					}
-
-					// Read old data and split
-					TArray<uint8> LayerHeightData;
-					LayerHeightData.AddZeroed((1 + Comp->ComponentSizeQuads)*(1 + Comp->ComponentSizeQuads) * sizeof(uint16));
-					// Because of edge problem, normal would be just copy from old component data
-					TArray<uint8> LayerNormalData;
-					LayerNormalData.AddZeroed((1 + Comp->ComponentSizeQuads)*(1 + Comp->ComponentSizeQuads) * sizeof(uint16));
-					LandscapeEdit.GetHeightDataFast(Comp->GetSectionBase().X, Comp->GetSectionBase().Y, Comp->GetSectionBase().X + Comp->ComponentSizeQuads, Comp->GetSectionBase().Y + Comp->ComponentSizeQuads, (uint16*)LayerHeightData.GetData(), 0, (uint16*)LayerNormalData.GetData(), *OldLayerHeightmap);
-
-					UTexture2D** NewLayerHeightmap = Layer.HeightmapTarget.Heightmaps.Find(HeightmapTexture);
-
-					if (NewLayerHeightmap == nullptr)
-					{
-						UTexture2D* LayerHeightmapTexture = Comp->GetLandscapeProxy()->CreateLandscapeTexture(HeightmapSizeU, HeightmapSizeV, TEXTUREGROUP_Terrain_Heightmap, TSF_BGRA8, bMoveToCurrentLevel ? Comp->GetWorld()->GetCurrentLevel()->GetOutermost() : nullptr);
-
-						MipSubsectionSizeQuads = Comp->SubsectionSizeQuads;
-						MipSizeU = HeightmapSizeU;
-						MipSizeV = HeightmapSizeV;
-						int32 MipIndex = 0;
-
-						while (MipSizeU > 1 && MipSizeV > 1 && MipSubsectionSizeQuads >= 1)
-						{
-							FColor* HeightmapTextureData = (FColor*)LayerHeightmapTexture->Source.LockMip(MipIndex);
-							FMemory::Memzero(HeightmapTextureData, MipSizeU*MipSizeV * sizeof(FColor));
-							LayerHeightmapTexture->Source.UnlockMip(MipIndex);
-
-							MipSizeU >>= 1;
-							MipSizeV >>= 1;
-
-							MipSubsectionSizeQuads = ((MipSubsectionSizeQuads + 1) >> 1) - 1;
-							++MipIndex;
-						}
-
-						Layer.HeightmapTarget.Heightmaps.Add(HeightmapTexture, LayerHeightmapTexture);
-
-						LandscapeEdit.SetHeightData(Comp->GetSectionBase().X, Comp->GetSectionBase().Y, Comp->GetSectionBase().X + Comp->ComponentSizeQuads, Comp->GetSectionBase().Y + Comp->ComponentSizeQuads, (uint16*)LayerHeightData.GetData(), 0, false, (uint16*)LayerNormalData.GetData(),
-												    false, LayerHeightmapTexture, nullptr, true, true, true);
-					}
-					else
-					{
-						LandscapeEdit.SetHeightData(Comp->GetSectionBase().X, Comp->GetSectionBase().Y, Comp->GetSectionBase().X + Comp->ComponentSizeQuads, Comp->GetSectionBase().Y + Comp->ComponentSizeQuads, (uint16*)LayerHeightData.GetData(), 0, false, (uint16*)LayerNormalData.GetData(),
-												   false, *NewLayerHeightmap, nullptr, true, true, true);
-					}
-				}
-			}
-		}
-		*/
-#endif
-
-		Comp->SetHeightmap(HeightmapTexture);
-
-		for (int32 i = 0; i < HeightmapTextureMipData.Num(); i++)
-		{
-			HeightmapTexture->Source.UnlockMip(i);
-		}		
+		Comp->SetHeightmap(NewHeightmapTexture);
 
 		LandscapeEdit.SetHeightData(Comp->GetSectionBase().X, Comp->GetSectionBase().Y, Comp->GetSectionBase().X + Comp->ComponentSizeQuads, Comp->GetSectionBase().Y + Comp->ComponentSizeQuads, (uint16*)HeightData.GetData(), 0, false, (uint16*)NormalData.GetData());
+		NewHeightmapTexture->PostEditChange();
 	}
 
 	// End material update
@@ -268,12 +146,10 @@ void ALandscape::SplitHeightmap(ULandscapeComponent* Comp, bool bMoveToCurrentLe
 		Comp->UpdateMaterialInstances();
 	}
 
-	// End of LandscapeEdit interface
-
 	// We disable automatic material update context, to manage it manually if we have a custom update context specified
-	GDisableAutomaticTextureMaterialUpdateDependencies = InOutUpdateContext != nullptr;
+	GDisableAutomaticTextureMaterialUpdateDependencies = (InOutUpdateContext != nullptr);
 
-	HeightmapTexture->PostEditChange();	
+	NewHeightmapTexture->PostEditChange();
 
 	if (InOutUpdateContext != nullptr)
 	{
@@ -292,7 +168,7 @@ void ALandscape::SplitHeightmap(ULandscapeComponent* Comp, bool bMoveToCurrentLe
 
 		for (UMaterialInterface* MaterialInterface : LandscapeMaterials)
 		{
-			if (DoesMaterialUseTexture(MaterialInterface, HeightmapTexture))
+			if (DoesMaterialUseTexture(MaterialInterface, NewHeightmapTexture))
 			{
 				UMaterial* Material = MaterialInterface->GetMaterial();
 				bool MaterialAlreadyCompute = false;
@@ -300,7 +176,7 @@ void ALandscape::SplitHeightmap(ULandscapeComponent* Comp, bool bMoveToCurrentLe
 
 				if (!MaterialAlreadyCompute)
 				{
-					if (Material->IsTextureForceRecompileCacheRessource(HeightmapTexture))
+					if (Material->IsTextureForceRecompileCacheRessource(NewHeightmapTexture))
 					{
 						InOutUpdateContext->AddMaterial(Material);
 						Material->UpdateMaterialShaderCacheAndTextureReferences();
@@ -312,6 +188,66 @@ void ALandscape::SplitHeightmap(ULandscapeComponent* Comp, bool bMoveToCurrentLe
 	
 	GDisableAutomaticTextureMaterialUpdateDependencies = false;
 
+#if WITH_EDITORONLY_DATA
+	if (GetMutableDefault<UEditorExperimentalSettings>()->bLandscapeLayerSystem)
+	{
+		FLandscapeEditDataInterface LandscapeEdit(Info);
+		FLandscapeLayersTexture2DCPUReadBackResource* NewCPUReadBackResource = new FLandscapeLayersTexture2DCPUReadBackResource(NewHeightmapTexture->Source.GetSizeX(), NewHeightmapTexture->Source.GetSizeY(), NewHeightmapTexture->GetPixelFormat(), NewHeightmapTexture->Source.GetNumMips());
+		BeginInitResource(NewCPUReadBackResource);
+		Proxy->HeightmapsCPUReadBack.Add(NewHeightmapTexture, NewCPUReadBackResource);
+
+		// Free OldHeightmapTexture's CPUReadBackResource if not used by any component
+		bool FreeCPURedBack = true;
+		for (ULandscapeComponent* Component : Proxy->LandscapeComponents)
+		{
+			if (Component != Comp && Component->GetHeightmap(false) == OldHeightmapTexture)
+			{
+				FreeCPURedBack = false;
+				break;
+			}
+		}
+		if (FreeCPURedBack)
+		{
+			FLandscapeLayersTexture2DCPUReadBackResource** OldCPUReadBackResource = Proxy->HeightmapsCPUReadBack.Find(OldHeightmapTexture);
+			if (OldCPUReadBackResource)
+			{
+				if (FLandscapeLayersTexture2DCPUReadBackResource* ResourceToDelete = *OldCPUReadBackResource)
+				{
+					ReleaseResourceAndFlush(ResourceToDelete);
+					delete ResourceToDelete;
+					Proxy->HeightmapsCPUReadBack.Remove(OldHeightmapTexture);
+				}
+			}
+		}
+
+		// Move layer content to new layer heightmap
+		ALandscape* Landscape = Info->LandscapeActor.Get();
+		Comp->ForEachLayer([&](const FGuid& LayerGuid, FLandscapeLayerComponentData& LayerData)
+		{
+			UTexture2D* OldLayerHeightmap = LayerData.HeightmapData.Texture;
+			if (OldLayerHeightmap != nullptr)
+			{
+				FScopedSetLandscapeEditingLayer Scope(Landscape, LayerGuid);
+				// Read old data and split
+				TArray<uint8> LayerHeightData;
+				LayerHeightData.AddZeroed((1 + Comp->ComponentSizeQuads)*(1 + Comp->ComponentSizeQuads) * sizeof(uint16));
+				// Because of edge problem, normal would be just copy from old component data
+				TArray<uint8> LayerNormalData;
+				LayerNormalData.AddZeroed((1 + Comp->ComponentSizeQuads)*(1 + Comp->ComponentSizeQuads) * sizeof(uint16));
+				LandscapeEdit.GetHeightDataFast(Comp->GetSectionBase().X, Comp->GetSectionBase().Y, Comp->GetSectionBase().X + Comp->ComponentSizeQuads, Comp->GetSectionBase().Y + Comp->ComponentSizeQuads, (uint16*)LayerHeightData.GetData(), 0, (uint16*)LayerNormalData.GetData());
+				{
+					UTexture2D* LayerHeightmapTexture = Proxy->CreateLandscapeTexture(HeightmapSizeU, HeightmapSizeV, TEXTUREGROUP_Terrain_Heightmap, TSF_BGRA8, TextureOuter);
+					ULandscapeComponent::CreateEmptyTextureMips(LayerHeightmapTexture, true);
+					LayerHeightmapTexture->PostEditChange();
+					// Set Layer heightmap texture
+					LayerData.HeightmapData.Texture = LayerHeightmapTexture;
+					LandscapeEdit.SetHeightData(Comp->GetSectionBase().X, Comp->GetSectionBase().Y, Comp->GetSectionBase().X + Comp->ComponentSizeQuads, Comp->GetSectionBase().Y + Comp->ComponentSizeQuads, (uint16*)LayerHeightData.GetData(), 0, false, (uint16*)LayerNormalData.GetData());
+				}
+			}
+		});
+	}
+#endif
+
 		// Reregister
 	if (InReregisterComponent)
 	{
diff --git a/Engine/Source/Runtime/Landscape/Classes/LandscapeComponent.h b/Engine/Source/Runtime/Landscape/Classes/LandscapeComponent.h
--- a/Engine/Source/Runtime/Landscape/Classes/LandscapeComponent.h
+++ b/Engine/Source/Runtime/Landscape/Classes/LandscapeComponent.h
@@ -577,6 +577,7 @@ public:
 	LANDSCAPE_API FLandscapeLayerComponentData* GetLayerData(const FGuid& InLayerGuid);
 	LANDSCAPE_API void AddLayerData(const FGuid& InLayerGuid, const FLandscapeLayerComponentData& InData);
 	LANDSCAPE_API void RemoveLayerData(const FGuid& InLayerGuid);
+	LANDSCAPE_API void ForEachLayer(TFunctionRef<void(const FGuid&, struct FLandscapeLayerComponentData&)> Fn);
 
 	FLandscapeLayerComponentData* GetEditingLayer();
 	const FLandscapeLayerComponentData* GetEditingLayer() const;
diff --git a/Engine/Source/Runtime/Landscape/Private/Landscape.cpp b/Engine/Source/Runtime/Landscape/Private/Landscape.cpp
--- a/Engine/Source/Runtime/Landscape/Private/Landscape.cpp
+++ b/Engine/Source/Runtime/Landscape/Private/Landscape.cpp
@@ -1427,6 +1427,14 @@ FLandscapeLayerComponentData* ULandscapeComponent::GetLayerData(const FGuid& InL
 	return LayersData.Find(InLayerGuid);
 }
 
+void ULandscapeComponent::ForEachLayer(TFunctionRef<void(const FGuid&, struct FLandscapeLayerComponentData&)> Fn)
+{
+	for (auto& Pair : LayersData)
+	{
+		Fn(Pair.Key, Pair.Value);
+	}
+}
+
 void ULandscapeComponent::AddLayerData(const FGuid& InLayerGuid, const FLandscapeLayerComponentData& InData)
 {
 	FLandscapeLayerComponentData& Data = LayersData.FindOrAdd(InLayerGuid);
