Fixed Landscape Delete Tool in Layer System Mode. When splitting heightmap, update of final layer texture was in fact updating current editing layer.

diff --git a/Engine/Source/Editor/LandscapeEditor/Private/LandscapeEdMode.cpp b/Engine/Source/Editor/LandscapeEditor/Private/LandscapeEdMode.cpp
--- a/Engine/Source/Editor/LandscapeEditor/Private/LandscapeEdMode.cpp
+++ b/Engine/Source/Editor/LandscapeEditor/Private/LandscapeEdMode.cpp
@@ -132,8 +132,8 @@ void ALandscape::SplitHeightmap(ULandscapeComponent* Comp, bool bMoveToCurrentLe
 		Comp->HeightmapScaleBias = FVector4(1.0f / (float)HeightmapSizeU, 1.0f / (float)HeightmapSizeV, 0.0f, 0.0f);
 		Comp->SetHeightmap(NewHeightmapTexture);
 
+		check(Comp->GetHeightmap(false) == Comp->GetHeightmap(true));
 		LandscapeEdit.SetHeightData(Comp->GetSectionBase().X, Comp->GetSectionBase().Y, Comp->GetSectionBase().X + Comp->ComponentSizeQuads, Comp->GetSectionBase().Y + Comp->ComponentSizeQuads, (uint16*)HeightData.GetData(), 0, false, (uint16*)NormalData.GetData());
-		NewHeightmapTexture->PostEditChange();
 	}
 
 	// End material update
diff --git a/Engine/Source/Editor/LandscapeEditor/Private/LandscapeEdModeComponentTools.cpp b/Engine/Source/Editor/LandscapeEditor/Private/LandscapeEdModeComponentTools.cpp
--- a/Engine/Source/Editor/LandscapeEditor/Private/LandscapeEdModeComponentTools.cpp
+++ b/Engine/Source/Editor/LandscapeEditor/Private/LandscapeEdModeComponentTools.cpp
@@ -1141,6 +1141,8 @@ public:
 	{
 	}
 
+	virtual bool ShouldUpdateEditingLayer() const { return false; }
+
 	virtual const TCHAR* GetToolName() override { return TEXT("DeleteComponent"); }
 	virtual FText GetDisplayName() override { return NSLOCTEXT("UnrealEd", "LandscapeMode_DeleteComponent", "Delete Landscape Components"); };
 
diff --git a/Engine/Source/Editor/LandscapeEditor/Private/LandscapeEdModeTools.h b/Engine/Source/Editor/LandscapeEditor/Private/LandscapeEdModeTools.h
--- a/Engine/Source/Editor/LandscapeEditor/Private/LandscapeEdModeTools.h
+++ b/Engine/Source/Editor/LandscapeEditor/Private/LandscapeEdModeTools.h
@@ -1121,6 +1121,8 @@ public:
 	{
 	}
 
+	virtual bool ShouldUpdateEditingLayer() const { return GetMutableDefault<UEditorExperimentalSettings>()->bLandscapeLayerSystem; }
+
 	virtual ELandscapeLayersContentUpdateFlag GetBeginToolContentUpdateFlag() const
 	{
 		bool bUpdateHeightmap = this->EdMode->CurrentToolTarget.TargetType == ELandscapeToolTargetType::Type::Heightmap; 
@@ -1140,7 +1142,7 @@ public:
 
 	virtual bool BeginTool(FEditorViewportClient* ViewportClient, const FLandscapeToolTarget& InTarget, const FVector& InHitLocation) override
 	{
-		if (GetMutableDefault<UEditorExperimentalSettings>()->bLandscapeLayerSystem)
+		if (ShouldUpdateEditingLayer())
 		{
 			ALandscape* Landscape = this->EdMode->GetLandscape();
 			if (Landscape)
@@ -1194,7 +1196,7 @@ public:
 			// Prevent landscape from baking textures while tool stroke is active
 			EdMode->CurrentToolTarget.LandscapeInfo->PostponeTextureBaking();
 
-			if (GetMutableDefault<UEditorExperimentalSettings>()->bLandscapeLayerSystem)
+			if (ShouldUpdateEditingLayer())
 			{
 				ALandscape* Landscape = this->EdMode->CurrentToolTarget.LandscapeInfo->LandscapeActor.Get();
 				if (Landscape != nullptr)
@@ -1218,7 +1220,7 @@ public:
 		EdMode->UpdateLayerUsageInformation(&EdMode->CurrentToolTarget.LayerInfo);
 		bExternalModifierPressed = false;
 
-		if (GetMutableDefault<UEditorExperimentalSettings>()->bLandscapeLayerSystem)
+		if (ShouldUpdateEditingLayer())
 		{
 			ALandscape* Landscape = this->EdMode->GetLandscape();
 			if (Landscape)
