- Landscape Layer System: Add Experimental Header in Landscape Editor Mode to distinguish if Layers are enabled or not

#rnx

diff --git a/Engine/Source/Editor/LandscapeEditor/Private/SLandscapeEditor.cpp b/Engine/Source/Editor/LandscapeEditor/Private/SLandscapeEditor.cpp
--- a/Engine/Source/Editor/LandscapeEditor/Private/SLandscapeEditor.cpp
+++ b/Engine/Source/Editor/LandscapeEditor/Private/SLandscapeEditor.cpp
@@ -4,6 +4,7 @@
 #include "Framework/MultiBox/MultiBoxDefs.h"
 #include "Framework/MultiBox/MultiBoxBuilder.h"
 #include "Widgets/Layout/SBorder.h"
+#include "Widgets/Input/SMultiLineEditableTextBox.h"
 #include "Modules/ModuleManager.h"
 #include "SlateOptMacros.h"
 #include "Widgets/Layout/SBox.h"
@@ -16,6 +17,7 @@
 #include "IDetailsView.h"
 #include "PropertyEditorModule.h"
 #include "IIntroTutorials.h"
+#include "Settings/EditorExperimentalSettings.h"
 
 #define LOCTEXT_NAMESPACE "LandscapeEditor"
 
@@ -337,7 +339,18 @@ void SLandscapeEditor::Construct(const FArguments& InArgs, TSharedRef<FLandscape
 		[
 			SNew(SVerticalBox)
 			.IsEnabled(this, &SLandscapeEditor::GetLandscapeEditorIsEnabled)
-
+			+ SVerticalBox::Slot()
+			.AutoHeight()
+			.Padding(0, 0, 0, 5)
+			[
+				SNew(SMultiLineEditableTextBox)
+				.Visibility_Lambda([]() { return GetMutableDefault<UEditorExperimentalSettings>()->bLandscapeLayerSystem? EVisibility::Visible : EVisibility::Collapsed; })
+				.IsReadOnly(true)
+				.Font(FCoreStyle::GetDefaultFontStyle("Bold", 10))
+				.Justification(ETextJustify::Center)
+				.BackgroundColor(FCoreStyle::Get().GetColor("InfoReporting.BackgroundColor"))
+				.Text(FText::FromString("Experimental Layer System"))
+			]
 			+ SVerticalBox::Slot()
 			.AutoHeight()
 			.Padding(4, 0, 4, 5)
