- Fixed Undo of a delete landscape component in Landscape Layer System mode
- Fixed SplitHeightmap in Landscape Layer System mode not properly reading height data from old texture.

diff --git a/Engine/Source/Editor/LandscapeEditor/Private/LandscapeEdMode.cpp b/Engine/Source/Editor/LandscapeEditor/Private/LandscapeEdMode.cpp
--- a/Engine/Source/Editor/LandscapeEditor/Private/LandscapeEdMode.cpp
+++ b/Engine/Source/Editor/LandscapeEditor/Private/LandscapeEdMode.cpp
@@ -115,7 +115,9 @@ void ALandscape::SplitHeightmap(ULandscapeComponent* Comp, bool bMoveToCurrentLe
 	
 	UTexture2D* OldHeightmapTexture = Comp->GetHeightmap(false);
 	UTexture2D* NewHeightmapTexture = NULL;
-	
+	FVector4 OldHeightmapScaleBias = Comp->HeightmapScaleBias;
+	FVector4 NewHeightmapScaleBias = FVector4(1.0f / (float)HeightmapSizeU, 1.0f / (float)HeightmapSizeV, 0.0f, 0.0f);
+
 	{
 		// Read old data and split
 		FLandscapeEditDataInterface LandscapeEdit(Info);
@@ -129,7 +131,7 @@ void ALandscape::SplitHeightmap(ULandscapeComponent* Comp, bool bMoveToCurrentLe
 		// Create the new heightmap texture
 		NewHeightmapTexture = Proxy->CreateLandscapeTexture(HeightmapSizeU, HeightmapSizeV, TEXTUREGROUP_Terrain_Heightmap, TSF_BGRA8, TextureOuter);
 		ULandscapeComponent::CreateEmptyTextureMips(NewHeightmapTexture, true);
-		Comp->HeightmapScaleBias = FVector4(1.0f / (float)HeightmapSizeU, 1.0f / (float)HeightmapSizeV, 0.0f, 0.0f);
+		Comp->HeightmapScaleBias = NewHeightmapScaleBias;
 		Comp->SetHeightmap(NewHeightmapTexture);
 
 		check(Comp->GetHeightmap(false) == Comp->GetHeightmap(true));
@@ -191,22 +193,21 @@ void ALandscape::SplitHeightmap(ULandscapeComponent* Comp, bool bMoveToCurrentLe
 #if WITH_EDITORONLY_DATA
 	if (GetMutableDefault<UEditorExperimentalSettings>()->bLandscapeLayerSystem)
 	{
-		FLandscapeEditDataInterface LandscapeEdit(Info);
 		FLandscapeLayersTexture2DCPUReadBackResource* NewCPUReadBackResource = new FLandscapeLayersTexture2DCPUReadBackResource(NewHeightmapTexture->Source.GetSizeX(), NewHeightmapTexture->Source.GetSizeY(), NewHeightmapTexture->GetPixelFormat(), NewHeightmapTexture->Source.GetNumMips());
 		BeginInitResource(NewCPUReadBackResource);
 		Proxy->HeightmapsCPUReadBack.Add(NewHeightmapTexture, NewCPUReadBackResource);
 
 		// Free OldHeightmapTexture's CPUReadBackResource if not used by any component
-		bool FreeCPURedBack = true;
+		bool FreeCPUReadBack = true;
 		for (ULandscapeComponent* Component : Proxy->LandscapeComponents)
 		{
 			if (Component != Comp && Component->GetHeightmap(false) == OldHeightmapTexture)
 			{
-				FreeCPURedBack = false;
+				FreeCPUReadBack = false;
 				break;
 			}
 		}
-		if (FreeCPURedBack)
+		if (FreeCPUReadBack)
 		{
 			FLandscapeLayersTexture2DCPUReadBackResource** OldCPUReadBackResource = Proxy->HeightmapsCPUReadBack.Find(OldHeightmapTexture);
 			if (OldCPUReadBackResource)
@@ -221,6 +222,7 @@ void ALandscape::SplitHeightmap(ULandscapeComponent* Comp, bool bMoveToCurrentLe
 		}
 
 		// Move layer content to new layer heightmap
+		FLandscapeEditDataInterface LandscapeEdit(Info);
 		ALandscape* Landscape = Info->LandscapeActor.Get();
 		Comp->ForEachLayer([&](const FGuid& LayerGuid, FLandscapeLayerComponentData& LayerData)
 		{
@@ -234,7 +236,12 @@ void ALandscape::SplitHeightmap(ULandscapeComponent* Comp, bool bMoveToCurrentLe
 				// Because of edge problem, normal would be just copy from old component data
 				TArray<uint8> LayerNormalData;
 				LayerNormalData.AddZeroed((1 + Comp->ComponentSizeQuads)*(1 + Comp->ComponentSizeQuads) * sizeof(uint16));
+
+				// Read using old heightmap scale/bias
+				Comp->HeightmapScaleBias = OldHeightmapScaleBias;
 				LandscapeEdit.GetHeightDataFast(Comp->GetSectionBase().X, Comp->GetSectionBase().Y, Comp->GetSectionBase().X + Comp->ComponentSizeQuads, Comp->GetSectionBase().Y + Comp->ComponentSizeQuads, (uint16*)LayerHeightData.GetData(), 0, (uint16*)LayerNormalData.GetData());
+				// Restore new heightmap scale/bias
+				Comp->HeightmapScaleBias = NewHeightmapScaleBias;
 				{
 					UTexture2D* LayerHeightmapTexture = Proxy->CreateLandscapeTexture(HeightmapSizeU, HeightmapSizeV, TEXTUREGROUP_Terrain_Heightmap, TSF_BGRA8, TextureOuter);
 					ULandscapeComponent::CreateEmptyTextureMips(LayerHeightmapTexture, true);
@@ -245,10 +252,13 @@ void ALandscape::SplitHeightmap(ULandscapeComponent* Comp, bool bMoveToCurrentLe
 				}
 			}
 		});
+
+		Comp->GetGlobalLayersData().TopLeftSectionBase  = Comp->GetSectionBase();
+		Landscape->RequestLayersContentUpdate(ELandscapeLayersContentUpdateFlag::All, true);
 	}
 #endif
 
-		// Reregister
+	// Reregister
 	if (InReregisterComponent)
 	{
 		FComponentReregisterContext ReregisterContext(Comp);
diff --git a/Engine/Source/Editor/LandscapeEditor/Private/LandscapeEdModeComponentTools.cpp b/Engine/Source/Editor/LandscapeEditor/Private/LandscapeEdModeComponentTools.cpp
--- a/Engine/Source/Editor/LandscapeEditor/Private/LandscapeEdModeComponentTools.cpp
+++ b/Engine/Source/Editor/LandscapeEditor/Private/LandscapeEdModeComponentTools.cpp
@@ -1141,7 +1141,7 @@ public:
 	{
 	}
 
-	virtual bool ShouldUpdateEditingLayer() const { return false; }
+	virtual bool ShouldUpdateEditingLayer() const override { return false; }
 
 	virtual const TCHAR* GetToolName() override { return TEXT("DeleteComponent"); }
 	virtual FText GetDisplayName() override { return NSLOCTEXT("UnrealEd", "LandscapeMode_DeleteComponent", "Delete Landscape Components"); };
diff --git a/Engine/Source/Runtime/Landscape/Private/LandscapeEditLayers.cpp b/Engine/Source/Runtime/Landscape/Private/LandscapeEditLayers.cpp
--- a/Engine/Source/Runtime/Landscape/Private/LandscapeEditLayers.cpp
+++ b/Engine/Source/Runtime/Landscape/Private/LandscapeEditLayers.cpp
@@ -2469,6 +2469,25 @@ void ALandscape::RegenerateLayersHeightmaps()
 		AllLandscapeComponents.Append(Proxy->LandscapeComponents);
 	});
 
+	// Handle missing Heightmap CPUReadback (Undo of a delete landscape component can trigger that case)
+	if (LayersContentUpdateFlags & ELandscapeLayersContentUpdateFlag::Heightmap_All)
+	{
+		Info->ForAllLandscapeProxies([&](ALandscapeProxy* Proxy)
+		{
+			for (ULandscapeComponent* Component : Proxy->LandscapeComponents)
+			{
+				UTexture2D* ComponentHeightmap = Component->GetHeightmap();
+				FLandscapeLayersTexture2DCPUReadBackResource** CPUReadback = Proxy->HeightmapsCPUReadBack.Find(ComponentHeightmap);
+				if (CPUReadback == nullptr)
+				{
+					FLandscapeLayersTexture2DCPUReadBackResource* NewCPUReadBackResource = new FLandscapeLayersTexture2DCPUReadBackResource(ComponentHeightmap->Source.GetSizeX(), ComponentHeightmap->Source.GetSizeY(), ComponentHeightmap->GetPixelFormat(), ComponentHeightmap->Source.GetNumMips());
+					BeginInitResource(NewCPUReadBackResource);
+					Proxy->HeightmapsCPUReadBack.Add(ComponentHeightmap, NewCPUReadBackResource);
+				}
+			}
+		});
+	}
+
 	if ((LayersContentUpdateFlags & ELandscapeLayersContentUpdateFlag::Heightmap_Render) != 0)
 	{
 		check(HeightmapRTList.Num() > 0);
