Fix incorrect notification on import landscape blank field

#rnx

diff --git a/Engine/Source/Editor/LandscapeEditor/Private/LandscapeEditorObject.cpp b/Engine/Source/Editor/LandscapeEditor/Private/LandscapeEditorObject.cpp
--- a/Engine/Source/Editor/LandscapeEditor/Private/LandscapeEditorObject.cpp
+++ b/Engine/Source/Editor/LandscapeEditor/Private/LandscapeEditorObject.cpp
@@ -517,7 +517,7 @@ void ULandscapeEditorObject::ImportLandscapeData()
 {
 	ILandscapeEditorModule& LandscapeEditorModule = FModuleManager::GetModuleChecked<ILandscapeEditorModule>("LandscapeEditor");
 	const ILandscapeHeightmapFileFormat* HeightmapFormat = LandscapeEditorModule.GetHeightmapFormatByExtension(*FPaths::GetExtension(ImportLandscape_HeightmapFilename, true));
-
+	
 	if (HeightmapFormat)
 	{
 		FLandscapeHeightmapImportData HeightmapImportData = HeightmapFormat->Import(*ImportLandscape_HeightmapFilename, {ImportLandscape_Width, ImportLandscape_Height});
@@ -525,7 +525,7 @@ void ULandscapeEditorObject::ImportLandscapeData()
 		ImportLandscape_HeightmapErrorMessage = HeightmapImportData.ErrorMessage;
 		ImportLandscape_Data = MoveTemp(HeightmapImportData.Data);
 	}
-	else
+	else if (!ImportLandscape_HeightmapFilename.IsEmpty())
 	{
 		ImportLandscape_HeightmapImportResult = ELandscapeImportResult::Error;
 		ImportLandscape_HeightmapErrorMessage = NSLOCTEXT("LandscapeEditor.NewLandscape", "Import_UnknownFileType", "File type not recognised");
