Fixed Landscape import button to be disabled if no file path is specified

diff --git a/Engine/Source/Editor/LandscapeEditor/Private/LandscapeEditorDetailCustomization_NewLandscape.cpp b/Engine/Source/Editor/LandscapeEditor/Private/LandscapeEditorDetailCustomization_NewLandscape.cpp
--- a/Engine/Source/Editor/LandscapeEditor/Private/LandscapeEditorDetailCustomization_NewLandscape.cpp
+++ b/Engine/Source/Editor/LandscapeEditor/Private/LandscapeEditorDetailCustomization_NewLandscape.cpp
@@ -911,10 +911,15 @@ bool FLandscapeEditorDetailCustomization_NewLandscape::GetImportButtonIsEnabled(
 	FEdModeLandscape* LandscapeEdMode = GetEditorMode();
 	if (LandscapeEdMode != nullptr)
 	{
+		bool bAllSourceFilePathsEmpty = true;
 		if (LandscapeEdMode->UISettings->ImportLandscape_HeightmapImportResult == ELandscapeImportResult::Error)
 		{
 			return false;
 		}
+		else if (!LandscapeEdMode->UISettings->ImportLandscape_HeightmapFilename.IsEmpty())
+		{
+			bAllSourceFilePathsEmpty = false;
+		}
 
 		for (int32 i = 0; i < LandscapeEdMode->UISettings->ImportLandscape_Layers.Num(); ++i)
 		{
@@ -922,9 +927,13 @@ bool FLandscapeEditorDetailCustomization_NewLandscape::GetImportButtonIsEnabled(
 			{
 				return false;
 			}
+			else if (!LandscapeEdMode->UISettings->ImportLandscape_Layers[i].SourceFilePath.IsEmpty())
+			{
+				bAllSourceFilePathsEmpty = false;
+			}
 		}
 
-		return true;
+		return !bAllSourceFilePathsEmpty;
 	}
 	return false;
 }
