- Landscape Layer System: Revert one unecessary change in previous CL

#rnx

diff --git a/Engine/Source/Editor/LandscapeEditor/Private/LandscapeEdModeComponentTools.cpp b/Engine/Source/Editor/LandscapeEditor/Private/LandscapeEdModeComponentTools.cpp
--- a/Engine/Source/Editor/LandscapeEditor/Private/LandscapeEdModeComponentTools.cpp
+++ b/Engine/Source/Editor/LandscapeEditor/Private/LandscapeEdModeComponentTools.cpp
@@ -1728,6 +1728,8 @@ public:
 	virtual ELandscapeLayersContentUpdateFlag GetBeginToolContentUpdateFlag() const override { return ELandscapeLayersContentUpdateFlag::All_Render; }
 
 	virtual ELandscapeLayersContentUpdateFlag GetTickToolContentUpdateFlag() const override { return ELandscapeLayersContentUpdateFlag::All_Render; }
+	
+	virtual ELandscapeLayersContentUpdateFlag GetEndToolContentUpdateFlag() const override { return ELandscapeLayersContentUpdateFlag::All; }
 		
 	virtual bool BeginTool(FEditorViewportClient* ViewportClient, const FLandscapeToolTarget& InTarget, const FVector& InHitLocation) override
 	{
diff --git a/Engine/Source/Editor/LandscapeEditor/Private/LandscapeEdModeTools.h b/Engine/Source/Editor/LandscapeEditor/Private/LandscapeEdModeTools.h
--- a/Engine/Source/Editor/LandscapeEditor/Private/LandscapeEdModeTools.h
+++ b/Engine/Source/Editor/LandscapeEditor/Private/LandscapeEdModeTools.h
@@ -1144,7 +1144,8 @@ public:
 
 	virtual ELandscapeLayersContentUpdateFlag GetEndToolContentUpdateFlag() const
 	{
-		return ELandscapeLayersContentUpdateFlag::All;
+		bool bUpdateHeightmap = this->EdMode->CurrentToolTarget.TargetType == ELandscapeToolTargetType::Type::Heightmap;
+		return bUpdateHeightmap ? ELandscapeLayersContentUpdateFlag::Heightmap_All : ELandscapeLayersContentUpdateFlag::Weightmap_All;
 	}
 
 	virtual bool BeginTool(FEditorViewportClient* ViewportClient, const FLandscapeToolTarget& InTarget, const FVector& InHitLocation) override
