- Landscape Layer System: New Heightmap Erase Tool (in sculpt mode)

#rnx

diff --git a/Engine/Source/Editor/EditorStyle/Private/SlateEditorStyle.cpp b/Engine/Source/Editor/EditorStyle/Private/SlateEditorStyle.cpp
--- a/Engine/Source/Editor/EditorStyle/Private/SlateEditorStyle.cpp
+++ b/Engine/Source/Editor/EditorStyle/Private/SlateEditorStyle.cpp
@@ -6594,6 +6594,7 @@ void FSlateEditorStyle::FStyle::SetupLandscapeEditorStyle()
 		Set("LandscapeEditor.ResizeLandscape.Small", new IMAGE_BRUSH("Icons/icon_Landscape_Tool_Copy_20x", Icon20x20));
 
 		Set("LandscapeEditor.SculptTool",       new IMAGE_BRUSH("Icons/icon_Landscape_Tool_Sculpt_40x",           Icon40x40));
+		Set("LandscapeEditor.EraseTool",		new IMAGE_BRUSH("Icons/icon_Landscape_Tool_Erase_40x", Icon40x40));
 		Set("LandscapeEditor.PaintTool",        new IMAGE_BRUSH("Icons/icon_Landscape_Tool_Paint_40x",            Icon40x40));
 		Set("LandscapeEditor.SmoothTool",       new IMAGE_BRUSH("Icons/icon_Landscape_Tool_Smooth_40x",           Icon40x40));
 		Set("LandscapeEditor.FlattenTool",      new IMAGE_BRUSH("Icons/icon_Landscape_Tool_Flatten_40x",          Icon40x40));
@@ -6605,6 +6606,7 @@ void FSlateEditorStyle::FStyle::SetupLandscapeEditorStyle()
 		Set("LandscapeEditor.VisibilityTool",   new IMAGE_BRUSH("Icons/icon_Landscape_Tool_Visibility_40x",       Icon40x40));
 		Set("LandscapeEditor.BPCustomTool", new IMAGE_BRUSH("Icons/icon_Landscape_Tool_Visibility_40x", Icon40x40));// TODO: change icon
 		Set("LandscapeEditor.SculptTool.Small",       new IMAGE_BRUSH("Icons/icon_Landscape_Tool_Sculpt_20x",           Icon20x20));
+		Set("LandscapeEditor.EraseTool.Small", new IMAGE_BRUSH("Icons/icon_Landscape_Tool_Erase_20x", Icon20x20));
 		Set("LandscapeEditor.PaintTool.Small",        new IMAGE_BRUSH("Icons/icon_Landscape_Tool_Paint_20x",            Icon20x20));
 		Set("LandscapeEditor.SmoothTool.Small",       new IMAGE_BRUSH("Icons/icon_Landscape_Tool_Smooth_20x",           Icon20x20));
 		Set("LandscapeEditor.FlattenTool.Small",      new IMAGE_BRUSH("Icons/icon_Landscape_Tool_Flatten_20x",          Icon20x20));
diff --git a/Engine/Source/Editor/LandscapeEditor/Private/LandscapeEdMode.cpp b/Engine/Source/Editor/LandscapeEditor/Private/LandscapeEdMode.cpp
--- a/Engine/Source/Editor/LandscapeEditor/Private/LandscapeEdMode.cpp
+++ b/Engine/Source/Editor/LandscapeEditor/Private/LandscapeEdMode.cpp
@@ -426,6 +426,7 @@ void FEdModeLandscape::InitializeToolModes()
 
 	if (GetMutableDefault<UEditorExperimentalSettings>()->bLandscapeLayerSystem)
 	{
+        ToolMode_Sculpt->ValidTools.Add(TEXT("Erase"));
 		ToolMode_Sculpt->ValidTools.Add(TEXT("BPCustom"));
 	}
 
diff --git a/Engine/Source/Editor/LandscapeEditor/Private/LandscapeEdModePaintTools.cpp b/Engine/Source/Editor/LandscapeEditor/Private/LandscapeEdModePaintTools.cpp
--- a/Engine/Source/Editor/LandscapeEditor/Private/LandscapeEdModePaintTools.cpp
+++ b/Engine/Source/Editor/LandscapeEditor/Private/LandscapeEdModePaintTools.cpp
@@ -290,6 +290,89 @@ public:
 	}
 };
 
+class FLandscapeToolStrokeErase : public FLandscapeToolStrokePaintBase<FHeightmapToolTarget>
+{
+	typedef FHeightmapToolTarget ToolTarget;
+	const typename ToolTarget::CacheClass::DataType FlattenHeight;
+
+public:
+	// Heightmap sculpt tool will continuously sculpt in the same location, weightmap paint tool doesn't
+	enum { UseContinuousApply = true };
+
+	FLandscapeToolStrokeErase(FEdModeLandscape* InEdMode, FEditorViewportClient* InViewportClient, const FLandscapeToolTarget& InTarget)
+		: FLandscapeToolStrokePaintBase<FHeightmapToolTarget>(InEdMode, InViewportClient, InTarget)
+		, FlattenHeight(LandscapeDataAccess::GetTexHeight(0.f))
+	{
+	}
+
+	void Apply(FEditorViewportClient* ViewportClient, FLandscapeBrush* Brush, const ULandscapeEditorObject* UISettings, const TArray<FLandscapeToolInteractorPosition>& InteractorPositions)
+	{
+		if (!this->LandscapeInfo) return;
+
+		// Get list of verts to update
+		FLandscapeBrushData BrushInfo = Brush->ApplyBrush(InteractorPositions);
+		if (!BrushInfo)
+		{
+			return;
+		}
+
+		int32 X1, Y1, X2, Y2;
+		BrushInfo.GetInclusiveBounds(X1, Y1, X2, Y2);
+
+		// Tablet pressure
+		float Pressure = ViewportClient->Viewport->IsPenActive() ? ViewportClient->Viewport->GetTabletPressure() : 1.0f;
+
+		// expand the area by one vertex in each direction to ensure normals are calculated correctly
+		X1 -= 1;
+		Y1 -= 1;
+		X2 += 1;
+		Y2 += 1;
+		
+		this->Cache.CacheData(X1, Y1, X2, Y2);
+
+		TArray<typename ToolTarget::CacheClass::DataType> Data;
+		this->Cache.GetCachedData(X1, Y1, X2, Y2, Data);
+
+		// Apply the brush
+		for (int32 Y = BrushInfo.GetBounds().Min.Y; Y < BrushInfo.GetBounds().Max.Y; Y++)
+		{
+			const float* BrushScanline = BrushInfo.GetDataPtr(FIntPoint(0, Y));
+			auto* DataScanline = Data.GetData() + (Y - Y1) * (X2 - X1 + 1) + (0 - X1);
+
+			for (int32 X = BrushInfo.GetBounds().Min.X; X < BrushInfo.GetBounds().Max.X; X++)
+			{
+				const float BrushValue = BrushScanline[X];
+
+				if (BrushValue > 0.0f)
+				{
+					float Strength = FMath::Clamp<float>(BrushValue * UISettings->ToolStrength * Pressure, 0.0f, 1.0f);
+
+					int32 Delta = DataScanline[X] - FlattenHeight;
+					if (Delta > 0)
+					{
+						DataScanline[X] = FMath::FloorToInt(FMath::Lerp((float)DataScanline[X], (float)FlattenHeight, Strength));
+					}
+					else
+					{
+						DataScanline[X] = FMath::CeilToInt(FMath::Lerp((float)DataScanline[X], (float)FlattenHeight, Strength));
+					}
+				}
+			}
+		}
+
+		ALandscape* Landscape = this->LandscapeInfo->LandscapeActor.Get();
+
+		if (Landscape != nullptr && Landscape->HasLayersContent && !GetMutableDefault<UEditorExperimentalSettings>()->bLandscapeLayerSystem)
+		{
+			FMessageLog("MapCheck").Warning()->AddToken(FTextToken::Create(LOCTEXT("LandscapeLayers_ChangingDataWithoutSettings", "This map contains landscape layer system content, modifying the landscape data will result in data loss when the map is reopened with Landscape Layer System settings on. Please enable Landscape Layer System settings before modifying the data.")));
+			FMessageLog("MapCheck").Open(EMessageSeverity::Warning);
+		}
+
+		this->Cache.SetCachedData(X1, Y1, X2, Y2, Data, UISettings->PaintingRestriction);
+		this->Cache.Flush();
+	}
+};
+
 //
 class FLandscapeToolStrokeSculpt : public FLandscapeToolStrokePaintBase<FHeightmapToolTarget>
 {
@@ -523,6 +606,18 @@ public:
 	virtual FText GetDisplayName() override { return NSLOCTEXT("UnrealEd", "LandscapeMode_Sculpt", "Sculpt"); };
 };
 
+class FLandscapeToolErase : public FLandscapeToolPaintBase<FHeightmapToolTarget, FLandscapeToolStrokeErase>
+{
+public:
+	FLandscapeToolErase(FEdModeLandscape* InEdMode)
+		: FLandscapeToolPaintBase<FHeightmapToolTarget, FLandscapeToolStrokeErase>(InEdMode)
+	{
+	}
+
+	virtual const TCHAR* GetToolName() override { return TEXT("Erase"); }
+	virtual FText GetDisplayName() override { return NSLOCTEXT("UnrealEd", "LandscapeMode_Erase", "Erase"); };
+};
+
 // 
 // FLandscapeToolSmooth
 //
@@ -1143,6 +1238,10 @@ void FEdModeLandscape::InitializeTool_Paint()
 	Tool_Sculpt->ValidBrushes.Add("BrushSet_Component");
 	LandscapeTools.Add(MoveTemp(Tool_Sculpt));
 
+	auto Tool_Erase = MakeUnique<FLandscapeToolErase>(this);
+	Tool_Erase->ValidBrushes.Add("BrushSet_Circle");
+	LandscapeTools.Add(MoveTemp(Tool_Erase));
+
 	auto Tool_Paint = MakeUnique<FLandscapeToolPaint>(this);
 	Tool_Paint->ValidBrushes.Add("BrushSet_Circle");
 	Tool_Paint->ValidBrushes.Add("BrushSet_Alpha");
diff --git a/Engine/Source/Editor/LandscapeEditor/Private/LandscapeEditorCommands.cpp b/Engine/Source/Editor/LandscapeEditor/Private/LandscapeEditorCommands.cpp
--- a/Engine/Source/Editor/LandscapeEditor/Private/LandscapeEditorCommands.cpp
+++ b/Engine/Source/Editor/LandscapeEditor/Private/LandscapeEditorCommands.cpp
@@ -35,6 +35,8 @@ void FLandscapeEditorCommands::RegisterCommands()
 
 	UI_COMMAND(SculptTool, "Tool - Sculpt", "", EUserInterfaceActionType::RadioButton, FInputChord());
 	NameToCommandMap.Add("Tool_Sculpt", SculptTool);
+	UI_COMMAND(EraseTool, "Tool - Erase", "", EUserInterfaceActionType::RadioButton, FInputChord());
+	NameToCommandMap.Add("Tool_Erase", EraseTool);
 	UI_COMMAND(PaintTool, "Tool - Paint", "", EUserInterfaceActionType::RadioButton, FInputChord());
 	NameToCommandMap.Add("Tool_Paint", PaintTool);
 	UI_COMMAND(SmoothTool, "Tool - Smooth", "", EUserInterfaceActionType::RadioButton, FInputChord());
diff --git a/Engine/Source/Editor/LandscapeEditor/Private/LandscapeEditorCommands.h b/Engine/Source/Editor/LandscapeEditor/Private/LandscapeEditorCommands.h
--- a/Engine/Source/Editor/LandscapeEditor/Private/LandscapeEditorCommands.h
+++ b/Engine/Source/Editor/LandscapeEditor/Private/LandscapeEditorCommands.h
@@ -34,6 +34,7 @@ public:
 	TSharedPtr<FUICommandInfo> ResizeLandscape;
 
 	TSharedPtr<FUICommandInfo> SculptTool;
+	TSharedPtr<FUICommandInfo> EraseTool;
 	TSharedPtr<FUICommandInfo> PaintTool;
 	TSharedPtr<FUICommandInfo> SmoothTool;
 	TSharedPtr<FUICommandInfo> FlattenTool;
diff --git a/Engine/Source/Editor/LandscapeEditor/Private/LandscapeEditorDetails.cpp b/Engine/Source/Editor/LandscapeEditor/Private/LandscapeEditorDetails.cpp
--- a/Engine/Source/Editor/LandscapeEditor/Private/LandscapeEditorDetails.cpp
+++ b/Engine/Source/Editor/LandscapeEditor/Private/LandscapeEditorDetails.cpp
@@ -133,6 +133,7 @@ FText FLandscapeEditorDetails::GetLocalizedName(FString Name)
 		LOCTEXT("ToolSet_NewLandscape", "New Landscape");
 		LOCTEXT("ToolSet_ResizeLandscape", "Change Component Size");
 		LOCTEXT("ToolSet_Sculpt", "Sculpt");
+		LOCTEXT("ToolSet_Erase", "Erase");
 		LOCTEXT("ToolSet_Paint", "Paint");
 		LOCTEXT("ToolSet_Smooth", "Smooth");
 		LOCTEXT("ToolSet_Flatten", "Flatten");
@@ -291,6 +292,12 @@ TSharedRef<SWidget> FLandscapeEditorDetails::GetToolSelector()
 		{
 			MenuBuilder.BeginSection(NAME_None, LOCTEXT("SculptToolsTitle", "Sculpting Tools"));
 			MenuBuilder.AddToolButton(NameToCommandMap.FindChecked("Tool_Sculpt"), NAME_None, LOCTEXT("Tool.Sculpt", "Sculpt"), LOCTEXT("Tool.Sculpt.Tooltip", "Sculpt height data.\nCtrl+Click to Raise, Ctrl+Shift+Click to lower"));
+			
+			if (GetMutableDefault<UEditorExperimentalSettings>()->bLandscapeLayerSystem)
+			{
+				MenuBuilder.AddToolButton(NameToCommandMap.FindChecked("Tool_Erase"), NAME_None, LOCTEXT("Tool.Erase", "Erase"), LOCTEXT("Tool.Erase.Tooltip", "Erase height data."));
+			}
+			
 			MenuBuilder.AddToolButton(NameToCommandMap.FindChecked("Tool_Smooth"), NAME_None, LOCTEXT("Tool.Smooth", "Smooth"), LOCTEXT("Tool.Smooth.Tooltip", "Smooths heightmaps or blend layers"));
 			MenuBuilder.AddToolButton(NameToCommandMap.FindChecked("Tool_Flatten"), NAME_None, LOCTEXT("Tool.Flatten", "Flatten"), LOCTEXT("Tool.Flatten.Tooltip", "Flattens an area of heightmap or blend layer"));
 			MenuBuilder.AddToolButton(NameToCommandMap.FindChecked("Tool_Ramp"), NAME_None, LOCTEXT("Tool.Ramp", "Ramp"), LOCTEXT("Tool.Ramp.Tooltip", "Creates a ramp between two points"));
diff --git a/Engine/Source/Editor/LandscapeEditor/Private/SLandscapeEditor.cpp b/Engine/Source/Editor/LandscapeEditor/Private/SLandscapeEditor.cpp
--- a/Engine/Source/Editor/LandscapeEditor/Private/SLandscapeEditor.cpp
+++ b/Engine/Source/Editor/LandscapeEditor/Private/SLandscapeEditor.cpp
@@ -95,6 +95,7 @@ void FLandscapeToolKit::Init(const TSharedPtr<IToolkitHost>& InitToolkitHost)
 	MAP_TOOL("ResizeLandscape");
 
 	MAP_TOOL("Sculpt");
+	MAP_TOOL("Erase");
 	MAP_TOOL("Paint");
 	MAP_TOOL("Smooth");
 	MAP_TOOL("Flatten");
diff --git a/Engine/Source/Editor/LandscapeEditor/Public/LandscapeEditorObject.h b/Engine/Source/Editor/LandscapeEditor/Public/LandscapeEditorObject.h
--- a/Engine/Source/Editor/LandscapeEditor/Public/LandscapeEditorObject.h
+++ b/Engine/Source/Editor/LandscapeEditor/Public/LandscapeEditorObject.h
@@ -251,7 +251,7 @@ class ULandscapeEditorObject : public UObject
 	// Common Tool Settings:
 
 	// Strength of the tool. If you're using a pen/tablet with pressure-sensing, the pressure used affects the strength of the tool.
-	UPROPERTY(Category="Tool Settings", EditAnywhere, NonTransactional, meta=(ShowForTools="Paint,Sculpt,Smooth,Flatten,Erosion,HydraErosion,Noise,Mask,CopyPaste", ClampMin="0", ClampMax="10", UIMin="0", UIMax="1"))
+	UPROPERTY(Category="Tool Settings", EditAnywhere, NonTransactional, meta=(ShowForTools="Paint,Sculpt,Erase,Smooth,Flatten,Erosion,HydraErosion,Noise,Mask,CopyPaste", ClampMin="0", ClampMax="10", UIMin="0", UIMax="1"))
 	float ToolStrength;
 
 	// Enable to make tools blend towards a target value
