Refresh Landscape Tool when toggling Experimental Landscape Layer System flag
#rnx

diff --git a/Engine/Source/Editor/LandscapeEditor/Private/LandscapeEdMode.cpp b/Engine/Source/Editor/LandscapeEditor/Private/LandscapeEdMode.cpp
--- a/Engine/Source/Editor/LandscapeEditor/Private/LandscapeEdMode.cpp
+++ b/Engine/Source/Editor/LandscapeEditor/Private/LandscapeEdMode.cpp
@@ -466,6 +466,13 @@ void FEdModeLandscape::Enter()
 
 	OnLevelActorDeletedDelegateHandle = GEngine->OnLevelActorDeleted().AddSP(this, &FEdModeLandscape::OnLevelActorRemoved);
 	OnLevelActorAddedDelegateHandle = GEngine->OnLevelActorAdded().AddSP(this, &FEdModeLandscape::OnLevelActorAdded);
+	OnLandscapeLayerSystemFlagChangedDelegateHandle = GetMutableDefault<UEditorExperimentalSettings>()->OnSettingChanged().AddLambda([this](FName PropertyName)
+	{ 
+		if (PropertyName == TEXT("bLandscapeLayerSystem"))
+		{
+			RefreshDetailPanel();
+		}
+	});
 
 	ALandscapeProxy* SelectedLandscape = GEditor->GetSelectedActors()->GetTop<ALandscapeProxy>();
 	if (SelectedLandscape)
@@ -662,6 +669,7 @@ void FEdModeLandscape::Exit()
 
 	GEngine->OnLevelActorDeleted().Remove(OnLevelActorDeletedDelegateHandle);
 	GEngine->OnLevelActorAdded().Remove(OnLevelActorAddedDelegateHandle);
+	GetMutableDefault<UEditorExperimentalSettings>()->OnSettingChanged().Remove(OnLandscapeLayerSystemFlagChangedDelegateHandle);
 	
 	FEditorSupportDelegates::WorldChange.Remove(OnWorldChangeDelegateHandle);
 	GetWorld()->OnLevelsChanged().Remove(OnLevelsChangedDelegateHandle);
diff --git a/Engine/Source/Editor/LandscapeEditor/Private/LandscapeEdMode.h b/Engine/Source/Editor/LandscapeEditor/Private/LandscapeEdMode.h
--- a/Engine/Source/Editor/LandscapeEditor/Private/LandscapeEdMode.h
+++ b/Engine/Source/Editor/LandscapeEditor/Private/LandscapeEdMode.h
@@ -583,6 +583,7 @@ private:
 
 	FDelegateHandle OnLevelActorDeletedDelegateHandle;
 	FDelegateHandle OnLevelActorAddedDelegateHandle;
+	FDelegateHandle OnLandscapeLayerSystemFlagChangedDelegateHandle;
 
 	/** Check if we are painting using the VREditor */
 	bool bIsPaintingInVR;
