- Landscape Layer System: Prevent editing of Landscapes in some situations with error banner
Landscape Actor not loaded (in layer mode)
Landscape has layer content but Experimental flag is disabled

#rnx

diff --git a/Engine/Source/Editor/LandscapeEditor/Private/LandscapeEdMode.cpp b/Engine/Source/Editor/LandscapeEditor/Private/LandscapeEdMode.cpp
--- a/Engine/Source/Editor/LandscapeEditor/Private/LandscapeEdMode.cpp
+++ b/Engine/Source/Editor/LandscapeEditor/Private/LandscapeEdMode.cpp
@@ -470,6 +470,7 @@ void FEdModeLandscape::Enter()
 	{ 
 		if (PropertyName == TEXT("bLandscapeLayerSystem"))
 		{
+			UpdateLandscapeList();
 			RefreshDetailPanel();
 		}
 	});
@@ -2306,15 +2307,6 @@ int32 FEdModeLandscape::UpdateLandscapeList()
 				ALandscapeProxy* LandscapeProxy = LandscapeInfo->GetLandscapeProxy();
 				if (LandscapeProxy)
 				{
-					// In Layer System for now disable Landscape Editing if Actor is not loaded
-					if (GetMutableDefault<UEditorExperimentalSettings>()->bLandscapeLayerSystem)
-					{
-						if (!LandscapeProxy->GetLandscapeActor())
-						{
-							continue;
-						}
-					}
-
 					if (CurrentToolTarget.LandscapeInfo == LandscapeInfo)
 					{
 						CurrentIndex = Index;
@@ -2383,6 +2375,12 @@ int32 FEdModeLandscape::UpdateLandscapeList()
 		}
 	}
 
+	if (!CanEditCurrentTarget())
+	{
+		SetCurrentToolMode("ToolMode_Manage", false);
+		SetCurrentTool("NewLandscape");
+	}
+
 	return CurrentIndex;
 }
 
@@ -2420,6 +2418,41 @@ void FEdModeLandscape::SetTargetLandscape(const TWeakObjectPtr<ULandscapeInfo>&
 	UpdateShownLayerList();
 }
 
+bool FEdModeLandscape::CanEditCurrentTarget(FText* Reason) const
+{
+	static FText DummyReason;
+	FText& LocalReason = Reason ? *Reason : DummyReason;
+
+	if (!CurrentToolTarget.LandscapeInfo.IsValid())
+	{
+		LocalReason = NSLOCTEXT("UnrealEd", "LandscapeInvalidTarget", "No landscape selected.");
+		return false;
+	}
+
+	if (!GetMutableDefault<UEditorExperimentalSettings>()->bLandscapeLayerSystem)
+	{
+		// If Any Proxy has layer content disable editing when Landscape Layers are disabled
+		bool bHasLayerContent = false;
+		CurrentToolTarget.LandscapeInfo->ForAllLandscapeProxies([&bHasLayerContent](ALandscapeProxy* Proxy) { bHasLayerContent |= Proxy->HasLayersContent; });
+		if (bHasLayerContent)
+		{
+			LocalReason = NSLOCTEXT("UnrealEd", "LandscapeWithLayerContent", "Landscape contains layers and experimental flag is not enabled.");
+			return false;
+		}
+
+		return true;
+	}
+
+	// Landscape Layer Editing not available without a loaded Landscape Actor
+	if (CurrentToolTarget.LandscapeInfo->LandscapeActor == nullptr)
+	{
+		LocalReason = NSLOCTEXT("UnrealEd", "LandscapeActorNotLoaded", "Landscape actor is not loaded. It is needed to do layer editing.");
+		return false;
+	}
+
+	return true;
+}
+
 void FEdModeLandscape::UpdateTargetList()
 {
 	LandscapeTargetList.Empty();
@@ -2838,7 +2871,7 @@ void FEdModeLandscape::HandleLevelsChanged(bool ShouldExitMode)
 	}
 
 	// if a landscape is added somehow then switch to sculpt
-	if (!bHadLandscape && CurrentToolTarget.LandscapeInfo != nullptr)
+	if (!bHadLandscape && CanEditCurrentTarget())
 	{
 		SetCurrentTool("Select");
 		SetCurrentTool("Sculpt");
diff --git a/Engine/Source/Editor/LandscapeEditor/Private/LandscapeEdMode.h b/Engine/Source/Editor/LandscapeEditor/Private/LandscapeEdMode.h
--- a/Engine/Source/Editor/LandscapeEditor/Private/LandscapeEdMode.h
+++ b/Engine/Source/Editor/LandscapeEditor/Private/LandscapeEdMode.h
@@ -478,6 +478,7 @@ public:
 	int32 UpdateLandscapeList();
 	void UpdateTargetList();
 	void SetTargetLandscape(const TWeakObjectPtr<ULandscapeInfo>& InLandscapeInfo);
+	bool CanEditCurrentTarget(FText* Reason = nullptr) const;
 
 	/** Update Display order list */
 	void UpdateTargetLayerDisplayOrder(ELandscapeLayerDisplayMode InTargetDisplayOrder);
diff --git a/Engine/Source/Editor/LandscapeEditor/Private/LandscapeEditorDetails.cpp b/Engine/Source/Editor/LandscapeEditor/Private/LandscapeEditorDetails.cpp
--- a/Engine/Source/Editor/LandscapeEditor/Private/LandscapeEditorDetails.cpp
+++ b/Engine/Source/Editor/LandscapeEditor/Private/LandscapeEditorDetails.cpp
@@ -6,6 +6,7 @@
 #include "Widgets/Text/STextBlock.h"
 #include "Framework/MultiBox/MultiBoxBuilder.h"
 #include "Widgets/Input/SComboButton.h"
+#include "Widgets/Input/SMultiLineEditableTextBox.h"
 #include "LandscapeEdMode.h"
 #include "LandscapeEditorDetailCustomization_NewLandscape.h"
 #include "LandscapeEditorDetailCustomization_ResizeLandscape.h"
@@ -21,6 +22,7 @@
 #include "LandscapeEditorDetailWidgets.h"
 #include "LandscapeEditorDetailCustomization_LayersBrushStack.h"
 #include "Settings/EditorExperimentalSettings.h"
+#include "Landscape.h"
 
 #define LOCTEXT_NAMESPACE "LandscapeEditor"
 
@@ -52,6 +54,30 @@ void FLandscapeEditorDetails::CustomizeDetails(IDetailLayoutBuilder& DetailBuild
 		]
 	];
 
+	FText Reason;
+	bool bDisabledEditing = LandscapeEdMode->CurrentToolTarget.LandscapeInfo.IsValid() && !LandscapeEdMode->CanEditCurrentTarget(&Reason);
+
+	if (bDisabledEditing)
+	{
+		LandscapeEditorCategory.AddCustomRow(FText::GetEmpty())
+			[
+				SNew(SMultiLineEditableTextBox)
+				.IsReadOnly(true)
+				.AutoWrapText(true)
+				.Font(FCoreStyle::GetDefaultFontStyle("Bold", 10))
+				.Justification(ETextJustify::Center)
+				.BackgroundColor(FCoreStyle::Get().GetColor("ErrorReporting.BackgroundColor"))
+				.ForegroundColor(FCoreStyle::Get().GetColor("ErrorReporting.ForegroundColor"))
+				.Text(Reason)
+			];
+	}
+		
+	// Only continue cuztomization if we are in NewLandscape mode or if editing is not disabled
+	if (bDisabledEditing && LandscapeEdMode->CurrentTool->GetToolName() != FName("NewLandscape"))
+	{
+		return;
+	}
+
 	FToolSelectorBuilder ToolBrushSelectorButtons(CommandList, FMultiBoxCustomization::None);
 	{
 		FUIAction ToolSelectorUIAction;
@@ -348,6 +374,11 @@ bool FLandscapeEditorDetails::GetToolSelectorIsVisible() const
 	FEdModeLandscape* LandscapeEdMode = GetEditorMode();
 	if (LandscapeEdMode && LandscapeEdMode->CurrentTool)
 	{
+		if (!LandscapeEdMode->CanEditCurrentTarget())
+		{
+			return false;
+		}
+
 		if (!IsToolActive("NewLandscape") || LandscapeEdMode->GetLandscapeList().Num() > 0)
 		{
 			return true;
diff --git a/Engine/Source/Editor/LandscapeEditor/Private/SLandscapeEditor.cpp b/Engine/Source/Editor/LandscapeEditor/Private/SLandscapeEditor.cpp
--- a/Engine/Source/Editor/LandscapeEditor/Private/SLandscapeEditor.cpp
+++ b/Engine/Source/Editor/LandscapeEditor/Private/SLandscapeEditor.cpp
@@ -189,7 +189,7 @@ bool FLandscapeToolKit::IsModeEnabled(FName ModeName) const
 	if (LandscapeEdMode)
 	{
 		// Manage is the only mode enabled if we have no landscape
-		if (ModeName == "ToolMode_Manage" || LandscapeEdMode->GetLandscapeList().Num() > 0)
+		if (ModeName == "ToolMode_Manage" || (LandscapeEdMode->GetLandscapeList().Num() > 0 && LandscapeEdMode->CanEditCurrentTarget()))
 		{
 			return true;
 		}
@@ -466,8 +466,16 @@ bool SLandscapeEditor::GetIsPropertyVisible(const FPropertyAndParent& PropertyAn
 	const UProperty& Property = PropertyAndParent.Property;
 
 	FEdModeLandscape* LandscapeEdMode = GetEditorMode();
+	
 	if (LandscapeEdMode != nullptr && LandscapeEdMode->CurrentTool != nullptr)
 	{
+	    // Hide all properties if the current target can't be edited. Except in New Landscape tool
+		if (LandscapeEdMode->CurrentTool->GetToolName() != FName("NewLandscape") &&
+			!LandscapeEdMode->CanEditCurrentTarget())
+		{
+			return false;
+		}
+
 		if (Property.HasMetaData("ShowForMask"))
 		{
 			const bool bMaskEnabled = LandscapeEdMode->CurrentTool &&
