Added support for the layer system into Change Component Size tool
Changed the ALandscapeProxy::Import function to support import of multiple layers as the function is used during the New Landscape with/without Import, World Composition Import/Add and Change Component Size.
The order of operations was changed to prevent having to call multiple times CreateLandscapeInfo, as the 1st few times, the content is not all ready.

diff --git a/Engine/Source/Editor/LandscapeEditor/Private/LandscapeEdMode.cpp b/Engine/Source/Editor/LandscapeEditor/Private/LandscapeEdMode.cpp
--- a/Engine/Source/Editor/LandscapeEditor/Private/LandscapeEdMode.cpp
+++ b/Engine/Source/Editor/LandscapeEditor/Private/LandscapeEdMode.cpp
@@ -2303,7 +2303,7 @@ int32 FEdModeLandscape::UpdateLandscapeList()
 		{
 			ULandscapeInfo* LandscapeInfo = It.Value();
 			if (LandscapeInfo && !LandscapeInfo->IsPendingKill())
-			{
+			{				
 				ALandscapeProxy* LandscapeProxy = LandscapeInfo->GetLandscapeProxy();
 				if (LandscapeProxy)
 				{
@@ -3800,7 +3800,7 @@ ALandscape* FEdModeLandscape::ChangeComponentSetting(int32 NumComponentsX, int32
 
 	const int32 NewComponentSizeQuads = NumSubsections * SubsectionSizeQuads;
 
-	ALandscape* Landscape = NULL;
+	ALandscape* NewLandscape = NULL;
 
 	ULandscapeInfo* LandscapeInfo = CurrentToolTarget.LandscapeInfo.Get();
 	if (ensure(LandscapeInfo != NULL))
@@ -3808,18 +3808,20 @@ ALandscape* FEdModeLandscape::ChangeComponentSetting(int32 NumComponentsX, int32
 		int32 OldMinX, OldMinY, OldMaxX, OldMaxY;
 		if (LandscapeInfo->GetLandscapeExtent(OldMinX, OldMinY, OldMaxX, OldMaxY))
 		{
-			ALandscapeProxy* OldLandscapeProxy = LandscapeInfo->GetLandscapeProxy();
+			ALandscape* OldLandscape = LandscapeInfo->LandscapeActor.Get();
+			check(OldLandscape != nullptr);
 
 			const int32 OldVertsX = OldMaxX - OldMinX + 1;
 			const int32 OldVertsY = OldMaxY - OldMinY + 1;
 			const int32 NewVertsX = NumComponentsX * NewComponentSizeQuads + 1;
 			const int32 NewVertsY = NumComponentsY * NewComponentSizeQuads + 1;
 
-			TArray<uint16> HeightData;
-			TArray<FLandscapeImportLayerInfo> ImportLayerInfos;
+			TMap<FGuid, TArray<uint16>> HeightDataPerLayers;			
+			TMap<FGuid, TArray<FLandscapeImportLayerInfo>> ImportMaterialLayerInfosPerLayers;
+
 			FVector LandscapeOffset = FVector::ZeroVector;
 			FIntPoint LandscapeOffsetQuads = FIntPoint::ZeroValue;
-			float LandscapeScaleFactor = 1.0f;
+			float LandscapeScaleFactor = bResample ? (float)OldLandscape->ComponentSizeQuads / NewComponentSizeQuads : 1.0f;
 
 			int32 NewMinX, NewMinY, NewMaxX, NewMaxY;
 
@@ -3832,77 +3834,126 @@ ALandscape* FEdModeLandscape::ChangeComponentSetting(int32 NumComponentsX, int32
 					NewMinY = OldMinY / LandscapeInfo->ComponentSizeQuads * NewComponentSizeQuads;
 					NewMaxX = NewMinX + NewVertsX - 1;
 					NewMaxY = NewMinY + NewVertsY - 1;
+				}
+				else
+				{
+					NewMinX = OldMinX + (OldVertsX - NewVertsX) / 2;
+					NewMinY = OldMinY + (OldVertsY - NewVertsY) / 2;
+					NewMaxX = NewMinX + NewVertsX - 1;
+					NewMaxY = NewMinY + NewVertsY - 1;
+
+					// offset landscape to component boundary
+					LandscapeOffset = FVector(NewMinX, NewMinY, 0) * OldLandscape->GetActorScale();
+					LandscapeOffsetQuads = FIntPoint(NewMinX, NewMinY);
+				}
+
+				auto ExtractHeightmapWeightmapContent = [&](TArray<uint16>& OutHeightData, TArray<FLandscapeImportLayerInfo>& OutImportMaterialLayerInfos)
+				{
+					if (bResample)
+					{
+						OutHeightData.AddZeroed(OldVertsX * OldVertsY * sizeof(uint16));
+
+						// GetHeightData alters its args, so make temp copies to avoid screwing things up
+						int32 TMinX = OldMinX, TMinY = OldMinY, TMaxX = OldMaxX, TMaxY = OldMaxY;
+						LandscapeEdit.GetHeightData(TMinX, TMinY, TMaxX, TMaxY, OutHeightData.GetData(), 0);
 
-					HeightData.AddZeroed(OldVertsX * OldVertsY * sizeof(uint16));
+						OutHeightData = LandscapeEditorUtils::ResampleData(OutHeightData, OldVertsX, OldVertsY, NewVertsX, NewVertsY);
 
-					// GetHeightData alters its args, so make temp copies to avoid screwing things up
-					int32 TMinX = OldMinX, TMinY = OldMinY, TMaxX = OldMaxX, TMaxY = OldMaxY;
-					LandscapeEdit.GetHeightData(TMinX, TMinY, TMaxX, TMaxY, HeightData.GetData(), 0);
+						for (const FLandscapeInfoLayerSettings& LayerSettings : LandscapeInfo->Layers)
+						{
+							if (LayerSettings.LayerInfoObj != NULL)
+							{
+								auto ImportLayerInfo = new(OutImportMaterialLayerInfos) FLandscapeImportLayerInfo(LayerSettings);
+								ImportLayerInfo->LayerData.AddZeroed(OldVertsX * OldVertsY * sizeof(uint8));
 
-					HeightData = LandscapeEditorUtils::ResampleData(HeightData,
-						OldVertsX, OldVertsY, NewVertsX, NewVertsY);
+								TMinX = OldMinX; TMinY = OldMinY; TMaxX = OldMaxX; TMaxY = OldMaxY;
+								LandscapeEdit.GetWeightData(LayerSettings.LayerInfoObj, TMinX, TMinY, TMaxX, TMaxY, ImportLayerInfo->LayerData.GetData(), 0);
 
-					for (const FLandscapeInfoLayerSettings& LayerSettings : LandscapeInfo->Layers)
+								ImportLayerInfo->LayerData = LandscapeEditorUtils::ResampleData(ImportLayerInfo->LayerData, OldVertsX, OldVertsY, NewVertsX, NewVertsY);
+							}
+						}
+					}
+					else
 					{
-						if (LayerSettings.LayerInfoObj != NULL)
+						const int32 RequestedMinX = FMath::Max(OldMinX, NewMinX);
+						const int32 RequestedMinY = FMath::Max(OldMinY, NewMinY);
+						const int32 RequestedMaxX = FMath::Min(OldMaxX, NewMaxX);
+						const int32 RequestedMaxY = FMath::Min(OldMaxY, NewMaxY);
+
+						const int32 RequestedVertsX = RequestedMaxX - RequestedMinX + 1;
+						const int32 RequestedVertsY = RequestedMaxY - RequestedMinY + 1;
+
+						OutHeightData.AddZeroed(RequestedVertsX * RequestedVertsY * sizeof(uint16));
+
+						// GetHeightData alters its args, so make temp copies to avoid screwing things up
+						int32 TMinX = RequestedMinX, TMinY = RequestedMinY, TMaxX = RequestedMaxX, TMaxY = RequestedMaxY;
+						LandscapeEdit.GetHeightData(TMinX, TMinY, TMaxX, OldMaxY, OutHeightData.GetData(), 0);
+
+						OutHeightData = LandscapeEditorUtils::ExpandData(OutHeightData, RequestedMinX, RequestedMinY, RequestedMaxX, RequestedMaxY, NewMinX, NewMinY, NewMaxX, NewMaxY);
+
+						for (const FLandscapeInfoLayerSettings& LayerSettings : LandscapeInfo->Layers)
 						{
-							auto ImportLayerInfo = new(ImportLayerInfos)FLandscapeImportLayerInfo(LayerSettings);
-							ImportLayerInfo->LayerData.AddZeroed(OldVertsX * OldVertsY * sizeof(uint8));
+							if (LayerSettings.LayerInfoObj != NULL)
+							{
+								auto ImportLayerInfo = new(OutImportMaterialLayerInfos) FLandscapeImportLayerInfo(LayerSettings);
+								ImportLayerInfo->LayerData.AddZeroed(NewVertsX * NewVertsY * sizeof(uint8));
 
-							TMinX = OldMinX; TMinY = OldMinY; TMaxX = OldMaxX; TMaxY = OldMaxY;
-							LandscapeEdit.GetWeightData(LayerSettings.LayerInfoObj, TMinX, TMinY, TMaxX, TMaxY, ImportLayerInfo->LayerData.GetData(), 0);
+								TMinX = RequestedMinX; TMinY = RequestedMinY; TMaxX = RequestedMaxX; TMaxY = RequestedMaxY;
+								LandscapeEdit.GetWeightData(LayerSettings.LayerInfoObj, TMinX, TMinY, TMaxX, TMaxY, ImportLayerInfo->LayerData.GetData(), 0);
 
-							ImportLayerInfo->LayerData = LandscapeEditorUtils::ResampleData(ImportLayerInfo->LayerData,
-								OldVertsX, OldVertsY,
-								NewVertsX, NewVertsY);
+								ImportLayerInfo->LayerData = LandscapeEditorUtils::ExpandData(ImportLayerInfo->LayerData, RequestedMinX, RequestedMinY, RequestedMaxX, RequestedMaxY, NewMinX, NewMinY, NewMaxX, NewMaxY);
+							}
 						}
 					}
+				};
 
-					LandscapeScaleFactor = (float)OldLandscapeProxy->ComponentSizeQuads / NewComponentSizeQuads;
-				}
-				else
+				if (GetMutableDefault<UEditorExperimentalSettings>()->bLandscapeLayerSystem)
 				{
-					NewMinX = OldMinX + (OldVertsX - NewVertsX) / 2;
-					NewMinY = OldMinY + (OldVertsY - NewVertsY) / 2;
-					NewMaxX = NewMinX + NewVertsX - 1;
-					NewMaxY = NewMinY + NewVertsY - 1;
-					const int32 RequestedMinX = FMath::Max(OldMinX, NewMinX);
-					const int32 RequestedMinY = FMath::Max(OldMinY, NewMinY);
-					const int32 RequestedMaxX = FMath::Min(OldMaxX, NewMaxX);
-					const int32 RequestedMaxY = FMath::Min(OldMaxY, NewMaxY);
+					int32 HeightCount = 0;
 
-					const int32 RequestedVertsX = RequestedMaxX - RequestedMinX + 1;
-					const int32 RequestedVertsY = RequestedMaxY - RequestedMinY + 1;
+					for (const FLandscapeLayer& OldLayer : OldLandscape->LandscapeLayers)
+					{
+						FScopedSetLandscapeEditingLayer Scope(OldLandscape, OldLayer.Guid);
 
-					HeightData.AddZeroed(RequestedVertsX * RequestedVertsY * sizeof(uint16));
+						TArray<uint16> HeightData;
+						TArray<FLandscapeImportLayerInfo> ImportMaterialLayerInfos;
 
-					// GetHeightData alters its args, so make temp copies to avoid screwing things up
-					int32 TMinX = RequestedMinX, TMinY = RequestedMinY, TMaxX = RequestedMaxX, TMaxY = RequestedMaxY;
-					LandscapeEdit.GetHeightData(TMinX, TMinY, TMaxX, OldMaxY, HeightData.GetData(), 0);
+						ExtractHeightmapWeightmapContent(HeightData, ImportMaterialLayerInfos);
 
-					HeightData = LandscapeEditorUtils::ExpandData(HeightData,
-						RequestedMinX, RequestedMinY, RequestedMaxX, RequestedMaxY,
-						NewMinX, NewMinY, NewMaxX, NewMaxY);
+						HeightCount = FMath::Max(HeightCount, HeightData.Num());
+						HeightDataPerLayers.Add(OldLayer.Guid, MoveTemp(HeightData));
+						ImportMaterialLayerInfosPerLayers.Add(OldLayer.Guid, MoveTemp(ImportMaterialLayerInfos));
+					}
 
-					for (const FLandscapeInfoLayerSettings& LayerSettings : LandscapeInfo->Layers)
-					{
-						if (LayerSettings.LayerInfoObj != NULL)
-						{
-							auto ImportLayerInfo = new(ImportLayerInfos)FLandscapeImportLayerInfo(LayerSettings);
-							ImportLayerInfo->LayerData.AddZeroed(NewVertsX * NewVertsY * sizeof(uint8));
+					TArray<uint16> DefaultHeightData;
+					DefaultHeightData.AddUninitialized(HeightCount);
 
-							TMinX = RequestedMinX; TMinY = RequestedMinY; TMaxX = RequestedMaxX; TMaxY = RequestedMaxY;
-							LandscapeEdit.GetWeightData(LayerSettings.LayerInfoObj, TMinX, TMinY, TMaxX, TMaxY, ImportLayerInfo->LayerData.GetData(), 0);
+					uint16 DefaultValue = LandscapeDataAccess::GetTexHeight(0.f);
 
-							ImportLayerInfo->LayerData = LandscapeEditorUtils::ExpandData(ImportLayerInfo->LayerData,
-								RequestedMinX, RequestedMinY, RequestedMaxX, RequestedMaxY,
-								NewMinX, NewMinY, NewMaxX, NewMaxY);
-						}
+					// Initialize blank heightmap data
+					for (int32 i = 0; i < DefaultHeightData.Num(); i++)
+					{
+						DefaultHeightData[i] = DefaultValue;
 					}
 
-					// offset landscape to component boundary
-					LandscapeOffset = FVector(NewMinX, NewMinY, 0) * OldLandscapeProxy->GetActorScale();
-					LandscapeOffsetQuads = FIntPoint(NewMinX, NewMinY);
+					HeightDataPerLayers.Add(FGuid(), DefaultHeightData);
+
+					TArray<FLandscapeImportLayerInfo> EmptyImportLayer;
+					ImportMaterialLayerInfosPerLayers.Add(FGuid(), EmptyImportLayer);
+				}
+				else
+				{
+					TArray<uint16> HeightData;
+					TArray<FLandscapeImportLayerInfo> ImportMaterialLayerInfos;
+
+					ExtractHeightmapWeightmapContent(HeightData, ImportMaterialLayerInfos);
+
+					HeightDataPerLayers.Add(FGuid(), MoveTemp(HeightData));
+					ImportMaterialLayerInfosPerLayers.Add(FGuid(), MoveTemp(ImportMaterialLayerInfos));
+				}
+
+				if (!bResample)
+				{
 					NewMinX = 0;
 					NewMinY = 0;
 					NewMaxX = NewVertsX - 1;
@@ -3912,84 +3963,71 @@ ALandscape* FEdModeLandscape::ChangeComponentSetting(int32 NumComponentsX, int32
 
 			Progress.EnterProgressFrame(CurrentTaskProgress++);
 
-			const FVector Location = OldLandscapeProxy->GetActorLocation() + LandscapeOffset;
+			const FVector Location = OldLandscape->GetActorLocation() + LandscapeOffset;
 			FActorSpawnParameters SpawnParams;
-			SpawnParams.OverrideLevel = OldLandscapeProxy->GetLevel();
-			Landscape = OldLandscapeProxy->GetWorld()->SpawnActor<ALandscape>(Location, OldLandscapeProxy->GetActorRotation(), SpawnParams);
-
-			const FVector OldScale = OldLandscapeProxy->GetActorScale();
-			Landscape->SetActorRelativeScale3D(FVector(OldScale.X * LandscapeScaleFactor, OldScale.Y * LandscapeScaleFactor, OldScale.Z));
-
-			Landscape->LandscapeMaterial = OldLandscapeProxy->LandscapeMaterial;
-			Landscape->LandscapeMaterialsOverride = OldLandscapeProxy->LandscapeMaterialsOverride;
-			Landscape->CollisionMipLevel = OldLandscapeProxy->CollisionMipLevel;
-			Landscape->Import(FGuid::NewGuid(), NewMinX, NewMinY, NewMaxX, NewMaxY, NumSubsections, SubsectionSizeQuads, HeightData.GetData(), *OldLandscapeProxy->ReimportHeightmapFilePath, ImportLayerInfos, ELandscapeImportAlphamapType::Additive);
-
-			Landscape->MaxLODLevel = OldLandscapeProxy->MaxLODLevel;
-			Landscape->LODDistanceFactor_DEPRECATED = OldLandscapeProxy->LODDistanceFactor_DEPRECATED;
-			Landscape->LODFalloff_DEPRECATED = OldLandscapeProxy->LODFalloff_DEPRECATED;
-			Landscape->TessellationComponentScreenSize = OldLandscapeProxy->TessellationComponentScreenSize;
-			Landscape->ComponentScreenSizeToUseSubSections = OldLandscapeProxy->ComponentScreenSizeToUseSubSections;
-			Landscape->UseTessellationComponentScreenSizeFalloff = OldLandscapeProxy->UseTessellationComponentScreenSizeFalloff;
-			Landscape->TessellationComponentScreenSizeFalloff = OldLandscapeProxy->TessellationComponentScreenSizeFalloff;
-			Landscape->LODDistributionSetting = OldLandscapeProxy->LODDistributionSetting;
-			Landscape->LOD0DistributionSetting = OldLandscapeProxy->LOD0DistributionSetting;
-			Landscape->OccluderGeometryLOD = OldLandscapeProxy->OccluderGeometryLOD;
-			Landscape->ExportLOD = OldLandscapeProxy->ExportLOD;
-			Landscape->StaticLightingLOD = OldLandscapeProxy->StaticLightingLOD;
-			Landscape->NegativeZBoundsExtension = OldLandscapeProxy->NegativeZBoundsExtension;
-			Landscape->PositiveZBoundsExtension = OldLandscapeProxy->PositiveZBoundsExtension;
-			Landscape->DefaultPhysMaterial = OldLandscapeProxy->DefaultPhysMaterial;
-			Landscape->StreamingDistanceMultiplier = OldLandscapeProxy->StreamingDistanceMultiplier;
-			Landscape->LandscapeHoleMaterial = OldLandscapeProxy->LandscapeHoleMaterial;
-			Landscape->StaticLightingResolution = OldLandscapeProxy->StaticLightingResolution;
-			Landscape->bCastStaticShadow = OldLandscapeProxy->bCastStaticShadow;
-			Landscape->bCastShadowAsTwoSided = OldLandscapeProxy->bCastShadowAsTwoSided;
-			Landscape->LightingChannels = OldLandscapeProxy->LightingChannels;
-			Landscape->bRenderCustomDepth = OldLandscapeProxy->bRenderCustomDepth;
-			Landscape->CustomDepthStencilValue = OldLandscapeProxy->CustomDepthStencilValue;
-			Landscape->LightmassSettings = OldLandscapeProxy->LightmassSettings;
-			Landscape->CollisionThickness = OldLandscapeProxy->CollisionThickness;
-			Landscape->BodyInstance.SetCollisionProfileName(OldLandscapeProxy->BodyInstance.GetCollisionProfileName());
-			if (Landscape->BodyInstance.DoesUseCollisionProfile() == false)
+			SpawnParams.OverrideLevel = OldLandscape->GetLevel();
+			NewLandscape = OldLandscape->GetWorld()->SpawnActor<ALandscape>(Location, OldLandscape->GetActorRotation(), SpawnParams);
+
+			const FVector OldScale = OldLandscape->GetActorScale();
+			NewLandscape->SetActorRelativeScale3D(FVector(OldScale.X * LandscapeScaleFactor, OldScale.Y * LandscapeScaleFactor, OldScale.Z));
+
+			NewLandscape->LandscapeMaterial = OldLandscape->LandscapeMaterial;
+			NewLandscape->LandscapeMaterialsOverride = OldLandscape->LandscapeMaterialsOverride;
+			NewLandscape->CollisionMipLevel = OldLandscape->CollisionMipLevel;
+			NewLandscape->MaxLODLevel = OldLandscape->MaxLODLevel;
+			NewLandscape->LODDistanceFactor_DEPRECATED = OldLandscape->LODDistanceFactor_DEPRECATED;
+			NewLandscape->LODFalloff_DEPRECATED = OldLandscape->LODFalloff_DEPRECATED;
+			NewLandscape->TessellationComponentScreenSize = OldLandscape->TessellationComponentScreenSize;
+			NewLandscape->ComponentScreenSizeToUseSubSections = OldLandscape->ComponentScreenSizeToUseSubSections;
+			NewLandscape->UseTessellationComponentScreenSizeFalloff = OldLandscape->UseTessellationComponentScreenSizeFalloff;
+			NewLandscape->TessellationComponentScreenSizeFalloff = OldLandscape->TessellationComponentScreenSizeFalloff;
+			NewLandscape->LODDistributionSetting = OldLandscape->LODDistributionSetting;
+			NewLandscape->LOD0DistributionSetting = OldLandscape->LOD0DistributionSetting;
+			NewLandscape->OccluderGeometryLOD = OldLandscape->OccluderGeometryLOD;
+			NewLandscape->ExportLOD = OldLandscape->ExportLOD;
+			NewLandscape->StaticLightingLOD = OldLandscape->StaticLightingLOD;
+			NewLandscape->NegativeZBoundsExtension = OldLandscape->NegativeZBoundsExtension;
+			NewLandscape->PositiveZBoundsExtension = OldLandscape->PositiveZBoundsExtension;
+			NewLandscape->DefaultPhysMaterial = OldLandscape->DefaultPhysMaterial;
+			NewLandscape->StreamingDistanceMultiplier = OldLandscape->StreamingDistanceMultiplier;
+			NewLandscape->LandscapeHoleMaterial = OldLandscape->LandscapeHoleMaterial;
+			NewLandscape->StaticLightingResolution = OldLandscape->StaticLightingResolution;
+			NewLandscape->bCastStaticShadow = OldLandscape->bCastStaticShadow;
+			NewLandscape->bCastShadowAsTwoSided = OldLandscape->bCastShadowAsTwoSided;
+			NewLandscape->LightingChannels = OldLandscape->LightingChannels;
+			NewLandscape->bRenderCustomDepth = OldLandscape->bRenderCustomDepth;
+			NewLandscape->CustomDepthStencilValue = OldLandscape->CustomDepthStencilValue;
+			NewLandscape->LightmassSettings = OldLandscape->LightmassSettings;
+			NewLandscape->CollisionThickness = OldLandscape->CollisionThickness;
+			NewLandscape->BodyInstance.SetCollisionProfileName(OldLandscape->BodyInstance.GetCollisionProfileName());
+			if (NewLandscape->BodyInstance.DoesUseCollisionProfile() == false)
 			{
-				Landscape->BodyInstance.SetCollisionEnabled(OldLandscapeProxy->BodyInstance.GetCollisionEnabled());
-				Landscape->BodyInstance.SetObjectType(OldLandscapeProxy->BodyInstance.GetObjectType());
-				Landscape->BodyInstance.SetResponseToChannels(OldLandscapeProxy->BodyInstance.GetResponseToChannels());
+				NewLandscape->BodyInstance.SetCollisionEnabled(OldLandscape->BodyInstance.GetCollisionEnabled());
+				NewLandscape->BodyInstance.SetObjectType(OldLandscape->BodyInstance.GetObjectType());
+				NewLandscape->BodyInstance.SetResponseToChannels(OldLandscape->BodyInstance.GetResponseToChannels());
 			}
-			Landscape->EditorLayerSettings = OldLandscapeProxy->EditorLayerSettings;
-			Landscape->bUsedForNavigation = OldLandscapeProxy->bUsedForNavigation;
-			Landscape->MaxPaintedLayersPerComponent = OldLandscapeProxy->MaxPaintedLayersPerComponent;
+			NewLandscape->EditorLayerSettings = OldLandscape->EditorLayerSettings;
+			NewLandscape->bUsedForNavigation = OldLandscape->bUsedForNavigation;
+			NewLandscape->MaxPaintedLayersPerComponent = OldLandscape->MaxPaintedLayersPerComponent;
 
-			if (GetMutableDefault<UEditorExperimentalSettings>()->bLandscapeLayerSystem)
-			{
-				ALandscape* OldLandscapeActor = Cast<ALandscape>(OldLandscapeProxy);
-				ALandscape* NewLandscapeActor = Cast<ALandscape>(Landscape);
+			TArray<FLandscapeLayer>* LandscapeLayers = GetMutableDefault<UEditorExperimentalSettings>()->bLandscapeLayerSystem ? &OldLandscape->LandscapeLayers : nullptr;
 
-				NewLandscapeActor->bInitializedWithFlagExperimentalLandscapeLayers = OldLandscapeActor->bInitializedWithFlagExperimentalLandscapeLayers;
+			NewLandscape->Import(FGuid::NewGuid(), NewMinX, NewMinY, NewMaxX, NewMaxY, NumSubsections, SubsectionSizeQuads, HeightDataPerLayers, *OldLandscape->ReimportHeightmapFilePath, ImportMaterialLayerInfosPerLayers, ELandscapeImportAlphamapType::Additive, LandscapeLayers);
 
-				// Moves specific config including brushes to new landscape actor
-				for (const FLandscapeLayer& Layer : OldLandscapeActor->LandscapeLayers)
-				{
-					FLandscapeLayer* NewLayer = NewLandscapeActor->LandscapeLayers.FindByPredicate( [&](const FLandscapeLayer& NewLayer)	{ return NewLayer.Name == Layer.Name; }	);
-					check(NewLayer != nullptr);
-					*NewLayer = Layer;
-				}
-			}
-
-			Landscape->CreateLandscapeInfo();
+			ULandscapeInfo* NewLandscapeInfo = NewLandscape->GetLandscapeInfo();
+			check(NewLandscapeInfo);
 
 			// Clone landscape splines
-			TLazyObjectPtr<ALandscape> OldLandscapeActor = LandscapeInfo->LandscapeActor;
-			if (OldLandscapeActor.IsValid() && OldLandscapeActor->SplineComponent != NULL)
+			ALandscape* OldLandscapeActor = LandscapeInfo->LandscapeActor.Get();
+			if (OldLandscapeActor != nullptr && OldLandscapeActor->SplineComponent != NULL)
 			{
 				ULandscapeSplinesComponent* OldSplines = OldLandscapeActor->SplineComponent;
-				ULandscapeSplinesComponent* NewSplines = DuplicateObject<ULandscapeSplinesComponent>(OldSplines, Landscape, OldSplines->GetFName());
-				NewSplines->AttachToComponent(Landscape->GetRootComponent(), FAttachmentTransformRules::KeepWorldTransform);
+				ULandscapeSplinesComponent* NewSplines = DuplicateObject<ULandscapeSplinesComponent>(OldSplines, NewLandscape, OldSplines->GetFName());
+				NewSplines->AttachToComponent(NewLandscape->GetRootComponent(), FAttachmentTransformRules::KeepWorldTransform);
 
 				const FVector OldSplineScale = OldSplines->GetRelativeTransform().GetScale3D();
 				NewSplines->SetRelativeScale3D(FVector(OldSplineScale.X / LandscapeScaleFactor, OldSplineScale.Y / LandscapeScaleFactor, OldSplineScale.Z));
-				Landscape->SplineComponent = NewSplines;
+				NewLandscape->SplineComponent = NewSplines;
 				NewSplines->RegisterComponent();
 
 				// TODO: Foliage on spline meshes
@@ -4000,7 +4038,6 @@ ALandscape* FEdModeLandscape::ChangeComponentSetting(int32 NumComponentsX, int32
 			if (bResample)
 			{
 				// Remap foliage to the resampled components
-				ULandscapeInfo* NewLandscapeInfo = Landscape->GetLandscapeInfo();
 				for (const TPair<FIntPoint, ULandscapeComponent*>& Entry : LandscapeInfo->XYtoComponentMap)
 				{
 					ULandscapeComponent* NewComponent = NewLandscapeInfo->XYtoComponentMap.FindRef(Entry.Key);
@@ -4035,8 +4072,6 @@ ALandscape* FEdModeLandscape::ChangeComponentSetting(int32 NumComponentsX, int32
 			}
 			else
 			{
-				ULandscapeInfo* NewLandscapeInfo = Landscape->GetLandscapeInfo();
-
 				// Move instances
 				for (const TPair<FIntPoint, ULandscapeComponent*>& OldEntry : LandscapeInfo->XYtoComponentMap)
 				{
@@ -4099,20 +4134,20 @@ ALandscape* FEdModeLandscape::ChangeComponentSetting(int32 NumComponentsX, int32
 			}
 
 			// Delete the old Landscape and all its proxies
-			for (ALandscapeStreamingProxy* Proxy : TActorRange<ALandscapeStreamingProxy>(OldLandscapeProxy->GetWorld()))
+			for (ALandscapeStreamingProxy* Proxy : TActorRange<ALandscapeStreamingProxy>(OldLandscape->GetWorld()))
 			{
 				if (Proxy->LandscapeActor == OldLandscapeActor)
 				{
 					Proxy->Destroy();
 				}
 			}
-			OldLandscapeProxy->Destroy();
+			OldLandscape->Destroy();
 		}
 	}
 
 	GEditor->RedrawLevelEditingViewports();
 
-	return Landscape;
+	return NewLandscape;
 }
 
 ELandscapeEditingState FEdModeLandscape::GetEditingState() const
@@ -4147,7 +4182,7 @@ ELandscapeEditingState FEdModeLandscape::GetEditingState() const
 int32 FEdModeLandscape::GetLayerCount() const
 {
 	ALandscape* Landscape = GetLandscape();
-	return Landscape ? Landscape->LandscapeLayers.Num() : 0;
+	return Landscape ? Landscape->GetLayerCount() : 0;
 }
 
 void FEdModeLandscape::SetCurrentLayer(int32 InLayerIndex)
diff --git a/Engine/Source/Editor/LandscapeEditor/Private/LandscapeEditorDetailCustomization_NewLandscape.cpp b/Engine/Source/Editor/LandscapeEditor/Private/LandscapeEditorDetailCustomization_NewLandscape.cpp
--- a/Engine/Source/Editor/LandscapeEditor/Private/LandscapeEditorDetailCustomization_NewLandscape.cpp
+++ b/Engine/Source/Editor/LandscapeEditor/Private/LandscapeEditorDetailCustomization_NewLandscape.cpp
@@ -785,37 +785,33 @@ FReply FLandscapeEditorDetailCustomization_NewLandscape::OnCreateButtonClicked()
 		LandscapeEdMode->GetWorld() != nullptr && 
 		LandscapeEdMode->GetWorld()->GetCurrentLevel()->bIsVisible)
 	{
-		const int32 ComponentCountX = LandscapeEdMode->UISettings->NewLandscape_ComponentCount.X;
-		const int32 ComponentCountY = LandscapeEdMode->UISettings->NewLandscape_ComponentCount.Y;
-		const int32 QuadsPerComponent = LandscapeEdMode->UISettings->NewLandscape_SectionsPerComponent * LandscapeEdMode->UISettings->NewLandscape_QuadsPerSection;
+		ULandscapeEditorObject* UISettings = LandscapeEdMode->UISettings;
+		const int32 ComponentCountX = UISettings->NewLandscape_ComponentCount.X;
+		const int32 ComponentCountY = UISettings->NewLandscape_ComponentCount.Y;
+		const int32 QuadsPerComponent = UISettings->NewLandscape_SectionsPerComponent * UISettings->NewLandscape_QuadsPerSection;
 		const int32 SizeX = ComponentCountX * QuadsPerComponent + 1;
 		const int32 SizeY = ComponentCountY * QuadsPerComponent + 1;
 
-		TOptional< TArray< FLandscapeImportLayerInfo > > ImportLayers = FNewLandscapeUtils::CreateImportLayersInfo( LandscapeEdMode->UISettings, LandscapeEdMode->NewLandscapePreviewMode );
+		TOptional< TArray< FLandscapeImportLayerInfo > > MaterialImportLayers = FNewLandscapeUtils::CreateImportLayersInfo(UISettings, LandscapeEdMode->NewLandscapePreviewMode );
 
-		if ( !ImportLayers )
+		if ( !MaterialImportLayers)
 		{
 			return FReply::Handled();
 		}
 
-		TArray<uint16> Data = FNewLandscapeUtils::ComputeHeightData( LandscapeEdMode->UISettings, ImportLayers.GetValue(), LandscapeEdMode->NewLandscapePreviewMode );
+		TMap<FGuid, TArray<uint16>> HeightDataPerLayers;
+		TMap<FGuid, TArray<FLandscapeImportLayerInfo>> MaterialLayerDataPerLayers;
+
+		MaterialLayerDataPerLayers.Add(FGuid(), MaterialImportLayers.GetValue());
+		HeightDataPerLayers.Add(FGuid(), FNewLandscapeUtils::ComputeHeightData(UISettings, MaterialImportLayers.GetValue(), LandscapeEdMode->NewLandscapePreviewMode));
 
 		FScopedTransaction Transaction(LOCTEXT("Undo", "Creating New Landscape"));
 
-		const FVector Offset = FTransform(LandscapeEdMode->UISettings->NewLandscape_Rotation, FVector::ZeroVector,
-			LandscapeEdMode->UISettings->NewLandscape_Scale).TransformVector(FVector(-ComponentCountX * QuadsPerComponent / 2, -ComponentCountY * QuadsPerComponent / 2, 0));
+		const FVector Offset = FTransform(UISettings->NewLandscape_Rotation, FVector::ZeroVector, UISettings->NewLandscape_Scale).TransformVector(FVector(-ComponentCountX * QuadsPerComponent / 2, -ComponentCountY * QuadsPerComponent / 2, 0));
 		
-		ALandscape* Landscape = LandscapeEdMode->GetWorld()->SpawnActor<ALandscape>(LandscapeEdMode->UISettings->NewLandscape_Location + Offset, LandscapeEdMode->UISettings->NewLandscape_Rotation);
-		Landscape->LandscapeMaterial = LandscapeEdMode->UISettings->NewLandscape_Material.Get();
-		Landscape->SetActorRelativeScale3D(LandscapeEdMode->UISettings->NewLandscape_Scale);
-
-		if (GetMutableDefault<UEditorExperimentalSettings>()->bLandscapeLayerSystem)
-		{
-			Landscape->bInitializedWithFlagExperimentalLandscapeLayers = true;
-		}
-
-		Landscape->Import(FGuid::NewGuid(), 0, 0, SizeX-1, SizeY-1, LandscapeEdMode->UISettings->NewLandscape_SectionsPerComponent, LandscapeEdMode->UISettings->NewLandscape_QuadsPerSection, Data.GetData(),
-			nullptr, ImportLayers.GetValue(), LandscapeEdMode->UISettings->ImportLandscape_AlphamapType);
+		ALandscape* Landscape = LandscapeEdMode->GetWorld()->SpawnActor<ALandscape>(UISettings->NewLandscape_Location + Offset, UISettings->NewLandscape_Rotation);
+		Landscape->LandscapeMaterial = UISettings->NewLandscape_Material.Get();
+		Landscape->SetActorRelativeScale3D(UISettings->NewLandscape_Scale);
 
 		// automatically calculate a lighting LOD that won't crash lightmass (hopefully)
 		// < 2048x2048 -> LOD0
@@ -826,14 +822,18 @@ FReply FLandscapeEditorDetailCustomization_NewLandscape::OnCreateButtonClicked()
 
 		if (LandscapeEdMode->NewLandscapePreviewMode == ENewLandscapePreviewMode::ImportLandscape)
 		{
-			Landscape->ReimportHeightmapFilePath = LandscapeEdMode->UISettings->ImportLandscape_HeightmapFilename;
+			Landscape->ReimportHeightmapFilePath = UISettings->ImportLandscape_HeightmapFilename;
 		}
 
-		ULandscapeInfo* LandscapeInfo = Landscape->CreateLandscapeInfo();
+		Landscape->Import(FGuid::NewGuid(), 0, 0, SizeX - 1, SizeY - 1, UISettings->NewLandscape_SectionsPerComponent, UISettings->NewLandscape_QuadsPerSection, HeightDataPerLayers, nullptr, MaterialLayerDataPerLayers, UISettings->ImportLandscape_AlphamapType);
+
+		ULandscapeInfo* LandscapeInfo = Landscape->GetLandscapeInfo();
+		check(LandscapeInfo);
+
 		LandscapeInfo->UpdateLayerInfoMap(Landscape);
 
 		// Import doesn't fill in the LayerInfo for layers with no data, do that now
-		const TArray<FLandscapeImportLayer>& ImportLandscapeLayersList = LandscapeEdMode->UISettings->ImportLandscape_Layers;
+		const TArray<FLandscapeImportLayer>& ImportLandscapeLayersList = UISettings->ImportLandscape_Layers;
 		for (int32 i = 0; i < ImportLandscapeLayersList.Num(); i++)
 		{
 			if (ImportLandscapeLayersList[i].LayerInfo != nullptr)
diff --git a/Engine/Source/Editor/WorldBrowser/Private/Tiles/WorldTileModel.cpp b/Engine/Source/Editor/WorldBrowser/Private/Tiles/WorldTileModel.cpp
--- a/Engine/Source/Editor/WorldBrowser/Private/Tiles/WorldTileModel.cpp
+++ b/Engine/Source/Editor/WorldBrowser/Private/Tiles/WorldTileModel.cpp
@@ -17,6 +17,7 @@
 #include "GameFramework/WorldSettings.h"
 #include "LandscapeInfo.h"
 #include "LandscapeStreamingProxy.h"
+#include "Landscape.h"
 
 
 #define LOCTEXT_NAMESPACE "WorldBrowser"
@@ -996,18 +997,15 @@ ALandscapeProxy* FWorldTileModel::ImportLandscapeTile(const FLandscapeImportSett
 	// Cache pointer to landscape in the level model
 	Landscape = LandscapeProxy;
 
+	TMap<FGuid, TArray<uint16>> HeightmapDataPerLayers;
+	TMap<FGuid, TArray<FLandscapeImportLayerInfo>> MaterialLayerDataPerLayer;
+
+	HeightmapDataPerLayers.Add(FGuid(), Settings.HeightData);
+	MaterialLayerDataPerLayer.Add(FGuid(), Settings.ImportLayers);	
+
 	// Create landscape components
-	LandscapeProxy->Import(
-		Settings.LandscapeGuid,
-		0, 0,
-		Settings.SizeX - 1,
-		Settings.SizeY - 1,
-		Settings.SectionsPerComponent,
-		Settings.QuadsPerSection,
-		Settings.HeightData.GetData(),
-		*Settings.HeightmapFilename,
-		Settings.ImportLayers,
-		Settings.ImportLayerType);
+	LandscapeProxy->Import(	Settings.LandscapeGuid, 0, 0, Settings.SizeX - 1, Settings.SizeY - 1, Settings.SectionsPerComponent, Settings.QuadsPerSection, HeightmapDataPerLayers, *Settings.HeightmapFilename,	
+							MaterialLayerDataPerLayer,	Settings.ImportLayerType);
 
 	return LandscapeProxy;
 }
diff --git a/Engine/Source/Runtime/Landscape/Classes/Landscape.h b/Engine/Source/Runtime/Landscape/Classes/Landscape.h
--- a/Engine/Source/Runtime/Landscape/Classes/Landscape.h
+++ b/Engine/Source/Runtime/Landscape/Classes/Landscape.h
@@ -145,6 +145,8 @@ struct FLandscapeLayer
 		, BlendMode(LSBM_AdditiveBlend)
 	{}
 
+	FLandscapeLayer(const FLandscapeLayer& OtherLayer) = default;
+
 	UPROPERTY(meta = (IgnoreForMemberInitializationTest))
 	FGuid Guid;
 
@@ -230,6 +232,7 @@ public:
 	LANDSCAPE_API void RequestLayersContentUpdateForceAll();
 	LANDSCAPE_API void RequestLayersContentUpdate(ELandscapeLayerUpdateMode InModeMask, bool bInForceUpdateAllComponents = false);
 	LANDSCAPE_API bool ReorderLayer(int32 InStartingLayerIndex, int32 InDestinationLayerIndex);
+	LANDSCAPE_API FLandscapeLayer* DuplicateLayer(const FLandscapeLayer& InOtherLayer);
 	LANDSCAPE_API void CreateLayer(FName InName = NAME_None);
 	LANDSCAPE_API void CreateDefaultLayer();
 	LANDSCAPE_API void CopyOldDataToDefaultLayer();
@@ -241,6 +244,7 @@ public:
 	LANDSCAPE_API void SetLayerAlpha(int32 InLayerIndex, const float InAlpha, bool bInHeightmap);
 	LANDSCAPE_API void SetLayerVisibility(int32 InLayerIndex, bool bInVisible);
 	LANDSCAPE_API void SetLayerLocked(int32 InLayerIndex, bool bLocked);
+	LANDSCAPE_API uint8 GetLayerCount() const;
 	LANDSCAPE_API struct FLandscapeLayer* GetLayer(int32 InLayerIndex);
 	LANDSCAPE_API const struct FLandscapeLayer* GetLayer(int32 InLayerIndex) const;
 	LANDSCAPE_API const struct FLandscapeLayer* GetLayer(const FGuid& InLayerGuid) const;
diff --git a/Engine/Source/Runtime/Landscape/Classes/LandscapeProxy.h b/Engine/Source/Runtime/Landscape/Classes/LandscapeProxy.h
--- a/Engine/Source/Runtime/Landscape/Classes/LandscapeProxy.h
+++ b/Engine/Source/Runtime/Landscape/Classes/LandscapeProxy.h
@@ -849,9 +849,8 @@ public:
 	LANDSCAPE_API static ULandscapeMaterialInstanceConstant* GetLayerThumbnailMIC(UMaterialInterface* LandscapeMaterial, FName LayerName, UTexture2D* ThumbnailWeightmap, UTexture2D* ThumbnailHeightmap, ALandscapeProxy* Proxy);
 
 	/** Import the given Height/Weight data into this landscape */
-	LANDSCAPE_API void Import(FGuid Guid, int32 MinX, int32 MinY, int32 MaxX, int32 MaxY, int32 NumSubsections, int32 SubsectionSizeQuads,
-							const uint16* HeightData, const TCHAR* HeightmapFileName,
-							const TArray<FLandscapeImportLayerInfo>& ImportLayerInfos, ELandscapeImportAlphamapType ImportLayerType);
+	LANDSCAPE_API void Import(const FGuid& InGuid, int32 InMinX, int32 InMinY, int32 InMaxX, int32 InMaxY, int32 InNumSubsections, int32 InSubsectionSizeQuads, const TMap<FGuid, TArray<uint16>>& InImportHeightData,
+							  const TCHAR* const InHeightmapFileName, const TMap<FGuid, TArray<FLandscapeImportLayerInfo>>& InImportMaterialLayerInfos, ELandscapeImportAlphamapType InImportMaterialLayerType, const TArray<struct FLandscapeLayer>* InImportLayers = nullptr);
 
 	/**
 	 * Exports landscape into raw mesh
diff --git a/Engine/Source/Runtime/Landscape/Private/Landscape.cpp b/Engine/Source/Runtime/Landscape/Private/Landscape.cpp
--- a/Engine/Source/Runtime/Landscape/Private/Landscape.cpp
+++ b/Engine/Source/Runtime/Landscape/Private/Landscape.cpp
@@ -1075,7 +1075,7 @@ ALandscape::ALandscape(const FObjectInitializer& ObjectInitializer)
 {
 #if WITH_EDITORONLY_DATA
 	bLockLocation = false;
-	bInitializedWithFlagExperimentalLandscapeLayers = false;
+	bInitializedWithFlagExperimentalLandscapeLayers = GetMutableDefault<UEditorExperimentalSettings>()->bLandscapeLayerSystem;
 	WasCompilingShaders = false;
 	LayerContentUpdateModes = 0;
 	bLayerForceUpdateAllComponents = false;
@@ -2749,7 +2749,7 @@ void ULandscapeInfo::RegisterActor(ALandscapeProxy* Proxy, bool bMapCheck)
 		StreamingProxy->FixupSharedData(LandscapeActor.Get());
 	}
 
-	if (LandscapeActor)
+	if (GetMutableDefault<UEditorExperimentalSettings>()->bLandscapeLayerSystem && LandscapeActor)
 	{
 		// Force update rendering resources
 		LandscapeActor->RequestLayersInitialization();
diff --git a/Engine/Source/Runtime/Landscape/Private/LandscapeEdit.cpp b/Engine/Source/Runtime/Landscape/Private/LandscapeEdit.cpp
--- a/Engine/Source/Runtime/Landscape/Private/LandscapeEdit.cpp
+++ b/Engine/Source/Runtime/Landscape/Private/LandscapeEdit.cpp
@@ -2136,28 +2136,25 @@ ULandscapeLayerInfoObject* ALandscapeProxy::CreateLayerInfo(const TCHAR* LayerNa
 #define HEIGHTDATA(X,Y) (HeightData[ FMath::Clamp<int32>(Y,0,VertsY) * VertsX + FMath::Clamp<int32>(X,0,VertsX) ])
 ENGINE_API extern bool GDisableAutomaticTextureMaterialUpdateDependencies;
 
-LANDSCAPE_API void ALandscapeProxy::Import(
-	const FGuid Guid,
-	const int32 MinX, const int32 MinY, const int32 MaxX, const int32 MaxY,
-	const int32 InNumSubsections, const int32 InSubsectionSizeQuads,
-	const uint16* const HeightData, const TCHAR* const HeightmapFileName,
-	const TArray<FLandscapeImportLayerInfo>& ImportLayerInfos, const ELandscapeImportAlphamapType ImportLayerType)
+LANDSCAPE_API void ALandscapeProxy::Import(const FGuid& InGuid, int32 InMinX, int32 InMinY, int32 InMaxX, int32 InMaxY, int32 InNumSubsections, int32 InSubsectionSizeQuads, const TMap<FGuid, TArray<uint16>>& InImportHeightData, 
+										   const TCHAR* const InHeightmapFileName, const TMap<FGuid, TArray<FLandscapeImportLayerInfo>>& InImportMaterialLayerInfos, ELandscapeImportAlphamapType InImportMaterialLayerType, const TArray<FLandscapeLayer>* InImportLayers)
 {
+	check(InGuid.IsValid());
+	check(InImportHeightData.Num() == InImportMaterialLayerInfos.Num());
+	check(GetMutableDefault<UEditorExperimentalSettings>()->bLandscapeLayerSystem || InImportLayers == nullptr);
+
 	GWarn->BeginSlowTask(LOCTEXT("BeingImportingLandscapeTask", "Importing Landscape"), true);
 
-	const int32 VertsX = MaxX - MinX + 1;
-	const int32 VertsY = MaxY - MinY + 1;
+	const int32 VertsX = InMaxX - InMinX + 1;
+	const int32 VertsY = InMaxY - InMinY + 1;
 
 	ComponentSizeQuads = InNumSubsections * InSubsectionSizeQuads;
 	NumSubsections = InNumSubsections;
 	SubsectionSizeQuads = InSubsectionSizeQuads;
-	LandscapeGuid = Guid;
+	LandscapeGuid = InGuid;
 
 	Modify();
 
-	// Create and initialize landscape info object
-	ULandscapeInfo* LandscapeInfo = CreateLandscapeInfo();
-
 	const int32 NumPatchesX = (VertsX - 1);
 	const int32 NumPatchesY = (VertsY - 1);
 
@@ -2172,19 +2169,14 @@ LANDSCAPE_API void ALandscapeProxy::Import(
 	{
 		for (int32 X = 0; X < NumComponentsX; X++)
 		{
-			const int32 BaseX = MinX + X * ComponentSizeQuads;
-			const int32 BaseY = MinY + Y * ComponentSizeQuads;
+			const int32 BaseX = InMinX + X * ComponentSizeQuads;
+			const int32 BaseY = InMinY + Y * ComponentSizeQuads;
 
 			ULandscapeComponent* LandscapeComponent = NewObject<ULandscapeComponent>(this, NAME_None, RF_Transactional);
 			LandscapeComponent->SetRelativeLocation(FVector(BaseX, BaseY, 0));
 			LandscapeComponent->SetupAttachment(GetRootComponent(), NAME_None);
 			LandscapeComponents.Add(LandscapeComponent);
-			LandscapeComponent->Init(
-				BaseX, BaseY,
-				ComponentSizeQuads,
-				NumSubsections,
-				SubsectionSizeQuads
-				);
+			LandscapeComponent->Init(BaseX, BaseY, ComponentSizeQuads, NumSubsections, SubsectionSizeQuads);
 
 			// Assign shared properties
 			LandscapeComponent->UpdatedSharedPropertiesFromActor();
@@ -2249,6 +2241,11 @@ LANDSCAPE_API void ALandscapeProxy::Import(
 
 	const FVector DrawScale3D = GetRootComponent()->RelativeScale3D;
 
+	// layer to import data (Final or 1st layer)
+	const FGuid FinalLayerGuid = FGuid();
+	const TArray<uint16>& HeightData = InImportHeightData.FindChecked(FinalLayerGuid);
+	const TArray<FLandscapeImportLayerInfo>& ImportLayerInfos = InImportMaterialLayerInfos.FindChecked(FinalLayerGuid);
+
 	// Calculate the normals for each of the two triangles per quad.
 	TArray<FVector> VertexNormals;
 	VertexNormals.AddZeroed(VertsX * VertsY);
@@ -2293,7 +2290,7 @@ LANDSCAPE_API void ALandscapeProxy::Import(
 				{
 					for (int32 AlphaY = 0; AlphaY <= LandscapeComponent->ComponentSizeQuads; AlphaY++)
 					{
-						const uint8* const OldAlphaRowStart = &ImportLayerInfos[LayerIndex].LayerData[(AlphaY + LandscapeComponent->GetSectionBase().Y - MinY) * VertsX + (LandscapeComponent->GetSectionBase().X - MinX)];
+						const uint8* const OldAlphaRowStart = &ImportLayerInfos[LayerIndex].LayerData[(AlphaY + LandscapeComponent->GetSectionBase().Y - InMinY) * VertsX + (LandscapeComponent->GetSectionBase().X - InMinX)];
 						uint8* const NewAlphaRowStart = &NewAlphaInfo->AlphaValues[AlphaY * (LandscapeComponent->ComponentSizeQuads + 1)];
 						FMemory::Memcpy(NewAlphaRowStart, OldAlphaRowStart, LandscapeComponent->ComponentSizeQuads + 1);
 					}
@@ -2334,7 +2331,7 @@ LANDSCAPE_API void ALandscapeProxy::Import(
 			// Discard the temporary alpha data
 			EditingAlphaLayerData.Empty();
 
-			if (ImportLayerType == ELandscapeImportAlphamapType::Layered)
+			if (InImportMaterialLayerType == ELandscapeImportAlphamapType::Layered)
 			{
 				// For each layer...
 				for (int32 WeightLayerIndex = WeightValues.Num() - 1; WeightLayerIndex >= 0; WeightLayerIndex--)
@@ -2593,8 +2590,8 @@ LANDSCAPE_API void ALandscapeProxy::Import(
 							const int32 WeightTexDataIdx = (TexX)+(TexY)* (WeightmapSize);
 
 							// copy height and normal data
-							const uint16 HeightValue = HEIGHTDATA(CompX + LandscapeComponent->GetSectionBase().X - MinX, CompY + LandscapeComponent->GetSectionBase().Y - MinY);
-							const FVector Normal = VertexNormals[CompX + LandscapeComponent->GetSectionBase().X - MinX + VertsX * (CompY + LandscapeComponent->GetSectionBase().Y - MinY)].GetSafeNormal();
+							const uint16 HeightValue = HEIGHTDATA(CompX + LandscapeComponent->GetSectionBase().X - InMinX, CompY + LandscapeComponent->GetSectionBase().Y - InMinY);
+							const FVector Normal = VertexNormals[CompX + LandscapeComponent->GetSectionBase().X - InMinX + VertsX * (CompY + LandscapeComponent->GetSectionBase().Y - InMinY)].GetSafeNormal();
 
 							HeightmapInfo.HeightmapTextureMipData[0][HeightTexDataIdx].R = HeightValue >> 8;
 							HeightmapInfo.HeightmapTextureMipData[0][HeightTexDataIdx].G = HeightValue & 255;
@@ -2761,33 +2758,135 @@ LANDSCAPE_API void ALandscapeProxy::Import(
 	// Must be after the FMaterialUpdateContext is destroyed
 	RecreateRenderStateContexts.Reset();
 
+	// Create and initialize landscape info object
+	ULandscapeInfo* LandscapeInfo = CreateLandscapeInfo();
+
 	if (GetMutableDefault<UEditorExperimentalSettings>()->bLandscapeLayerSystem)
 	{
+		// Components need to be registered to be able to import the layer content and we will remove them if they should have not been visible
+		bool ShouldComponentBeRegistered = GetLevel()->bIsVisible;
+		RegisterAllComponents();
+
 		ALandscape* LandscapeActor = GetLandscapeActor();
 		check(LandscapeActor != nullptr);
-		// Only create Layers on main Landscape
-		if (LandscapeActor == this)
+
+		// Create the default layer
+		if (LandscapeActor->GetLayerCount() == 0 && InImportLayers == nullptr)
 		{
-			// Create the default layer
 			LandscapeActor->CreateDefaultLayer();
-		} 
+		}
+
+		TSet<ULandscapeComponent*> ComponentsToProcess;
+
+		struct FLayerImportSettings
+		{
+			FGuid SourceLayerGuid;
+			FGuid DestinationLayerGuid;
+		};
+
+		TArray<FLayerImportSettings> LayerImportSettings;		
+
+		// Only create Layers on main Landscape
+		if (LandscapeActor == this && InImportLayers != nullptr)
+		{
+			for (const FLandscapeLayer& OldLayer : *InImportLayers)
+			{
+				FLandscapeLayer* NewLayer = LandscapeActor->DuplicateLayer(OldLayer);
+				check(NewLayer != nullptr);
+
+				FLayerImportSettings ImportSettings;
+				ImportSettings.SourceLayerGuid = OldLayer.Guid;
+				ImportSettings.DestinationLayerGuid = NewLayer->Guid;
+				LayerImportSettings.Add(ImportSettings);
+			}
+
+			LandscapeInfo->GetComponentsInRegion(InMinX, InMinY, InMaxX, InMaxY, ComponentsToProcess);
+		}
 		else
 		{
-			LandscapeActor->AddLayersToProxy(this);
+			// In the case of a streaming proxy, we will generate the layer data for each components that the proxy hold so no need of the grid min/max to calculate the components to update
+			if (LandscapeActor != this)
+			{
+				LandscapeActor->AddLayersToProxy(this);
+			}
+
+			// And we will fill all the landscape components with the provided final layer content put into the default layer (aka layer index 0)
+			const FLandscapeLayer* DefaultLayer = LandscapeActor->GetLayer(0);
+			check(DefaultLayer != nullptr);
+
+			FLayerImportSettings ImportSettings;
+			ImportSettings.SourceLayerGuid = FinalLayerGuid;
+			ImportSettings.DestinationLayerGuid = DefaultLayer->Guid;
+			LayerImportSettings.Add(ImportSettings);
+
+			ComponentsToProcess.Append(LandscapeComponents);
 		}
-	}
 
-	if (GetLevel()->bIsVisible)
+		TSet<UTexture2D*> LayersTextures;
+
+		for (const FLayerImportSettings& ImportSettings : LayerImportSettings)
+		{
+			FLandscapeEditDataInterface LandscapeEdit(LandscapeInfo, false);
+			FScopedSetLandscapeEditingLayer Scope(LandscapeActor, ImportSettings.DestinationLayerGuid);
+
+			const TArray<uint16>* ImportHeightData = InImportHeightData.Find(ImportSettings.SourceLayerGuid);
+
+			if (ImportHeightData != nullptr)
+			{
+				LandscapeEdit.SetHeightData(InMinX, InMinY, InMaxX, InMaxY, (uint16*)ImportHeightData->GetData(), 0, false, nullptr);
+			}
+
+			const TArray<FLandscapeImportLayerInfo>* ImportWeightData = InImportMaterialLayerInfos.Find(ImportSettings.SourceLayerGuid);
+
+			if (ImportWeightData != nullptr)
+			{
+				for (const FLandscapeImportLayerInfo& MaterialLayerInfo : *ImportWeightData)
+				{
+					if (MaterialLayerInfo.LayerInfo != nullptr)
+					{
+						LandscapeEdit.SetAlphaData(MaterialLayerInfo.LayerInfo, InMinX, InMinY, InMaxX, InMaxY, MaterialLayerInfo.LayerData.GetData(), 0, ELandscapeLayerPaintingRestriction::None, true, false);
+					}
+				}
+			}
+
+			for (ULandscapeComponent* Component : ComponentsToProcess)
+			{
+				FLandscapeLayerComponentData* ComponentLayerData = Component->GetLayerData(ImportSettings.DestinationLayerGuid);
+				check(ComponentLayerData != nullptr);
+
+				LayersTextures.Add(ComponentLayerData->HeightmapData.Texture);
+				LayersTextures.Append(ComponentLayerData->WeightmapData.Textures);
+			}
+		}
+
+		// Retrigger a caching of the platform data as we wrote again in the textures
+		for (UTexture2D* Texture : LayersTextures)
+		{
+			Texture->ClearAllCachedCookedPlatformData();
+			Texture->BeginCachePlatformData();
+		}
+
+		LandscapeActor->RequestLayersContentUpdateForceAll();
+
+		if (!ShouldComponentBeRegistered)
+		{
+			UnregisterAllComponents();
+		}
+	}	
+	else
 	{
-		// Update our new components
-		ReregisterAllComponents();
+		if (GetLevel()->bIsVisible)
+		{
+			ReregisterAllComponents();
+		}
+
+		LandscapeInfo->RecreateCollisionComponents();
+		LandscapeInfo->UpdateAllAddCollisions();
 	}
 
-	ReimportHeightmapFilePath = HeightmapFileName;
+	ReimportHeightmapFilePath = InHeightmapFileName;
 
 	LandscapeInfo->UpdateLayerInfoMap();
-	LandscapeInfo->RecreateCollisionComponents();
-	LandscapeInfo->UpdateAllAddCollisions();
 
 	GWarn->EndSlowTask();
 }
diff --git a/Engine/Source/Runtime/Landscape/Private/LandscapeEditInterface.cpp b/Engine/Source/Runtime/Landscape/Private/LandscapeEditInterface.cpp
--- a/Engine/Source/Runtime/Landscape/Private/LandscapeEditInterface.cpp
+++ b/Engine/Source/Runtime/Landscape/Private/LandscapeEditInterface.cpp
@@ -30,7 +30,8 @@ extern const size_t ChannelOffsets[4] = {STRUCT_OFFSET(FColor,R), STRUCT_OFFSET(
 //
 // FLandscapeEditDataInterface
 //
-FLandscapeEditDataInterface::FLandscapeEditDataInterface(ULandscapeInfo* InLandscapeInfo)
+FLandscapeEditDataInterface::FLandscapeEditDataInterface(ULandscapeInfo* InLandscapeInfo, bool bInUploadTextureChangesToGPU)
+	: FLandscapeTextureDataInterface(bInUploadTextureChangesToGPU)
 {
 	if (InLandscapeInfo)
 	{
@@ -42,6 +43,10 @@ FLandscapeEditDataInterface::FLandscapeEditDataInterface(ULandscapeInfo* InLands
 	}
 }
 
+FLandscapeTextureDataInterface::FLandscapeTextureDataInterface(bool bInUploadTextureChangesToGPU)
+	: bUploadTextureChangesToGPU(bInUploadTextureChangesToGPU)
+{}
+
 FLandscapeTextureDataInterface::~FLandscapeTextureDataInterface()
 {
 	Flush();
@@ -51,12 +56,15 @@ void FLandscapeTextureDataInterface::Flush()
 {
 	bool bNeedToWaitForUpdate = false;
 
-	// Update all textures
-	for( TMap<UTexture2D*, FLandscapeTextureDataInfo*>::TIterator It(TextureDataMap); It;  ++It )
+	if (bUploadTextureChangesToGPU)
 	{
-		if( It.Value()->UpdateTextureData() )
+		// Update all textures
+		for (TMap<UTexture2D*, FLandscapeTextureDataInfo*>::TIterator It(TextureDataMap); It; ++It)
 		{
-			bNeedToWaitForUpdate = true;
+			if (It.Value()->UpdateTextureData())
+			{
+				bNeedToWaitForUpdate = true;
+			}
 		}
 	}
 
diff --git a/Engine/Source/Runtime/Landscape/Private/LandscapeEditLayers.cpp b/Engine/Source/Runtime/Landscape/Private/LandscapeEditLayers.cpp
--- a/Engine/Source/Runtime/Landscape/Private/LandscapeEditLayers.cpp
+++ b/Engine/Source/Runtime/Landscape/Private/LandscapeEditLayers.cpp
@@ -4541,6 +4541,11 @@ void ALandscape::SetLayerLocked(int32 InLayerIndex, bool bLocked)
 	Layer->bLocked = bLocked;
 }
 
+uint8 ALandscape::GetLayerCount() const
+{
+	return LandscapeLayers.Num();
+}
+
 FLandscapeLayer* ALandscape::GetLayer(int32 InLayerIndex)
 {
 	if (LandscapeLayers.IsValidIndex(InLayerIndex))
@@ -4876,6 +4881,41 @@ void ALandscape::CreateDefaultLayer()
 	RequestLayersInitialization();
 }
 
+FLandscapeLayer* ALandscape::DuplicateLayer(const FLandscapeLayer& InOtherLayer)
+{
+	ULandscapeInfo* LandscapeInfo = GetLandscapeInfo();
+	if (!LandscapeInfo || !GetMutableDefault<UEditorExperimentalSettings>()->bLandscapeLayerSystem)
+	{
+		return nullptr;
+	}
+
+	FLandscapeLayer NewLayer(InOtherLayer);
+	NewLayer.Guid = FGuid::NewGuid();
+
+	// Copy Brush and reparent to the new landscape level if required
+	for (FLandscapeLayerBrush& Brush : NewLayer.Brushes)
+	{
+		if (Brush.BPCustomBrush != nullptr)
+		{
+			if (Brush.BPCustomBrush->GetTypedOuter<ULevel>() != GetTypedOuter<ULevel>())
+			{
+				Brush.BPCustomBrush = DuplicateObject<ALandscapeBlueprintCustomBrush>(Brush.BPCustomBrush, GetTypedOuter<ULevel>());
+			}
+
+			Brush.BPCustomBrush->SetOwningLandscape(this);
+		}
+	}
+
+	int32 AddedIndex = LandscapeLayers.Add(NewLayer);
+
+	// Create associated layer data in each landscape proxy
+	LandscapeInfo->ForAllLandscapeProxies([&NewLayer](ALandscapeProxy* Proxy)
+	{
+		Proxy->AddLayer(NewLayer.Guid);
+	});
+
+	return &LandscapeLayers[AddedIndex];
+}
 
 void ALandscape::CreateLayer(FName InName)
 {
diff --git a/Engine/Source/Runtime/Landscape/Public/LandscapeEdit.h b/Engine/Source/Runtime/Landscape/Public/LandscapeEdit.h
--- a/Engine/Source/Runtime/Landscape/Public/LandscapeEdit.h
+++ b/Engine/Source/Runtime/Landscape/Public/LandscapeEdit.h
@@ -82,7 +82,7 @@ private:
 
 struct LANDSCAPE_API FLandscapeTextureDataInterface
 {
-	// tor
+	FLandscapeTextureDataInterface(bool bInUploadTextureChangesToGPU = true);
 	virtual ~FLandscapeTextureDataInterface();
 
 	// Texture data access
@@ -108,13 +108,14 @@ struct LANDSCAPE_API FLandscapeTextureDataInterface
 
 private:
 	TMap<UTexture2D*, FLandscapeTextureDataInfo*> TextureDataMap;
+	bool bUploadTextureChangesToGPU;
 };
 
 
 struct LANDSCAPE_API FLandscapeEditDataInterface : public FLandscapeTextureDataInterface
 {
 	// tor
-	FLandscapeEditDataInterface(ULandscapeInfo* InLandscape);
+	FLandscapeEditDataInterface(ULandscapeInfo* InLandscape, bool bInUploadTextureChangesToGPU = true);
 
 	// Misc
 	bool GetComponentsInRegion(int32 X1, int32 Y1, int32 X2, int32 Y2, TSet<ULandscapeComponent*>* OutComponents = NULL);
