- Landscape Layer System: Max Layer Setting (default to 8)

#rnx

diff --git a/Engine/Source/Editor/DetailCustomizations/Private/DetailCustomizations.cpp b/Engine/Source/Editor/DetailCustomizations/Private/DetailCustomizations.cpp
--- a/Engine/Source/Editor/DetailCustomizations/Private/DetailCustomizations.cpp
+++ b/Engine/Source/Editor/DetailCustomizations/Private/DetailCustomizations.cpp
@@ -129,6 +129,7 @@
 #include "MotionControllerDetails.h"
 #include "MotionControllerPinFactory.h"
 #include "LandscapeUIDetails.h"
+#include "LandscapeSettingsDetails.h"
 #include "PerPlatformPropertyCustomization.h"
 #include "SkeletalMeshReductionSettingsDetails.h"
 #include "SkeletalMeshLODSettingsDetails.h"
@@ -381,6 +382,7 @@ void FDetailCustomizationsModule::RegisterObjectCustomizations()
 
 	RegisterCustomClassLayout("Landscape", FOnGetDetailCustomizationInstance::CreateStatic(&FLandscapeUIDetails::MakeInstance));
 	RegisterCustomClassLayout("LandscapeProxy", FOnGetDetailCustomizationInstance::CreateStatic(&FLandscapeUIDetails::MakeInstance));
+	RegisterCustomClassLayout("LandscapeSettings", FOnGetDetailCustomizationInstance::CreateStatic(&FLandscapeSettingsDetails::MakeInstance));
 }
 
 
diff --git a/Engine/Source/Editor/DetailCustomizations/Private/LandscapeSettingsDetails.cpp b/Engine/Source/Editor/DetailCustomizations/Private/LandscapeSettingsDetails.cpp
new file mode 100644
--- /dev/null
+++ b/Engine/Source/Editor/DetailCustomizations/Private/LandscapeSettingsDetails.cpp
@@ -0,0 +1,47 @@
+// Copyright 1998-2019 Epic Games, Inc. All Rights Reserved.
+
+#include "LandscapeSettingsDetails.h"
+#include "Widgets/DeclarativeSyntaxSupport.h"
+#include "PropertyHandle.h"
+#include "DetailLayoutBuilder.h"
+#include "Runtime/Landscape/Public/LandscapeSettings.h"
+#include "Settings/EditorExperimentalSettings.h"
+//#include "Editor.h"
+
+FLandscapeSettingsDetails::FLandscapeSettingsDetails()
+{
+}
+
+
+TSharedRef<IDetailCustomization> FLandscapeSettingsDetails::MakeInstance()
+{
+	return MakeShareable( new FLandscapeSettingsDetails);
+}
+
+void FLandscapeSettingsDetails::CustomizeDetails( IDetailLayoutBuilder& DetailBuilder )
+{
+	TArray<TWeakObjectPtr<UObject>> EditingObjects;
+	DetailBuilder.GetObjectsBeingCustomized(EditingObjects);
+	DetailLayoutBuilder = &DetailBuilder;
+
+	GetMutableDefault<UEditorExperimentalSettings>()->OnSettingChanged().AddSP(this, &FLandscapeSettingsDetails::OnEditorExperimentalSettingsChanged);
+
+	if (EditingObjects.Num() == 1)
+	{
+		ULandscapeSettings* Settings = Cast<ULandscapeSettings>(EditingObjects[0].Get());
+
+		if (Settings != nullptr && !GetMutableDefault<UEditorExperimentalSettings>()->bLandscapeLayerSystem)
+		{
+			TSharedRef<IPropertyHandle> MaxNumberOfLayersProperty = DetailBuilder.GetProperty(GET_MEMBER_NAME_CHECKED(ULandscapeSettings, MaxNumberOfLayers));
+			DetailBuilder.HideProperty(MaxNumberOfLayersProperty);
+		}
+	}
+}
+
+void FLandscapeSettingsDetails::OnEditorExperimentalSettingsChanged(FName PropertyName)
+{
+	if (DetailLayoutBuilder && PropertyName == TEXT("bLandscapeLayerSystem"))
+	{
+		DetailLayoutBuilder->ForceRefreshDetails();
+	}
+}
\ No newline at end of file
diff --git a/Engine/Source/Editor/DetailCustomizations/Private/LandscapeSettingsDetails.h b/Engine/Source/Editor/DetailCustomizations/Private/LandscapeSettingsDetails.h
new file mode 100644
--- /dev/null
+++ b/Engine/Source/Editor/DetailCustomizations/Private/LandscapeSettingsDetails.h
@@ -0,0 +1,29 @@
+// Copyright 1998-2019 Epic Games, Inc. All Rights Reserved.
+
+#pragma once
+
+#include "CoreMinimal.h"
+#include "UObject/WeakObjectPtr.h"
+#include "Types/SlateEnums.h"
+#include "IDetailCustomization.h"
+
+class IDetailLayoutBuilder;
+
+class FLandscapeSettingsDetails : public IDetailCustomization
+{
+public:
+	/** Makes a new instance of this detail layout class for a specific detail view requesting it */
+	static TSharedRef<IDetailCustomization> MakeInstance();
+
+	/** IDetailCustomization interface */
+	virtual void CustomizeDetails( IDetailLayoutBuilder& DetailBuilder ) override;
+	/** End IDetailCustomization interface */
+
+	void OnEditorExperimentalSettingsChanged(FName PropertyName);
+private:
+
+	/** Use MakeInstance to create an instance of this class */
+	FLandscapeSettingsDetails();
+
+	IDetailLayoutBuilder* DetailLayoutBuilder;
+};
diff --git a/Engine/Source/Editor/LandscapeEditor/Private/LandscapeEditorDetailCustomization_Layers.cpp b/Engine/Source/Editor/LandscapeEditor/Private/LandscapeEditorDetailCustomization_Layers.cpp
--- a/Engine/Source/Editor/LandscapeEditor/Private/LandscapeEditorDetailCustomization_Layers.cpp
+++ b/Engine/Source/Editor/LandscapeEditor/Private/LandscapeEditorDetailCustomization_Layers.cpp
@@ -371,9 +371,10 @@ TSharedPtr<SWidget> FLandscapeEditorCustomNodeBuilder_Layers::OnLayerContextMenu
 		MenuBuilder.BeginSection("LandscapeEditorLayerActions", LOCTEXT("LandscapeEditorLayerActions.Heading", "Layers"));
 		{
 			// Create Layer
-			FUIAction CreateLayerAction = FUIAction(FExecuteAction::CreateLambda([SharedThis] { SharedThis->CreateLayer(); }));
-			MenuBuilder.AddMenuEntry(LOCTEXT("CreateLayer", "Create"), LOCTEXT("CreateLayerTooltip", "Create Layer"), FSlateIcon(), CreateLayerAction);
-
+			FUIAction CreateLayerAction = FUIAction(FExecuteAction::CreateLambda([SharedThis] { SharedThis->CreateLayer(); }), FCanExecuteAction::CreateLambda([Landscape] { return !Landscape->MaxLayersReached(); }));
+			TAttribute<FText> CreateLayerText(Landscape->MaxLayersReached() ? LOCTEXT("MaxLayerReached", "Create (Max layers reached)") : LOCTEXT("CreateLayer", "Create"));
+			MenuBuilder.AddMenuEntry(CreateLayerText, LOCTEXT("CreateLayerTooltip", "Create Layer"), FSlateIcon(), CreateLayerAction);
+	
 			if (Layer)
 			{
 				// Rename Layer
@@ -386,12 +387,10 @@ TSharedPtr<SWidget> FLandscapeEditorCustomNodeBuilder_Layers::OnLayerContextMenu
 					FUIAction ClearLayerAction = FUIAction(FExecuteAction::CreateLambda([SharedThis, InLayerIndex] { SharedThis->ClearLayer(InLayerIndex); }));
 					MenuBuilder.AddMenuEntry(LOCTEXT("ClearLayer", "Clear..."), LOCTEXT("ClearLayerTooltip", "Clear Layer"), FSlateIcon(), ClearLayerAction);
 
-					if (Landscape->LandscapeLayers.Num() > 1)
-					{
-						// Delete Layer
-						FUIAction DeleteLayerAction = FUIAction(FExecuteAction::CreateLambda([SharedThis, InLayerIndex] { SharedThis->DeleteLayer(InLayerIndex); } ));
-						MenuBuilder.AddMenuEntry(LOCTEXT("DeleteLayer", "Delete..."), LOCTEXT("DeleteLayerTooltip", "Delete Layer"), FSlateIcon(), DeleteLayerAction);
-					}
+					// Delete Layer
+					FUIAction DeleteLayerAction = FUIAction(FExecuteAction::CreateLambda([SharedThis, InLayerIndex] { SharedThis->DeleteLayer(InLayerIndex); }), FCanExecuteAction::CreateLambda([Landscape] { return Landscape->LandscapeLayers.Num() > 1; }));
+					TAttribute<FText> DeleteLayerText(Landscape->LandscapeLayers.Num() == 1 ? LOCTEXT("CantDeleteLastLayer", "Delete (Last layer)") : LOCTEXT("DeleteLayer", "Delete..."));
+					MenuBuilder.AddMenuEntry(DeleteLayerText, LOCTEXT("DeleteLayerTooltip", "Delete Layer"), FSlateIcon(), DeleteLayerAction);
 				}
 
 				if (Landscape->GetLandscapeSplinesReservedLayer() != Landscape->GetLayer(InLayerIndex))
diff --git a/Engine/Source/Runtime/Landscape/Classes/Landscape.h b/Engine/Source/Runtime/Landscape/Classes/Landscape.h
--- a/Engine/Source/Runtime/Landscape/Classes/Landscape.h
+++ b/Engine/Source/Runtime/Landscape/Classes/Landscape.h
@@ -256,6 +256,7 @@ public:
 	LANDSCAPE_API void DeleteLayer(int32 InLayerIndex);
 	LANDSCAPE_API void SetEditingLayer(const FGuid& InLayerGuid = FGuid());
 	LANDSCAPE_API const FGuid& GetEditingLayer() const;
+	LANDSCAPE_API bool MaxLayersReached() const;
 	LANDSCAPE_API void ShowOnlySelectedLayer(int32 InLayerIndex);
 	LANDSCAPE_API void ShowAllLayers();
 	LANDSCAPE_API void UpdateLandscapeSplines(const FGuid& InLayerGuid = FGuid(), bool bUpdateOnlySelected = false);
diff --git a/Engine/Source/Runtime/Landscape/Private/LandscapeEditLayers.cpp b/Engine/Source/Runtime/Landscape/Private/LandscapeEditLayers.cpp
--- a/Engine/Source/Runtime/Landscape/Private/LandscapeEditLayers.cpp
+++ b/Engine/Source/Runtime/Landscape/Private/LandscapeEditLayers.cpp
@@ -30,6 +30,7 @@ LandscapeEditLayers.cpp: Landscape editing layers mode
 #include "Materials/MaterialExpressionLandscapeVisibilityMask.h"
 #include "ShaderCompiler.h"
 #include "Algo/Count.h"
+#include "LandscapeSettings.h"
 #endif
 
 #define LOCTEXT_NAMESPACE "Landscape"
@@ -4933,6 +4934,11 @@ const FGuid& ALandscape::GetEditingLayer() const
 	return EditingLayer;
 }
 
+bool ALandscape::MaxLayersReached() const
+{
+	return LandscapeLayers.Num() >= GetDefault<ULandscapeSettings>()->MaxNumberOfLayers;
+}
+
 TMap<UTexture2D*, TArray<ULandscapeComponent*>> ALandscape::GenerateComponentsPerHeightmaps() const
 {
 	ULandscapeInfo* LandscapeInfo = GetLandscapeInfo();
@@ -5010,7 +5016,7 @@ FLandscapeLayer* ALandscape::DuplicateLayer(const FLandscapeLayer& InOtherLayer)
 void ALandscape::CreateLayer(FName InName)
 {
 	ULandscapeInfo* LandscapeInfo = GetLandscapeInfo();
-	if (!LandscapeInfo || !GetMutableDefault<UEditorExperimentalSettings>()->bLandscapeLayerSystem)
+	if (!LandscapeInfo || MaxLayersReached() || !GetMutableDefault<UEditorExperimentalSettings>()->bLandscapeLayerSystem)
 	{
 		return;
 	}
diff --git a/Engine/Source/Runtime/Landscape/Public/LandscapeSettings.h b/Engine/Source/Runtime/Landscape/Public/LandscapeSettings.h
new file mode 100644
--- /dev/null
+++ b/Engine/Source/Runtime/Landscape/Public/LandscapeSettings.h
@@ -0,0 +1,17 @@
+// Copyright 1998-2019 Epic Games, Inc. All Rights Reserved.
+
+#pragma once
+
+#include "CoreMinimal.h"
+#include "Engine/DeveloperSettings.h"
+#include "LandscapeSettings.generated.h"
+
+UCLASS(config = Engine, defaultconfig, meta = (DisplayName = "Landscape"))
+class LANDSCAPE_API ULandscapeSettings : public UDeveloperSettings
+{
+	GENERATED_BODY()
+
+public:
+	UPROPERTY(config, EditAnywhere, Category = "Layers", meta=(UIMin = "1", UIMax = "32", ClampMin = "1", ClampMax = "32", ToolTip = "This option controls the maximum editing layers that can be added to a Landscape"))
+	int32 MaxNumberOfLayers = 8;
+};
\ No newline at end of file
