- Review #6476389 changes

#rnx

diff --git a/Engine/Source/Editor/DetailCustomizations/Private/LandscapeSettingsDetails.cpp b/Engine/Source/Editor/DetailCustomizations/Private/LandscapeSettingsDetails.cpp
--- a/Engine/Source/Editor/DetailCustomizations/Private/LandscapeSettingsDetails.cpp
+++ b/Engine/Source/Editor/DetailCustomizations/Private/LandscapeSettingsDetails.cpp
@@ -9,9 +9,14 @@
 //#include "Editor.h"
 
 FLandscapeSettingsDetails::FLandscapeSettingsDetails()
+	: DetailLayoutBuilder(nullptr)
 {
 }
 
+FLandscapeSettingsDetails::~FLandscapeSettingsDetails()
+{
+	GetMutableDefault<UEditorExperimentalSettings>()->OnSettingChanged().RemoveAll(this);
+}
 
 TSharedRef<IDetailCustomization> FLandscapeSettingsDetails::MakeInstance()
 {
@@ -22,9 +27,12 @@ void FLandscapeSettingsDetails::CustomizeDetails( IDetailLayoutBuilder& DetailBu
 {
 	TArray<TWeakObjectPtr<UObject>> EditingObjects;
 	DetailBuilder.GetObjectsBeingCustomized(EditingObjects);
-	DetailLayoutBuilder = &DetailBuilder;
-
-	GetMutableDefault<UEditorExperimentalSettings>()->OnSettingChanged().AddSP(this, &FLandscapeSettingsDetails::OnEditorExperimentalSettingsChanged);
+	
+	if (!DetailLayoutBuilder)
+	{
+		DetailLayoutBuilder = &DetailBuilder;
+		GetMutableDefault<UEditorExperimentalSettings>()->OnSettingChanged().AddSP(this, &FLandscapeSettingsDetails::OnEditorExperimentalSettingsChanged);
+	}
 
 	if (EditingObjects.Num() == 1)
 	{
diff --git a/Engine/Source/Editor/DetailCustomizations/Private/LandscapeSettingsDetails.h b/Engine/Source/Editor/DetailCustomizations/Private/LandscapeSettingsDetails.h
--- a/Engine/Source/Editor/DetailCustomizations/Private/LandscapeSettingsDetails.h
+++ b/Engine/Source/Editor/DetailCustomizations/Private/LandscapeSettingsDetails.h
@@ -12,6 +12,8 @@ class IDetailLayoutBuilder;
 class FLandscapeSettingsDetails : public IDetailCustomization
 {
 public:
+	~FLandscapeSettingsDetails();
+
 	/** Makes a new instance of this detail layout class for a specific detail view requesting it */
 	static TSharedRef<IDetailCustomization> MakeInstance();
 
diff --git a/Engine/Source/Editor/LandscapeEditor/Private/LandscapeEditorDetailCustomization_Layers.cpp b/Engine/Source/Editor/LandscapeEditor/Private/LandscapeEditorDetailCustomization_Layers.cpp
--- a/Engine/Source/Editor/LandscapeEditor/Private/LandscapeEditorDetailCustomization_Layers.cpp
+++ b/Engine/Source/Editor/LandscapeEditor/Private/LandscapeEditorDetailCustomization_Layers.cpp
@@ -371,8 +371,8 @@ TSharedPtr<SWidget> FLandscapeEditorCustomNodeBuilder_Layers::OnLayerContextMenu
 		MenuBuilder.BeginSection("LandscapeEditorLayerActions", LOCTEXT("LandscapeEditorLayerActions.Heading", "Layers"));
 		{
 			// Create Layer
-			FUIAction CreateLayerAction = FUIAction(FExecuteAction::CreateLambda([SharedThis] { SharedThis->CreateLayer(); }), FCanExecuteAction::CreateLambda([Landscape] { return !Landscape->MaxLayersReached(); }));
-			TAttribute<FText> CreateLayerText(Landscape->MaxLayersReached() ? LOCTEXT("MaxLayerReached", "Create (Max layers reached)") : LOCTEXT("CreateLayer", "Create"));
+			FUIAction CreateLayerAction = FUIAction(FExecuteAction::CreateLambda([SharedThis] { SharedThis->CreateLayer(); }), FCanExecuteAction::CreateLambda([Landscape] { return !Landscape->IsMaxLayersReached(); }));
+			TAttribute<FText> CreateLayerText(Landscape->IsMaxLayersReached() ? LOCTEXT("MaxLayersReached", "Create (Max layers reached)") : LOCTEXT("CreateLayer", "Create"));
 			MenuBuilder.AddMenuEntry(CreateLayerText, LOCTEXT("CreateLayerTooltip", "Create Layer"), FSlateIcon(), CreateLayerAction);
 	
 			if (Layer)
diff --git a/Engine/Source/Runtime/Landscape/Classes/Landscape.h b/Engine/Source/Runtime/Landscape/Classes/Landscape.h
--- a/Engine/Source/Runtime/Landscape/Classes/Landscape.h
+++ b/Engine/Source/Runtime/Landscape/Classes/Landscape.h
@@ -256,7 +256,7 @@ public:
 	LANDSCAPE_API void DeleteLayer(int32 InLayerIndex);
 	LANDSCAPE_API void SetEditingLayer(const FGuid& InLayerGuid = FGuid());
 	LANDSCAPE_API const FGuid& GetEditingLayer() const;
-	LANDSCAPE_API bool MaxLayersReached() const;
+	LANDSCAPE_API bool IsMaxLayersReached() const;
 	LANDSCAPE_API void ShowOnlySelectedLayer(int32 InLayerIndex);
 	LANDSCAPE_API void ShowAllLayers();
 	LANDSCAPE_API void UpdateLandscapeSplines(const FGuid& InLayerGuid = FGuid(), bool bUpdateOnlySelected = false);
diff --git a/Engine/Source/Runtime/Landscape/Private/LandscapeEditLayers.cpp b/Engine/Source/Runtime/Landscape/Private/LandscapeEditLayers.cpp
--- a/Engine/Source/Runtime/Landscape/Private/LandscapeEditLayers.cpp
+++ b/Engine/Source/Runtime/Landscape/Private/LandscapeEditLayers.cpp
@@ -4934,7 +4934,7 @@ const FGuid& ALandscape::GetEditingLayer() const
 	return EditingLayer;
 }
 
-bool ALandscape::MaxLayersReached() const
+bool ALandscape::IsMaxLayersReached() const
 {
 	return LandscapeLayers.Num() >= GetDefault<ULandscapeSettings>()->MaxNumberOfLayers;
 }
@@ -5016,7 +5016,7 @@ FLandscapeLayer* ALandscape::DuplicateLayer(const FLandscapeLayer& InOtherLayer)
 void ALandscape::CreateLayer(FName InName)
 {
 	ULandscapeInfo* LandscapeInfo = GetLandscapeInfo();
-	if (!LandscapeInfo || MaxLayersReached() || !GetMutableDefault<UEditorExperimentalSettings>()->bLandscapeLayerSystem)
+	if (!LandscapeInfo || IsMaxLayersReached() || !GetMutableDefault<UEditorExperimentalSettings>()->bLandscapeLayerSystem)
 	{
 		return;
 	}
