- Landscape Layer System:
Feed the bCanPaint param to the material to get the red brush effect when layer can't be edited.
Fix MessageBox focus issue by delaying the popup to the MouseUp event.

#rnx

diff --git a/Engine/Source/Editor/LandscapeEditor/Private/LandscapeEdMode.cpp b/Engine/Source/Editor/LandscapeEditor/Private/LandscapeEdMode.cpp
--- a/Engine/Source/Editor/LandscapeEditor/Private/LandscapeEdMode.cpp
+++ b/Engine/Source/Editor/LandscapeEditor/Private/LandscapeEdMode.cpp
@@ -485,6 +485,8 @@ FGuid FEdModeLandscape::GetLandscapeSelectedLayer() const
 /** FEdMode: Called when the mode is entered */
 void FEdModeLandscape::Enter()
 {
+	ErrorReasonOnMouseUp = FText::GetEmpty();
+
 	// Call parent implementation
 	FEdMode::Enter();
 
@@ -506,7 +508,7 @@ void FEdModeLandscape::Enter()
 	OnLevelActorAddedDelegateHandle = GEngine->OnLevelActorAdded().AddSP(this, &FEdModeLandscape::OnLevelActorAdded);
 	OnLandscapeLayerSystemFlagChangedDelegateHandle = GetMutableDefault<UEditorExperimentalSettings>()->OnSettingChanged().AddLambda([this](FName PropertyName)
 	{ 
-		if (PropertyName == TEXT("bLandscapeLayerSystem"))
+		if (PropertyName == GET_MEMBER_NAME_CHECKED(UEditorExperimentalSettings, bLandscapeLayerSystem))
 		{
 			UpdateLandscapeList();
 			RefreshDetailPanel();
@@ -1824,6 +1826,14 @@ bool FEdModeLandscape::InputKey(FEditorViewportClient* ViewportClient, FViewport
 		return false;
 	}
 
+	if (!ErrorReasonOnMouseUp.IsEmpty() && Key == EKeys::LeftMouseButton && Event == IE_Released)
+	{
+		FMessageDialog::Open(EAppMsgType::Ok, ErrorReasonOnMouseUp);
+		ErrorReasonOnMouseUp = FText::GetEmpty();
+		return false;
+	}
+
+
 	if(IsAdjustingBrush(Viewport))
 	{
 		ToolActiveViewport = Viewport;
@@ -1932,10 +1942,8 @@ bool FEdModeLandscape::InputKey(FEditorViewportClient* ViewportClient, FViewport
 					FVector HitLocation;
 					if (LandscapeMouseTrace(ViewportClient, HitLocation))
 					{
-						FText Reason;
-						if (!CanEditLayer(&Reason))
+						if (!CanEditLayer(&ErrorReasonOnMouseUp))
 						{
-							FMessageDialog::Open(EAppMsgType::Ok, Reason);
 							return true;
 						}
 
diff --git a/Engine/Source/Editor/LandscapeEditor/Private/LandscapeEdMode.h b/Engine/Source/Editor/LandscapeEditor/Private/LandscapeEdMode.h
--- a/Engine/Source/Editor/LandscapeEditor/Private/LandscapeEdMode.h
+++ b/Engine/Source/Editor/LandscapeEditor/Private/LandscapeEdMode.h
@@ -257,6 +257,8 @@ public:
 
 	ULandscapeEditorObject* UISettings;
 
+	FText ErrorReasonOnMouseUp;
+
 	FLandscapeToolMode* CurrentToolMode;
 	FLandscapeTool* CurrentTool;
 	FLandscapeBrush* CurrentBrush;
diff --git a/Engine/Source/Editor/LandscapeEditor/Private/LandscapeEdModeBrushes.cpp b/Engine/Source/Editor/LandscapeEditor/Private/LandscapeEdModeBrushes.cpp
--- a/Engine/Source/Editor/LandscapeEditor/Private/LandscapeEdModeBrushes.cpp
+++ b/Engine/Source/Editor/LandscapeEditor/Private/LandscapeEdModeBrushes.cpp
@@ -202,7 +202,11 @@ public:
 			const ALandscapeProxy* LandscapeProxy = Component->GetLandscapeProxy();
 			const ULandscapeLayerInfoObject* LayerInfo = EdMode->CurrentToolTarget.LayerInfo.Get();
 
-			if (EdMode->CurrentToolTarget.TargetType == ELandscapeToolTargetType::Weightmap &&
+			if (!EdMode->CanEditLayer())
+			{
+				bCanPaint = false;
+			}
+			else if (EdMode->CurrentToolTarget.TargetType == ELandscapeToolTargetType::Weightmap &&
 				EdMode->UISettings->PaintingRestriction != ELandscapeLayerPaintingRestriction::None)
 			{
 				if (EdMode->UISettings->PaintingRestriction == ELandscapeLayerPaintingRestriction::UseComponentWhitelist &&
