- Landscape Layer System: Fix World Settings Layer enable not toggling when LandscapeEdMode wasn't yet created also remove register/unregister of static callback from non static constructor/destroy

#rnx

diff --git a/Engine/Source/Editor/LandscapeEditor/Private/LandscapeEdMode.cpp b/Engine/Source/Editor/LandscapeEditor/Private/LandscapeEdMode.cpp
--- a/Engine/Source/Editor/LandscapeEditor/Private/LandscapeEdMode.cpp
+++ b/Engine/Source/Editor/LandscapeEditor/Private/LandscapeEdMode.cpp
@@ -532,9 +532,7 @@ void FEdModeLandscape::Enter()
 			}
 		});
 	}
-
-	ALandscape::RegisterChangeLandscapeLayersStateDelegate();
-
+		
 	UpdateToolModes();
 
 	ALandscapeProxy* SelectedLandscape = GEditor->GetSelectedActors()->GetTop<ALandscapeProxy>();
diff --git a/Engine/Source/Editor/LandscapeEditor/Private/LandscapeEditorModule.cpp b/Engine/Source/Editor/LandscapeEditor/Private/LandscapeEditorModule.cpp
--- a/Engine/Source/Editor/LandscapeEditor/Private/LandscapeEditorModule.cpp
+++ b/Engine/Source/Editor/LandscapeEditor/Private/LandscapeEditorModule.cpp
@@ -71,6 +71,7 @@ public:
 			true,
 			300
 			);
+		ALandscape::RegisterChangeLandscapeLayersStateDelegate();
 
 		// register customizations
 		FPropertyEditorModule& PropertyModule = FModuleManager::GetModuleChecked<FPropertyEditorModule>("PropertyEditor");
@@ -118,6 +119,7 @@ public:
 	 */
 	virtual void ShutdownModule() override
 	{
+		ALandscape::UnregisterChangeLandscapeLayersStateDelegate();
 		FLandscapeEditorCommands::Unregister();
 
 		FEditorDelegates::PreSaveWorld.Remove(PreSaveWorldHandle);
diff --git a/Engine/Source/Runtime/Landscape/Private/Landscape.cpp b/Engine/Source/Runtime/Landscape/Private/Landscape.cpp
--- a/Engine/Source/Runtime/Landscape/Private/Landscape.cpp
+++ b/Engine/Source/Runtime/Landscape/Private/Landscape.cpp
@@ -1090,13 +1090,7 @@ ALandscape::ALandscape(const FObjectInitializer& ObjectInitializer)
 	WeightmapScratchExtractLayerTextureResource = nullptr;
 	WeightmapScratchPackLayerTextureResource = nullptr;
 	bLandscapeLayersAreInitialized = false;
-	LandscapeEdMode = nullptr;
-
-	if (!HasAnyFlags(RF_ArchetypeObject | RF_ClassDefaultObject))
-	{
-		ALandscape::RegisterChangeLandscapeLayersStateDelegate();
-	}
-	
+	LandscapeEdMode = nullptr;	
 #endif // WITH_EDITORONLY_DATA
 }
 
diff --git a/Engine/Source/Runtime/Landscape/Private/LandscapeEditLayers.cpp b/Engine/Source/Runtime/Landscape/Private/LandscapeEditLayers.cpp
--- a/Engine/Source/Runtime/Landscape/Private/LandscapeEditLayers.cpp
+++ b/Engine/Source/Runtime/Landscape/Private/LandscapeEditLayers.cpp
@@ -4740,8 +4740,6 @@ void ALandscapeProxy::InitializeLayerWithEmptyContent(const FGuid& InLayerGuid)
 void ALandscape::BeginDestroy()
 {
 #if WITH_EDITOR
-	ALandscape::UnregisterChangeLandscapeLayersStateDelegate();
-
 	if (CanHaveLayersContent())
 	{
 		if (CombinedLayersWeightmapAllMaterialLayersResource != nullptr)
