- Landscape Layer System : Fixed crash when creating a new landscape and selecting tools that are invalid.

diff --git a/Engine/Source/Editor/LandscapeEditor/Private/LandscapeEdMode.cpp b/Engine/Source/Editor/LandscapeEditor/Private/LandscapeEdMode.cpp
--- a/Engine/Source/Editor/LandscapeEditor/Private/LandscapeEdMode.cpp
+++ b/Engine/Source/Editor/LandscapeEditor/Private/LandscapeEdMode.cpp
@@ -420,6 +420,10 @@ void FEdModeLandscape::UpdateToolModes()
 
 	FLandscapeToolMode* ToolMode_Sculpt = new(LandscapeToolModes)FLandscapeToolMode(TEXT("ToolMode_Sculpt"), ELandscapeToolTargetTypeMask::Heightmap | ELandscapeToolTargetTypeMask::Visibility);
 	ToolMode_Sculpt->ValidTools.Add(TEXT("Sculpt"));
+	if (CanHaveLandscapeLayersContent())
+	{
+		ToolMode_Sculpt->ValidTools.Add(TEXT("Erase"));
+	}
 	ToolMode_Sculpt->ValidTools.Add(TEXT("Smooth"));
 	ToolMode_Sculpt->ValidTools.Add(TEXT("Flatten"));
 	ToolMode_Sculpt->ValidTools.Add(TEXT("Ramp"));
@@ -431,7 +435,6 @@ void FEdModeLandscape::UpdateToolModes()
 
 	if (CanHaveLandscapeLayersContent())
 	{
-        ToolMode_Sculpt->ValidTools.Add(TEXT("Erase"));
 		ToolMode_Sculpt->ValidTools.Add(TEXT("BPCustom"));
 	}
 
@@ -489,6 +492,15 @@ ULandscapeLayerInfoObject* FEdModeLandscape::GetSelectedLandscapeLayerInfo() con
 	return CurrentToolTarget.LayerInfo.Get();
 }
 
+void FEdModeLandscape::SetLandscapeInfo(ULandscapeInfo* InLandscapeInfo)
+{
+	if (CurrentToolTarget.LandscapeInfo != InLandscapeInfo)
+	{
+		CurrentToolTarget.LandscapeInfo = InLandscapeInfo;
+		UpdateToolModes();
+	}
+}
+
 /** FEdMode: Called when the mode is entered */
 void FEdModeLandscape::Enter()
 {
@@ -537,13 +549,13 @@ void FEdModeLandscape::Enter()
 			}
 		});
 	}
-		
+
 	UpdateToolModes();
 
 	ALandscapeProxy* SelectedLandscape = GEditor->GetSelectedActors()->GetTop<ALandscapeProxy>();
 	if (SelectedLandscape)
 	{
-		CurrentToolTarget.LandscapeInfo = SelectedLandscape->GetLandscapeInfo();
+		SetLandscapeInfo(SelectedLandscape->GetLandscapeInfo());
 		GEditor->SelectNone(false, true);
 		GEditor->SelectActor(SelectedLandscape, true, false);
 	}
@@ -2456,7 +2468,7 @@ int32 FEdModeLandscape::UpdateLandscapeList()
 				CurrentBrush->LeaveBrush();
 				CurrentTool->ExitTool();
 			}
-			CurrentToolTarget.LandscapeInfo = LandscapeList[0].Info;
+			SetLandscapeInfo(LandscapeList[0].Info);
 			CurrentIndex = 0;
 
 			SetCurrentLayer(0);
@@ -2481,7 +2493,7 @@ int32 FEdModeLandscape::UpdateLandscapeList()
 		else
 		{
 			// no landscape, switch to "new landscape" tool
-			CurrentToolTarget.LandscapeInfo = nullptr;
+			SetLandscapeInfo(nullptr);
 			UpdateTargetList();
 			SetCurrentToolMode("ToolMode_Manage", false);
 			SetCurrentTool("NewLandscape");
@@ -2511,7 +2523,7 @@ void FEdModeLandscape::SetTargetLandscape(const TWeakObjectPtr<ULandscapeInfo>&
 		LandscapeProxy->OnMaterialChangedDelegate().RemoveAll(this);
 	}
 
-	CurrentToolTarget.LandscapeInfo = InLandscapeInfo.Get();
+	SetLandscapeInfo(InLandscapeInfo.Get());
 	UpdateTargetList();
 	// force a Leave and Enter the current tool, in case it has something about the current landscape cached
 	SetCurrentTool(CurrentToolIndex);
@@ -3333,7 +3345,7 @@ bool FEdModeLandscape::Select(AActor* InActor, bool bInSelected)
 
 		if (CurrentToolTarget.LandscapeInfo != Landscape->GetLandscapeInfo())
 		{
-			CurrentToolTarget.LandscapeInfo = Landscape->GetLandscapeInfo();
+			SetLandscapeInfo(Landscape->GetLandscapeInfo());
 			UpdateTargetList();
 
 			// If we were in "New Landscape" mode and we select a landscape then switch to editing mode
diff --git a/Engine/Source/Editor/LandscapeEditor/Private/LandscapeEdMode.h b/Engine/Source/Editor/LandscapeEditor/Private/LandscapeEdMode.h
--- a/Engine/Source/Editor/LandscapeEditor/Private/LandscapeEdMode.h
+++ b/Engine/Source/Editor/LandscapeEditor/Private/LandscapeEdMode.h
@@ -578,6 +578,8 @@ public:
 	{
 		return GetEditingState() == ELandscapeEditingState::Enabled;
 	}
+
+	void SetLandscapeInfo(ULandscapeInfo* InLandscapeInfo);
 	
 private:
 	TArray<TSharedRef<FLandscapeTargetListInfo>> LandscapeTargetList;
diff --git a/Engine/Source/Editor/LandscapeEditor/Private/LandscapeEditorDetailCustomization_NewLandscape.cpp b/Engine/Source/Editor/LandscapeEditor/Private/LandscapeEditorDetailCustomization_NewLandscape.cpp
--- a/Engine/Source/Editor/LandscapeEditor/Private/LandscapeEditorDetailCustomization_NewLandscape.cpp
+++ b/Engine/Source/Editor/LandscapeEditor/Private/LandscapeEditorDetailCustomization_NewLandscape.cpp
@@ -856,7 +856,7 @@ FReply FLandscapeEditorDetailCustomization_NewLandscape::OnCreateButtonClicked()
 		}
 
 		LandscapeEdMode->UpdateLandscapeList();
-		LandscapeEdMode->CurrentToolTarget.LandscapeInfo = LandscapeInfo;
+		LandscapeEdMode->SetLandscapeInfo(LandscapeInfo);
 		LandscapeEdMode->CurrentToolTarget.TargetType = ELandscapeToolTargetType::Heightmap;
 		LandscapeEdMode->CurrentToolTarget.LayerInfo = nullptr;
 		LandscapeEdMode->CurrentToolTarget.LayerName = NAME_None;
diff --git a/Engine/Source/Runtime/Landscape/Private/LandscapeEditLayers.cpp b/Engine/Source/Runtime/Landscape/Private/LandscapeEditLayers.cpp
--- a/Engine/Source/Runtime/Landscape/Private/LandscapeEditLayers.cpp
+++ b/Engine/Source/Runtime/Landscape/Private/LandscapeEditLayers.cpp
@@ -4223,7 +4223,11 @@ void ALandscape::MonitorLandscapeEdModeChanges()
 		bRequiredEditingClientFullUpdate = true;
 	}
 
-	TWeakObjectPtr<ULandscapeLayerInfoObject> NewLayerInfoObject = LandscapeEdMode->GetSelectedLandscapeLayerInfo();
+	TWeakObjectPtr<ULandscapeLayerInfoObject> NewLayerInfoObject;
+	if (LandscapeEdMode)
+	{
+		NewLayerInfoObject = LandscapeEdMode->GetSelectedLandscapeLayerInfo();
+	}
 	if (LandscapeEdModeInfo.SelectedLayerInfoObject != NewLayerInfoObject)
 	{
 		LandscapeEdModeInfo.SelectedLayerInfoObject = NewLayerInfoObject;
