- Landscape Layer System: Fix new bug preventing creation of new landscape after having deleted one

#rnx

diff --git a/Engine/Source/Editor/LandscapeEditor/Private/LandscapeEdMode.cpp b/Engine/Source/Editor/LandscapeEditor/Private/LandscapeEdMode.cpp
--- a/Engine/Source/Editor/LandscapeEditor/Private/LandscapeEdMode.cpp
+++ b/Engine/Source/Editor/LandscapeEditor/Private/LandscapeEdMode.cpp
@@ -2217,17 +2217,18 @@ void FEdModeLandscape::SetCurrentTool(int32 ToolIndex)
 	{
 		CurrentTool->PreviousBrushIndex = CurrentBrushSetIndex;
 		CurrentTool->ExitTool();
+		CurrentTool = nullptr;
 	}
 	CurrentToolIndex = LandscapeTools.IsValidIndex(ToolIndex) ? ToolIndex : 0;
-	CurrentTool = LandscapeTools[CurrentToolIndex].Get();
-	if (!CurrentToolMode->ValidTools.Contains(CurrentTool->GetToolName()))
+	FLandscapeTool* NewTool = LandscapeTools[CurrentToolIndex].Get();
+	if (!CurrentToolMode->ValidTools.Contains(NewTool->GetToolName()))
 	{
 		// if tool isn't valid for this mode then automatically switch modes
 		// this mostly happens with shortcut keys
 		bool bFoundValidMode = false;
 		for (int32 i = 0; i < LandscapeToolModes.Num(); ++i)
 		{
-			if (LandscapeToolModes[i].ValidTools.Contains(CurrentTool->GetToolName()))
+			if (LandscapeToolModes[i].ValidTools.Contains(NewTool->GetToolName()))
 			{
 				SetCurrentToolMode(LandscapeToolModes[i].ToolModeName, false);
 				bFoundValidMode = true;
@@ -2237,11 +2238,15 @@ void FEdModeLandscape::SetCurrentTool(int32 ToolIndex)
 		
 		// default to first valid tool of current mode
 		if (!bFoundValidMode)
-		{
+		{		
 			SetCurrentTool(CurrentToolMode->ValidTools[0]);
+			return;
 		}
 	}
 
+	// Assign 
+	CurrentTool = NewTool;
+
 	// Set target type appropriate for tool
 	if (CurrentTool->GetSupportedTargetTypes() == ELandscapeToolTargetTypeMask::NA)
 	{
@@ -2449,11 +2454,7 @@ int32 FEdModeLandscape::UpdateLandscapeList()
 	{
 		if (LandscapeList.Num() > 0)
 		{
-			if (CurrentTool != nullptr)
-			{
-				CurrentBrush->LeaveBrush();
-				CurrentTool->ExitTool();
-			}
+			FName CurrentToolName = CurrentTool != nullptr ? CurrentTool->GetToolName() : FName();
 			SetLandscapeInfo(LandscapeList[0].Info);
 			CurrentIndex = 0;
 
@@ -2469,11 +2470,10 @@ int32 FEdModeLandscape::UpdateLandscapeList()
 
 			UpdateTargetList();
 			UpdateShownLayerList();
-
-			if (CurrentTool != nullptr)
+						
+			if (!CurrentToolName.IsNone())
 			{
-				CurrentTool->EnterTool();
-				CurrentBrush->EnterBrush();
+				SetCurrentTool(CurrentToolName);
 			}
 		}
 		else
