Landscape BP Custom Brush: Fixed undo/redo of brush reordering
#rnx

diff --git a/Engine/Source/Editor/LandscapeEditor/Private/LandscapeEditorDetailCustomization_Layers.cpp b/Engine/Source/Editor/LandscapeEditor/Private/LandscapeEditorDetailCustomization_Layers.cpp
--- a/Engine/Source/Editor/LandscapeEditor/Private/LandscapeEditorDetailCustomization_Layers.cpp
+++ b/Engine/Source/Editor/LandscapeEditor/Private/LandscapeEditorDetailCustomization_Layers.cpp
@@ -717,7 +717,6 @@ FReply FLandscapeEditorCustomNodeBuilder_Layers::HandleAcceptDrop(FDragDropEvent
 			{
 				LandscapeEdMode->SetCurrentLayer(DestinationLayerIndex);
 				LandscapeEdMode->RefreshDetailPanel();
-				LandscapeEdMode->RequestLayersContentUpdateForceAll();
 				return FReply::Handled();
 			}
 		}
diff --git a/Engine/Source/Editor/LandscapeEditor/Private/LandscapeEditorDetailCustomization_LayersBrushStack.cpp b/Engine/Source/Editor/LandscapeEditor/Private/LandscapeEditorDetailCustomization_LayersBrushStack.cpp
--- a/Engine/Source/Editor/LandscapeEditor/Private/LandscapeEditorDetailCustomization_LayersBrushStack.cpp
+++ b/Engine/Source/Editor/LandscapeEditor/Private/LandscapeEditorDetailCustomization_LayersBrushStack.cpp
@@ -344,28 +344,16 @@ FReply FLandscapeEditorCustomNodeBuilder_LayersBrushStack::HandleAcceptDrop(FDra
 	if (DragDropOperation.IsValid())
 	{
 		FEdModeLandscape* LandscapeEdMode = GetEditorMode();
-
-		if (LandscapeEdMode != nullptr)
+		ALandscape* Landscape = LandscapeEdMode ? LandscapeEdMode->GetLandscape() : nullptr;
+		if (Landscape)
 		{
-			TArray<int8>& BrushOrderStack = LandscapeEdMode->GetBrushesOrderForCurrentLayer(LandscapeEdMode->CurrentToolTarget.TargetType);
-
-			if (BrushOrderStack.IsValidIndex(DragDropOperation->SlotIndexBeingDragged) && BrushOrderStack.IsValidIndex(SlotIndex))
+			int32 StartingLayerIndex = DragDropOperation->SlotIndexBeingDragged;
+			int32 DestinationLayerIndex = SlotIndex;
+			const FScopedTransaction Transaction(LOCTEXT("Landscape_LayerBrushes_Reorder", "Reorder Layer Brush"));
+			if (Landscape->ReorderLayerBrush(LandscapeEdMode->GetCurrentLayerIndex(), LandscapeEdMode->CurrentToolTarget.TargetType, StartingLayerIndex, DestinationLayerIndex))
 			{
-				int32 StartingLayerIndex = DragDropOperation->SlotIndexBeingDragged;
-				int32 DestinationLayerIndex = SlotIndex;
-
-				if (StartingLayerIndex != INDEX_NONE && DestinationLayerIndex != INDEX_NONE)
-				{
-					int8 MovingBrushIndex = BrushOrderStack[StartingLayerIndex];
-					 
-					BrushOrderStack.RemoveAt(StartingLayerIndex);
-					BrushOrderStack.Insert(MovingBrushIndex, DestinationLayerIndex);
-
-					LandscapeEdMode->RefreshDetailPanel();
-					LandscapeEdMode->RequestLayersContentUpdateForceAll();
-
-					return FReply::Handled();
-				}
+				LandscapeEdMode->RefreshDetailPanel();
+				return FReply::Handled();
 			}
 		}
 	}
diff --git a/Engine/Source/Runtime/Landscape/Classes/Landscape.h b/Engine/Source/Runtime/Landscape/Classes/Landscape.h
--- a/Engine/Source/Runtime/Landscape/Classes/Landscape.h
+++ b/Engine/Source/Runtime/Landscape/Classes/Landscape.h
@@ -284,6 +284,7 @@ public:
 
 	LANDSCAPE_API void AddBrushToLayer(int32 InLayerIndex, int32 InTargetType, class ALandscapeBlueprintCustomBrush* InBrush);
 	LANDSCAPE_API void RemoveBrushFromLayer(int32 InLayerIndex, int32 InTargetType, class ALandscapeBlueprintCustomBrush* InBrush);
+	LANDSCAPE_API bool ReorderLayerBrush(int32 InLayerIndex, int32 InTargetType, int32 InStartingLayerBrushIndex, int32 InDestinationLayerBrushIndex);
 	LANDSCAPE_API bool AreAllBrushesCommitedToLayer(int32 InLayerIndex, int32 InTargetType);
 	LANDSCAPE_API void SetBrushesCommitStateForLayer(int32 InLayerIndex, int32 InTargetType, bool InCommited);
 	LANDSCAPE_API TArray<int8>& GetBrushesOrderForLayer(int32 InLayerIndex, int32 InTargetType);
diff --git a/Engine/Source/Runtime/Landscape/Private/LandscapeEditLayers.cpp b/Engine/Source/Runtime/Landscape/Private/LandscapeEditLayers.cpp
--- a/Engine/Source/Runtime/Landscape/Private/LandscapeEditLayers.cpp
+++ b/Engine/Source/Runtime/Landscape/Private/LandscapeEditLayers.cpp
@@ -5353,6 +5353,7 @@ bool ALandscape::ReorderLayer(int32 InStartingLayerIndex, int32 InDestinationLay
 		FLandscapeLayer Layer = LandscapeLayers[InStartingLayerIndex];
 		LandscapeLayers.RemoveAt(InStartingLayerIndex);
 		LandscapeLayers.Insert(Layer, InDestinationLayerIndex);
+		RequestLayersContentUpdateForceAll();
 		return true;
 	}
 	return false;
@@ -5415,6 +5416,26 @@ void ALandscape::SetLayerSubstractiveBlendStatus(int32 InLayerIndex, bool InStat
 	RequestLayersContentUpdateForceAll(ELandscapeLayerUpdateMode::Update_Weightmap_All);
 }
 
+bool ALandscape::ReorderLayerBrush(int32 InLayerIndex, int32 InTargetType, int32 InStartingLayerBrushIndex, int32 InDestinationLayerBrushIndex)
+{
+	if (const FLandscapeLayer* Layer = GetLayer(InLayerIndex))
+	{
+		TArray<int8>& BrushOrderStack = GetBrushesOrderForLayer(InLayerIndex, InTargetType);
+		if (InStartingLayerBrushIndex != InDestinationLayerBrushIndex &&
+			BrushOrderStack.IsValidIndex(InStartingLayerBrushIndex) &&
+			BrushOrderStack.IsValidIndex(InDestinationLayerBrushIndex))
+		{
+			Modify();
+			int8 MovingBrushIndex = BrushOrderStack[InStartingLayerBrushIndex];
+			BrushOrderStack.RemoveAt(InStartingLayerBrushIndex);
+			BrushOrderStack.Insert(MovingBrushIndex, InDestinationLayerBrushIndex);
+			RequestLayersContentUpdateForceAll();
+			return true;
+		}
+	}
+	return false;
+}
+
 void ALandscape::AddBrushToLayer(int32 InLayerIndex, int32 InTargetType, ALandscapeBlueprintCustomBrush* InBrush)
 {
 	FLandscapeLayer* Layer = GetLayer(InLayerIndex);
