- Landscape Layer System: BP Brush Fixes
 Ability to add an unassigned brush to a layer with right-click add sub menu
 Ability to unassign brush from layer with right-click remove
 Scope the Add Brush Action so that it is properly undoable
 Remove Brush from layer if it is removed from the level
 Remove useless redundant expandable category from Layers and Brushes
 Add Visibility toggle to brushes

#rnx

diff --git a/Engine/Source/Editor/LandscapeEditor/Private/LandscapeEdMode.cpp b/Engine/Source/Editor/LandscapeEditor/Private/LandscapeEdMode.cpp
--- a/Engine/Source/Editor/LandscapeEditor/Private/LandscapeEdMode.cpp
+++ b/Engine/Source/Editor/LandscapeEditor/Private/LandscapeEdMode.cpp
@@ -569,6 +569,7 @@ void FEdModeLandscape::Enter()
 	// For now depends on the SpawnActor() above in order to get the current editor world as edmodes don't get told
 	UpdateLandscapeList();
 	UpdateTargetList();
+	UpdateBrushList();
 
 	OnWorldChangeDelegateHandle                 = FEditorSupportDelegates::WorldChange.AddRaw(this, &FEdModeLandscape::HandleLevelsChanged, true);
 	OnLevelsChangedDelegateHandle				= GetWorld()->OnLevelsChanged().AddRaw(this, &FEdModeLandscape::HandleLevelsChanged, true);
@@ -2383,6 +2384,11 @@ void FEdModeLandscape::SetCurrentBrush(int32 BrushIndex)
 	}
 }
 
+const TArray<ALandscapeBlueprintCustomBrush*>& FEdModeLandscape::GetBrushList() const
+{
+	return BrushList;
+}
+
 const TArray<TSharedRef<FLandscapeTargetListInfo>>& FEdModeLandscape::GetTargetList() const
 {
 	return LandscapeTargetList;
@@ -2987,6 +2993,7 @@ void FEdModeLandscape::HandleLevelsChanged(bool ShouldExitMode)
 	UpdateLandscapeList();
 	UpdateTargetList();
 	UpdateShownLayerList();
+	UpdateBrushList();
 
 	// if the Landscape is deleted then close the landscape editor
 	if (ShouldExitMode && bHadLandscape && CurrentToolTarget.LandscapeInfo == nullptr)
@@ -4474,6 +4481,7 @@ void FEdModeLandscape::AddBrushToCurrentLayer(ALandscapeBlueprintCustomBrush* In
 	}
 
 	Landscape->AddBrushToLayer(GetCurrentLayerIndex(), InBrush);
+	RefreshDetailPanel();
 }
 
 void FEdModeLandscape::RemoveBrushFromCurrentLayer(ALandscapeBlueprintCustomBrush* InBrush)
@@ -4486,6 +4494,7 @@ void FEdModeLandscape::RemoveBrushFromCurrentLayer(ALandscapeBlueprintCustomBrus
 	}
 
 	Landscape->RemoveBrushFromLayer(GetCurrentLayerIndex(), InBrush);
+	RefreshDetailPanel();
 }
 
 ALandscapeBlueprintCustomBrush* FEdModeLandscape::GetBrushForCurrentLayer(int8 InBrushIndex) const
@@ -4696,6 +4705,16 @@ bool FEdModeLandscape::NeedToFillEmptyMaterialLayers() const
 	return bCanFill;
 }
 
+void FEdModeLandscape::UpdateBrushList()
+{
+	BrushList.Empty();
+	for (TObjectIterator<ALandscapeBlueprintCustomBrush> BrushIt(RF_Transient|RF_ClassDefaultObject|RF_ArchetypeObject, true, EInternalObjectFlags::PendingKill); BrushIt; ++BrushIt)
+	{
+		BrushList.AddUnique(*BrushIt);
+	}
+}
+
+
 void FEdModeLandscape::OnLevelActorAdded(AActor* InActor)
 {
 	if (ALandscape* Landscape = Cast <ALandscape>(InActor))
@@ -4707,6 +4726,7 @@ void FEdModeLandscape::OnLevelActorAdded(AActor* InActor)
 
 	if (Brush != nullptr && Brush->GetTypedOuter<UPackage>() != GetTransientPackage())
 	{
+		BrushList.AddUnique(Brush);
 		AddBrushToCurrentLayer(Brush);
 		RefreshDetailPanel();
 	}
@@ -4723,6 +4743,7 @@ void FEdModeLandscape::OnLevelActorRemoved(AActor* InActor)
 
 	if (Brush != nullptr && Brush->GetTypedOuter<UPackage>() != GetTransientPackage())
 	{
+		BrushList.Remove(Brush);
 		RemoveBrushFromCurrentLayer(Brush);
 		RefreshDetailPanel();
 	}
diff --git a/Engine/Source/Editor/LandscapeEditor/Private/LandscapeEdMode.h b/Engine/Source/Editor/LandscapeEditor/Private/LandscapeEdMode.h
--- a/Engine/Source/Editor/LandscapeEditor/Private/LandscapeEdMode.h
+++ b/Engine/Source/Editor/LandscapeEditor/Private/LandscapeEdMode.h
@@ -508,6 +508,9 @@ public:
 	void SetCurrentBrush(FName BrushName);
 	void SetCurrentBrush(int32 BrushIndex);
 
+	void UpdateBrushList();
+	const TArray<ALandscapeBlueprintCustomBrush*>& GetBrushList() const;
+
 	const TArray<TSharedRef<FLandscapeTargetListInfo>>& GetTargetList() const;
 	const TArray<FName>* GetTargetDisplayOrderList() const;
 	const TArray<FName>& GetTargetShownList() const;
@@ -614,6 +617,7 @@ public:
 private:
 	TArray<TSharedRef<FLandscapeTargetListInfo>> LandscapeTargetList;
 	TArray<FLandscapeListInfo> LandscapeList;
+	TArray<ALandscapeBlueprintCustomBrush*> BrushList;
 	TArray<FName> ShownTargetLayerList;
 	
 	/** Represent the index offset of the target layer in LandscapeTargetList */
diff --git a/Engine/Source/Editor/LandscapeEditor/Private/LandscapeEdModeBPCustomTools.cpp b/Engine/Source/Editor/LandscapeEditor/Private/LandscapeEdModeBPCustomTools.cpp
--- a/Engine/Source/Editor/LandscapeEditor/Private/LandscapeEdModeBPCustomTools.cpp
+++ b/Engine/Source/Editor/LandscapeEditor/Private/LandscapeEdModeBPCustomTools.cpp
@@ -96,6 +96,7 @@ public:
 			SpawnInfo.bNoFail = true;
 			SpawnInfo.OverrideLevel = Info->LandscapeActor.Get()->GetTypedOuter<ULevel>(); // always spawn in the same level as the one containing the ALandscape
 
+			FScopedTransaction Transaction(LOCTEXT("LandscapeEdModeBPCustomToolSpawn", "Create landscape brush"));
 			ALandscapeBlueprintCustomBrush* Brush = ViewportClient->GetWorld()->SpawnActor<ALandscapeBlueprintCustomBrush>(EdMode->UISettings->BlueprintCustomBrush, SpawnLocation, FRotator(0.0f), SpawnInfo);
 			EdMode->UISettings->BlueprintCustomBrush = nullptr;
 
diff --git a/Engine/Source/Editor/LandscapeEditor/Private/LandscapeEditorDetailCustomization_Layers.cpp b/Engine/Source/Editor/LandscapeEditor/Private/LandscapeEditorDetailCustomization_Layers.cpp
--- a/Engine/Source/Editor/LandscapeEditor/Private/LandscapeEditorDetailCustomization_Layers.cpp
+++ b/Engine/Source/Editor/LandscapeEditor/Private/LandscapeEditorDetailCustomization_Layers.cpp
@@ -127,19 +127,7 @@ void FLandscapeEditorCustomNodeBuilder_Layers::SetOnRebuildChildren(FSimpleDeleg
 
 void FLandscapeEditorCustomNodeBuilder_Layers::GenerateHeaderRowContent(FDetailWidgetRow& NodeRow)
 {
-	FEdModeLandscape* LandscapeEdMode = GetEditorMode();
-
-	if (LandscapeEdMode == NULL)
-	{
-		return;
-	}
 
-	NodeRow.NameWidget
-		[
-			SNew(STextBlock)
-			.Font(IDetailLayoutBuilder::GetDetailFont())
-			.Text(FText::FromString(TEXT("Layers")))
-		];
 }
 
 BEGIN_SLATE_FUNCTION_BUILD_OPTIMIZATION
@@ -156,7 +144,7 @@ void FLandscapeEditorCustomNodeBuilder_Layers::GenerateChildContent(IDetailChild
 		LayerList->SetDropIndicator_Above(*FEditorStyle::GetBrush("LandscapeEditor.TargetList.DropZone.Above"));
 		LayerList->SetDropIndicator_Below(*FEditorStyle::GetBrush("LandscapeEditor.TargetList.DropZone.Below"));
 
-		ChildrenBuilder.AddCustomRow(FText::FromString(FString(TEXT("Layers"))))
+		ChildrenBuilder.AddCustomRow(FText::FromString(FString(TEXT("Test"))))
 			.Visibility(EVisibility::Visible)
 			[
 				LayerList.ToSharedRef()
@@ -359,6 +347,26 @@ void FLandscapeEditorCustomNodeBuilder_Layers::SetLayerName(const FText& InText,
 	}
 }
 
+
+void FLandscapeEditorCustomNodeBuilder_Layers::FillAddBrushMenu(FMenuBuilder& MenuBuilder, TArray<ALandscapeBlueprintCustomBrush*> Brushes)
+{
+	for (ALandscapeBlueprintCustomBrush* Brush : Brushes)
+	{
+		if (Brush->GetOwningLandscape() == nullptr)
+		{
+			TSharedRef<FLandscapeEditorCustomNodeBuilder_Layers> SharedThis = AsShared();
+			FUIAction AddAction = FUIAction(FExecuteAction::CreateLambda([SharedThis, Brush]() { SharedThis->AddBrushToCurrentLayer(Brush); }));
+			MenuBuilder.AddMenuEntry(FText::FromString(Brush->GetActorLabel()), FText(), FSlateIcon(), AddAction);
+		}
+	}
+}
+
+void FLandscapeEditorCustomNodeBuilder_Layers::AddBrushToCurrentLayer(ALandscapeBlueprintCustomBrush* Brush)
+{
+	const FScopedTransaction Transaction(LOCTEXT("LandscapeBrushAddToCurrentLayerTransaction", "Add brush to current layer"));
+	GetEditorMode()->AddBrushToCurrentLayer(Brush);
+}
+
 TSharedPtr<SWidget> FLandscapeEditorCustomNodeBuilder_Layers::OnLayerContextMenuOpening(int32 InLayerIndex)
 {
 	FEdModeLandscape* LandscapeEdMode = GetEditorMode();
@@ -439,9 +447,25 @@ TSharedPtr<SWidget> FLandscapeEditorCustomNodeBuilder_Layers::OnLayerContextMenu
 		}
 		MenuBuilder.EndSection();
 
+		const TArray<ALandscapeBlueprintCustomBrush*>& Brushes = LandscapeEdMode->GetBrushList();
+		if (Brushes.ContainsByPredicate([](ALandscapeBlueprintCustomBrush* Brush) { return Brush->GetOwningLandscape() == nullptr; }))
+		{
+			MenuBuilder.BeginSection("LandscapeEditorBrushActions", LOCTEXT("LandscapeEditorBrushActions.Heading", "Brushes"));
+			{
+				MenuBuilder.AddSubMenu(
+					LOCTEXT("LandscaeEditorBrushAddSubMenu", "Add"),
+					LOCTEXT("LandscaeEditorBrushAddSubMenuToolTip", "Add brush to current layer"),
+					FNewMenuDelegate::CreateSP(this, &FLandscapeEditorCustomNodeBuilder_Layers::FillAddBrushMenu, Brushes),
+					false,
+					FSlateIcon()
+				);
+			}
+			MenuBuilder.EndSection();
+		}
+
 		return MenuBuilder.MakeWidget();
 	}
-	return NULL;
+	return nullptr;
 }
 
 void FLandscapeEditorCustomNodeBuilder_Layers::ForceUpdateSplines()
diff --git a/Engine/Source/Editor/LandscapeEditor/Private/LandscapeEditorDetailCustomization_Layers.h b/Engine/Source/Editor/LandscapeEditor/Private/LandscapeEditorDetailCustomization_Layers.h
--- a/Engine/Source/Editor/LandscapeEditor/Private/LandscapeEditorDetailCustomization_Layers.h
+++ b/Engine/Source/Editor/LandscapeEditor/Private/LandscapeEditorDetailCustomization_Layers.h
@@ -23,6 +23,7 @@ class FDetailWidgetRow;
 class IDetailChildrenBuilder;
 class IDetailLayoutBuilder;
 class SDragAndDropVerticalBox;
+class ALandscapeBlueprintCustomBrush;
 
 /**
  * Slate widgets customizer for the layers list in the Landscape Editor
@@ -97,6 +98,9 @@ protected:
 	FReply OnToggleLock(int32 InLayerIndex);
 	const FSlateBrush* GetLockBrushForLayer(int32 InLayerIndex) const;
 
+	void FillAddBrushMenu(FMenuBuilder& MenuBuilder, TArray<ALandscapeBlueprintCustomBrush*> Brushes);
+	void AddBrushToCurrentLayer(ALandscapeBlueprintCustomBrush* Brush);
+
 private:
 
 	/** Widgets for displaying and editing the layer name */
diff --git a/Engine/Source/Editor/LandscapeEditor/Private/LandscapeEditorDetailCustomization_LayersBrushStack.cpp b/Engine/Source/Editor/LandscapeEditor/Private/LandscapeEditorDetailCustomization_LayersBrushStack.cpp
--- a/Engine/Source/Editor/LandscapeEditor/Private/LandscapeEditorDetailCustomization_LayersBrushStack.cpp
+++ b/Engine/Source/Editor/LandscapeEditor/Private/LandscapeEditorDetailCustomization_LayersBrushStack.cpp
@@ -76,19 +76,6 @@ void FLandscapeEditorCustomNodeBuilder_LayersBrushStack::SetOnRebuildChildren(FS
 
 void FLandscapeEditorCustomNodeBuilder_LayersBrushStack::GenerateHeaderRowContent(FDetailWidgetRow& NodeRow)
 {
-	FEdModeLandscape* LandscapeEdMode = GetEditorMode();
-	
-	if (LandscapeEdMode == NULL)
-	{
-		return;	
-	}
-
-	NodeRow.NameWidget
-	[
-		SNew(STextBlock)
-		.Font(IDetailLayoutBuilder::GetDetailFont())
-		.Text(FText::FromString(TEXT("Stack")))
-	];
 }
 
 BEGIN_SLATE_FUNCTION_BUILD_OPTIMIZATION
@@ -144,10 +131,28 @@ TSharedPtr<SWidget> FLandscapeEditorCustomNodeBuilder_LayersBrushStack::Generate
 	TSharedPtr<SWidget> RowWidget = SNew(SLandscapeEditorSelectableBorder)
 		.Padding(0)
 		.VAlign(VAlign_Center)
+		.OnContextMenuOpening(this, &FLandscapeEditorCustomNodeBuilder_LayersBrushStack::OnBrushContextMenuOpening, InBrushIndex)
 		.OnSelected(this, &FLandscapeEditorCustomNodeBuilder_LayersBrushStack::OnBrushSelectionChanged, InBrushIndex)
 		.IsSelected(TAttribute<bool>::Create(TAttribute<bool>::FGetter::CreateSP(this, &FLandscapeEditorCustomNodeBuilder_LayersBrushStack::IsBrushSelected, InBrushIndex)))
 		[	
 			SNew(SHorizontalBox)
+			+ SHorizontalBox::Slot()
+			.AutoWidth()
+			.VAlign(VAlign_Center)
+			[
+				SNew(SButton)
+				.ContentPadding(0)
+				.ButtonStyle(FEditorStyle::Get(), "NoBorder")
+				.OnClicked(this, &FLandscapeEditorCustomNodeBuilder_LayersBrushStack::OnToggleVisibility, InBrushIndex)
+				.ToolTipText(LOCTEXT("LandscapeBrushVisibility", "Toggle Brush Visibility"))
+				.HAlign(HAlign_Center)
+				.VAlign(VAlign_Center)
+				.Content()
+				[
+					SNew(SImage)
+					.Image(this, &FLandscapeEditorCustomNodeBuilder_LayersBrushStack::GetVisibilityBrushForBrush, InBrushIndex)
+				]
+			]
 			+SHorizontalBox::Slot()
 			.VAlign(VAlign_Center)
 			.Padding(4, 0)
@@ -170,11 +175,57 @@ TSharedPtr<SWidget> FLandscapeEditorCustomNodeBuilder_LayersBrushStack::Generate
 }
 END_SLATE_FUNCTION_BUILD_OPTIMIZATION
 
+void FLandscapeEditorCustomNodeBuilder_LayersBrushStack::RemoveBrush(ALandscapeBlueprintCustomBrush* Brush)
+{
+	const FScopedTransaction Transaction(LOCTEXT("LandscapeBrushRemoveTransaction", "Remove Brush"));
+	GetEditorMode()->RemoveBrushFromCurrentLayer(Brush);
+}
+
+TSharedPtr<SWidget> FLandscapeEditorCustomNodeBuilder_LayersBrushStack::OnBrushContextMenuOpening(int32 InBrushIndex)
+{
+	if (ALandscapeBlueprintCustomBrush* CurrentBrush = GetBrush(InBrushIndex))
+	{
+		FMenuBuilder MenuBuilder(true, NULL);
+		MenuBuilder.BeginSection("LandscapeEditorBrushActions", LOCTEXT("LandscapeEditorBrushActions.Heading", "Brushes"));
+		{
+			TSharedRef<FLandscapeEditorCustomNodeBuilder_LayersBrushStack> SharedThis = AsShared();
+
+			// Remove Brush
+			FUIAction RemoveAction = FUIAction(FExecuteAction::CreateLambda([SharedThis, CurrentBrush] { SharedThis->RemoveBrush(CurrentBrush); }));
+			MenuBuilder.AddMenuEntry(LOCTEXT("LandscapeBrushRemove", "Remove"), LOCTEXT("LandscapeBrushRemoveToolTip", "Remove brush from current layer"), FSlateIcon(), RemoveAction);
+
+		
+		}
+		MenuBuilder.EndSection();
+		return MenuBuilder.MakeWidget();
+	}
+
+	return nullptr;
+}
+
+FReply FLandscapeEditorCustomNodeBuilder_LayersBrushStack::OnToggleVisibility(int32 InBrushIndex)
+{
+	if (ALandscapeBlueprintCustomBrush* Brush = GetBrush(InBrushIndex))
+	{
+		const FScopedTransaction Transaction(LOCTEXT("Landscape_Brush_SetVisibility", "Set Brush Visibility"));
+		bool bVisible = Brush->IsVisible();
+		Brush->SetIsVisible(!bVisible);
+	}
+	return FReply::Handled();
+}
+
+const FSlateBrush* FLandscapeEditorCustomNodeBuilder_LayersBrushStack::GetVisibilityBrushForBrush(int32 InBrushIndex) const
+{
+	ALandscapeBlueprintCustomBrush* Brush = GetBrush(InBrushIndex);
+	bool bIsVisible = Brush && Brush->IsVisible();
+	return bIsVisible ? FEditorStyle::GetBrush("Level.VisibleIcon16x") : FEditorStyle::GetBrush("Level.NotVisibleIcon16x");
+}
+
 bool FLandscapeEditorCustomNodeBuilder_LayersBrushStack::IsBrushEnabled(int32 InBrushIndex) const
 {
 	ALandscapeBlueprintCustomBrush* Brush = GetBrush(InBrushIndex);
-	FEdModeLandscape* LandscapeEdMode = GetEditorMode();
-	return Brush && LandscapeEdMode && ((Brush->IsAffectingHeightmap() && LandscapeEdMode->CurrentToolTarget.TargetType == ELandscapeToolTargetType::Type::Heightmap) || (Brush->IsAffectingWeightmap() && LandscapeEdMode->CurrentToolTarget.TargetType == ELandscapeToolTargetType::Type::Weightmap));
+	const FEdModeLandscape* LandscapeEdMode = GetEditorMode();
+	return Brush && ((Brush->IsAffectingHeightmap() && LandscapeEdMode->CurrentToolTarget.TargetType == ELandscapeToolTargetType::Type::Heightmap) || (Brush->IsAffectingWeightmap() && LandscapeEdMode->CurrentToolTarget.TargetType == ELandscapeToolTargetType::Type::Weightmap));
 }
 
 bool FLandscapeEditorCustomNodeBuilder_LayersBrushStack::IsBrushSelected(int32 InBrushIndex) const
diff --git a/Engine/Source/Editor/LandscapeEditor/Private/LandscapeEditorDetailCustomization_LayersBrushStack.h b/Engine/Source/Editor/LandscapeEditor/Private/LandscapeEditorDetailCustomization_LayersBrushStack.h
--- a/Engine/Source/Editor/LandscapeEditor/Private/LandscapeEditorDetailCustomization_LayersBrushStack.h
+++ b/Engine/Source/Editor/LandscapeEditor/Private/LandscapeEditorDetailCustomization_LayersBrushStack.h
@@ -56,8 +56,13 @@ protected:
 
 	static class FEdModeLandscape* GetEditorMode();
 
-	TSharedPtr<SWidget> GenerateRow(int32 InLayerIndex);
+	TSharedPtr<SWidget> GenerateRow(int32 InBrushIndex);
+	TSharedPtr<SWidget> OnBrushContextMenuOpening(int32 InBrushIndex);
 
+	void RemoveBrush(ALandscapeBlueprintCustomBrush* Brush);
+	
+	FReply OnToggleVisibility(int32 InBrushIndex);
+	const FSlateBrush* GetVisibilityBrushForBrush(int32 InBrushIndex) const;
 	bool IsBrushSelected(int32 InBrushIndex) const;
 	bool IsBrushEnabled(int32 InBrushIndex) const;
 	void OnBrushSelectionChanged(int32 InBrushIndex);
diff --git a/Engine/Source/Runtime/Landscape/Classes/Landscape.h b/Engine/Source/Runtime/Landscape/Classes/Landscape.h
--- a/Engine/Source/Runtime/Landscape/Classes/Landscape.h
+++ b/Engine/Source/Runtime/Landscape/Classes/Landscape.h
@@ -278,6 +278,7 @@ public:
 	LANDSCAPE_API void SetLayerSubstractiveBlendStatus(int32 InLayerIndex, bool InStatus, const TWeakObjectPtr<ULandscapeLayerInfoObject>& InLayerInfoObj);
 
 	LANDSCAPE_API void AddBrushToLayer(int32 InLayerIndex, class ALandscapeBlueprintCustomBrush* InBrush);
+	LANDSCAPE_API void RemoveBrush(class ALandscapeBlueprintCustomBrush* InBrush);
 	LANDSCAPE_API void RemoveBrushFromLayer(int32 InLayerIndex, class ALandscapeBlueprintCustomBrush* InBrush);
 	LANDSCAPE_API bool ReorderLayerBrush(int32 InLayerIndex, int32 InStartingLayerBrushIndex, int32 InDestinationLayerBrushIndex);
 	LANDSCAPE_API class ALandscapeBlueprintCustomBrush* GetBrushForLayer(int32 InLayerIndex, int8 BrushIndex) const;
diff --git a/Engine/Source/Runtime/Landscape/Classes/LandscapeBPCustomBrush.h b/Engine/Source/Runtime/Landscape/Classes/LandscapeBPCustomBrush.h
--- a/Engine/Source/Runtime/Landscape/Classes/LandscapeBPCustomBrush.h
+++ b/Engine/Source/Runtime/Landscape/Classes/LandscapeBPCustomBrush.h
@@ -24,7 +24,7 @@ private:
 	TArray<FName> AffectedWeightmapLayers;
 
 #if WITH_EDITORONLY_DATA
-	UPROPERTY(NonTransactional, DuplicateTransient)
+	UPROPERTY(DuplicateTransient)
 	class ALandscape* OwningLandscape;
 
 	UPROPERTY(NonTransactional, DuplicateTransient)
@@ -32,6 +32,9 @@ private:
 
 	UPROPERTY(Transient)
 	bool bIsInitialized;
+
+	UPROPERTY(Transient)
+	bool bIsVisible;
 #endif
 public:
 
@@ -55,11 +58,16 @@ public:
 	bool IsInitialized() const { return bIsInitialized; }
 	void SetIsInitialized(bool InIsInitialized);
 
+	bool IsVisible() const { return bIsVisible; }
+	void SetIsVisible(bool bInIsVisible);
+
 	void SetOwningLandscape(class ALandscape* InOwningLandscape);
 	class ALandscape* GetOwningLandscape() const;
 
 	virtual void PostEditMove(bool bFinished) override;
 	virtual void PostEditChangeProperty(FPropertyChangedEvent& PropertyChangedEvent) override;
+
+	virtual void Destroyed() override;
 #endif
 };
 
diff --git a/Engine/Source/Runtime/Landscape/Private/LandscapeBPCustomBrush.cpp b/Engine/Source/Runtime/Landscape/Private/LandscapeBPCustomBrush.cpp
--- a/Engine/Source/Runtime/Landscape/Private/LandscapeBPCustomBrush.cpp
+++ b/Engine/Source/Runtime/Landscape/Private/LandscapeBPCustomBrush.cpp
@@ -12,6 +12,7 @@ ALandscapeBlueprintCustomBrush::ALandscapeBlueprintCustomBrush(const FObjectInit
 	: OwningLandscape(nullptr)
 	, bIsCommited(false)
 	, bIsInitialized(false)
+	, bIsVisible(true)
 #endif
 {
 	USceneComponent* SceneComp = CreateDefaultSubobject<USceneComponent>(TEXT("RootComponent"));
@@ -52,8 +53,16 @@ void ALandscapeBlueprintCustomBrush::SetCommitState(bool InCommited)
 #endif
 }
 
+void ALandscapeBlueprintCustomBrush::SetIsVisible(bool bInIsVisible)
+{
+	Modify();
+	bIsVisible = bInIsVisible;
+	GetOwningLandscape()->OnBPCustomBrushChanged();
+}
+
 void ALandscapeBlueprintCustomBrush::SetOwningLandscape(ALandscape* InOwningLandscape)
 {
+	Modify();
 	OwningLandscape = InOwningLandscape;
 }
 
@@ -97,6 +106,17 @@ void ALandscapeBlueprintCustomBrush::PostEditChangeProperty(FPropertyChangedEven
 		OwningLandscape->OnBPCustomBrushChanged();
 	}
 }
+
+void ALandscapeBlueprintCustomBrush::Destroyed()
+{
+	Super::Destroyed();
+
+	if (OwningLandscape)
+	{
+		OwningLandscape->RemoveBrush(this);
+	}
+}
+
 #endif
 
 ALandscapeBlueprintCustomSimulationBrush::ALandscapeBlueprintCustomSimulationBrush(const FObjectInitializer& ObjectInitializer)
diff --git a/Engine/Source/Runtime/Landscape/Private/LandscapeEditLayers.cpp b/Engine/Source/Runtime/Landscape/Private/LandscapeEditLayers.cpp
--- a/Engine/Source/Runtime/Landscape/Private/LandscapeEditLayers.cpp
+++ b/Engine/Source/Runtime/Landscape/Private/LandscapeEditLayers.cpp
@@ -2738,7 +2738,7 @@ int32 ALandscape::RegenerateLayersHeightmaps(const TArray<ULandscapeComponent*>&
 
 					FLandscapeLayerBrush& Brush = Layer.Brushes[i];
 
-					if (Brush.BPCustomBrush == nullptr || !Brush.BPCustomBrush->IsAffectingHeightmap())
+					if (Brush.BPCustomBrush == nullptr || !Brush.BPCustomBrush->IsAffectingHeightmap() || !Brush.BPCustomBrush->IsVisible())
 					{
 						continue;
 					}
@@ -3473,7 +3473,7 @@ int32 ALandscape::RegenerateLayersWeightmaps(const TArray<ULandscapeComponent*>&
 					{
 						FLandscapeLayerBrush& Brush = Layer.Brushes[i];
 
-						if (Brush.BPCustomBrush == nullptr)
+						if (Brush.BPCustomBrush == nullptr || !Brush.BPCustomBrush->IsAffectingWeightmap() || !Brush.BPCustomBrush->IsVisible())
 						{
 							continue;
 						}
@@ -3593,7 +3593,7 @@ int32 ALandscape::RegenerateLayersWeightmaps(const TArray<ULandscapeComponent*>&
 
 							FLandscapeLayerBrush& Brush = Layer.Brushes[i];
 
-							if (Brush.BPCustomBrush == nullptr || !Brush.BPCustomBrush->IsAffectingWeightmap() || !Brush.BPCustomBrush->IsAffectingWeightmapLayer(LayerInfoObj->LayerName))
+							if (Brush.BPCustomBrush == nullptr || !Brush.BPCustomBrush->IsAffectingWeightmap() || !Brush.BPCustomBrush->IsAffectingWeightmapLayer(LayerInfoObj->LayerName) || !Brush.BPCustomBrush->IsVisible())
 							{
 								continue;
 							}
@@ -5446,6 +5446,14 @@ void ALandscape::AddBrushToLayer(int32 InLayerIndex, ALandscapeBlueprintCustomBr
 	}
 }
 
+void ALandscape::RemoveBrush(ALandscapeBlueprintCustomBrush* InBrush)
+{
+	for (int32 LayerIndex = 0; LayerIndex < LandscapeLayers.Num(); ++LayerIndex)
+	{
+		RemoveBrushFromLayer(LayerIndex, InBrush);
+	}
+}
+
 void ALandscape::RemoveBrushFromLayer(int32 InLayerIndex, ALandscapeBlueprintCustomBrush* InBrush)
 {
 	if (FLandscapeLayer* Layer = GetLayer(InLayerIndex))
