- Landscape Smooth Tool now supports also working on combined layers (new option bCombinedLayersOperation)
- Added class FLandscapeLayerDataCache to handle caching combined layers data and compute delta when writing
- Refactored Flatten Tool to use new FLandscapeLayerDataCache
- Optimized greatly performance of Flatten Tool in CombinedLayers mode
#rnx

diff --git a/Engine/Source/Editor/LandscapeEditor/Private/LandscapeEdModePaintTools.cpp b/Engine/Source/Editor/LandscapeEditor/Private/LandscapeEdModePaintTools.cpp
--- a/Engine/Source/Editor/LandscapeEditor/Private/LandscapeEdModePaintTools.cpp
+++ b/Engine/Source/Editor/LandscapeEditor/Private/LandscapeEdModePaintTools.cpp
@@ -593,6 +593,180 @@ public:
 	virtual FText GetDisplayName() override { return NSLOCTEXT("UnrealEd", "LandscapeMode_Erase", "Erase"); };
 };
 
+template<class ToolTarget>
+class FLandscapeLayerDataCache
+{
+public:
+	FLandscapeLayerDataCache(const FLandscapeToolTarget& InTarget, typename ToolTarget::CacheClass& Cache)
+		: LandscapeInfo(nullptr)
+		, Landscape(nullptr)
+		, EditingLayerIndex(MAX_uint8)
+		, bIsInitialized(false)
+		, bCombinedLayerOperation(false)
+		, bVisibilityChanged(false)
+		, bTargetIsHeightmap(InTarget.TargetType == ELandscapeToolTargetType::Heightmap)
+		, CacheUpToEditingLayer(Cache)
+		, CacheBottomLayers(InTarget)
+	{
+	}
+
+	void Initialize(ULandscapeInfo* InLandscapeInfo, bool InCombinedLayerOperation)
+	{
+		if (!bIsInitialized)
+		{
+			LandscapeInfo = InLandscapeInfo;
+			Landscape = LandscapeInfo ? LandscapeInfo->LandscapeActor.Get() : nullptr;
+			bCombinedLayerOperation = Landscape && Landscape->HasLayersContent() && InCombinedLayerOperation && bTargetIsHeightmap;
+			if (bCombinedLayerOperation)
+			{
+				EditingLayerGuid = Landscape->GetEditingLayer();
+				for (int i = 0; i < Landscape->GetLayerCount(); ++i)
+				{
+					FLandscapeLayer* CurrentLayer = Landscape->GetLayer(i);
+					BackupLayerVisibility.Add(CurrentLayer->bVisible);
+					if (CurrentLayer->Guid == EditingLayerGuid)
+					{
+						EditingLayerIndex = i;
+					}
+				}
+				check(EditingLayerIndex < Landscape->GetLayerCount());
+			}
+			bIsInitialized = true;
+		}
+	}
+
+	void Read(int32 X1, int32 Y1, int32 X2, int32 Y2, TArray<typename ToolTarget::CacheClass::DataType>& Data)
+	{
+		check(bIsInitialized);
+		if (bCombinedLayerOperation)
+		{
+			TArray<bool> NewLayerVisibility;
+			for (int i = 0; i < Landscape->GetLayerCount(); ++i)
+			{
+				FLandscapeLayer* CurrentLayer = Landscape->GetLayer(i);
+				NewLayerVisibility.Add((i > EditingLayerIndex) ? false : CurrentLayer->bVisible);
+			}
+
+			FIntRect Bounds;
+			const float BoundStep = LandscapeInfo->ComponentSizeQuads / 4;
+			Bounds.Min.X = (FMath::FloorToInt((float)X1 / BoundStep) - 1) * BoundStep;
+			Bounds.Min.Y = (FMath::FloorToInt((float)Y1 / BoundStep) - 1) * BoundStep;
+			Bounds.Max.X = (FMath::CeilToInt((float)X2 / BoundStep) + 1) * BoundStep;
+			Bounds.Max.Y = (FMath::CeilToInt((float)Y2 / BoundStep) + 1) * BoundStep;
+
+			FScopedSetLandscapeEditingLayer Scope(Landscape, FGuid());
+			CacheUpToEditingLayer.GetDataAndCache(X1, Y1, X2, Y2, Data, [&]() -> FIntRect
+			{
+				TSet<ULandscapeComponent*> AffectedComponents;
+				LandscapeInfo->GetComponentsInRegion(X1, Y1, X2, Y2, AffectedComponents);
+				SynchronousUpdateComponentVisibilityForHeight(AffectedComponents, NewLayerVisibility);
+				return Bounds;
+			});
+
+			CacheBottomLayers.GetDataAndCache(X1, Y1, X2, Y2, BottomLayersData, [&]() -> FIntRect
+			{
+				NewLayerVisibility[EditingLayerIndex] = false;
+				TSet<ULandscapeComponent*> AffectedComponents;
+				LandscapeInfo->GetComponentsInRegion(X1, Y1, X2, Y2, AffectedComponents);
+				SynchronousUpdateComponentVisibilityForHeight(AffectedComponents, NewLayerVisibility);
+				return Bounds;
+			});
+		}
+		else
+		{
+			CacheUpToEditingLayer.CacheData(X1, Y1, X2, Y2);
+			CacheUpToEditingLayer.GetCachedData(X1, Y1, X2, Y2, Data);
+		}
+	}
+
+	void Write(int32 X1, int32 Y1, int32 X2, int32 Y2, TArray<typename ToolTarget::CacheClass::DataType>& Data, ELandscapeLayerPaintingRestriction PaintingRestriction)
+	{
+		check(bIsInitialized);
+		if (bCombinedLayerOperation)
+		{
+			const float Alpha = Landscape->GetLayerAlpha(EditingLayerIndex, bTargetIsHeightmap);
+			const float InverseAlpha = (Alpha != 0.f) ? 1.f / Alpha : 1.f;
+			TArray<typename ToolTarget::CacheClass::DataType> DataContribution;
+			DataContribution.Empty(Data.Num());
+			DataContribution.AddUninitialized(Data.Num());
+			check(Data.Num() == BottomLayersData.Num());
+			for (int i = 0; i < Data.Num(); ++i)
+			{
+				float Contribution = (LandscapeDataAccess::GetLocalHeight(Data[i]) - LandscapeDataAccess::GetLocalHeight(BottomLayersData[i])) * InverseAlpha;
+				DataContribution[i] = LandscapeDataAccess::GetTexHeight(Contribution);
+			}
+			check(EditingLayerGuid == Landscape->GetEditingLayer());
+
+			// Restore layers visibility
+			SetLayersVisibility(BackupLayerVisibility);
+			// Only store data in cache
+			CacheUpToEditingLayer.SetCachedData(X1, Y1, X2, Y2, Data, PaintingRestriction, false);
+			// Effectively write the contribution
+			CacheUpToEditingLayer.DataAccess.SetData(X1, Y1, X2, Y2, DataContribution.GetData(), PaintingRestriction);
+			CacheUpToEditingLayer.DataAccess.Flush();
+			if (bVisibilityChanged)
+			{
+				TSet<ULandscapeComponent*> AffectedComponents;
+				LandscapeInfo->GetComponentsInRegion(X1, Y1, X2, Y2, AffectedComponents);
+				SynchronousUpdateHeightmapForComponents(AffectedComponents);
+				bVisibilityChanged = false;
+			}
+		}
+		else
+		{
+			CacheUpToEditingLayer.SetCachedData(X1, Y1, X2, Y2, Data, PaintingRestriction);
+			CacheUpToEditingLayer.Flush();
+		}
+	}
+
+private:
+
+	void SynchronousUpdateHeightmapForComponents(const TSet<ULandscapeComponent*>& InComponents)
+	{
+		for (ULandscapeComponent* Component : InComponents)
+		{
+			Component->RequestHeightmapUpdate();
+		}
+		Landscape->ForceUpdateLayersContent();
+	};
+
+	void SetLayersVisibility(const TArray<bool>& InLayerVisibility)
+	{
+		check(InLayerVisibility.Num() == Landscape->GetLayerCount());
+		for (int i = 0; i < InLayerVisibility.Num(); ++i)
+		{
+			if (FLandscapeLayer* Layer = Landscape->GetLayer(i))
+			{
+				if (Layer->bVisible != InLayerVisibility[i])
+				{
+					Layer->bVisible = InLayerVisibility[i];
+					bVisibilityChanged = true;
+				}
+			}
+		}
+	};
+
+	void SynchronousUpdateComponentVisibilityForHeight(const TSet<ULandscapeComponent*>& InComponents, const TArray<bool>& InLayerVisibility)
+	{
+		SetLayersVisibility(InLayerVisibility);
+		SynchronousUpdateHeightmapForComponents(InComponents);
+	};
+
+	ULandscapeInfo* LandscapeInfo;
+	ALandscape* Landscape;
+	FGuid EditingLayerGuid;
+	uint8 EditingLayerIndex;
+	TArray<bool> BackupLayerVisibility;
+	TArray<typename ToolTarget::CacheClass::DataType> BottomLayersData;
+	bool bIsInitialized;
+	bool bCombinedLayerOperation;
+	bool bVisibilityChanged;
+	bool bTargetIsHeightmap;
+
+	typename ToolTarget::CacheClass& CacheUpToEditingLayer;
+	typename ToolTarget::CacheClass CacheBottomLayers;
+};
+
 // 
 // FLandscapeToolSmooth
 //
@@ -600,9 +774,13 @@ public:
 template<class ToolTarget>
 class FLandscapeToolStrokeSmooth : public FLandscapeToolStrokePaintBase<ToolTarget>
 {
+	bool bTargetIsHeightmap;
+	FLandscapeLayerDataCache<ToolTarget> LayerDataCache;
 public:
 	FLandscapeToolStrokeSmooth(FEdModeLandscape* InEdMode, FEditorViewportClient* InViewportClient, const FLandscapeToolTarget& InTarget)
 		: FLandscapeToolStrokePaintBase<ToolTarget>(InEdMode, InViewportClient, InTarget)
+		, bTargetIsHeightmap(InTarget.TargetType == ELandscapeToolTargetType::Heightmap)
+		, LayerDataCache(InTarget, this->Cache)
 	{
 	}
 
@@ -610,6 +788,9 @@ public:
 	{
 		if (!this->LandscapeInfo) return;
 
+		ALandscape* Landscape = this->LandscapeInfo->LandscapeActor.Get();
+		const bool bCombinedLayerOperation = bTargetIsHeightmap && UISettings->bCombinedLayersOperation && Landscape && Landscape->HasLayersContent();
+
 		// Get list of verts to update
 		FLandscapeBrushData BrushInfo = Brush->ApplyBrush(InteractorPositions);
 		if (!BrushInfo)
@@ -633,11 +814,9 @@ public:
 			Y2 += 1;
 		}
 
-		this->Cache.CacheData(X1, Y1, X2, Y2);
-
 		TArray<typename ToolTarget::CacheClass::DataType> Data;
-		this->Cache.GetCachedData(X1, Y1, X2, Y2, Data);
-
+		LayerDataCache.Initialize(LandscapeInfo, bCombinedLayerOperation);
+		LayerDataCache.Read(X1, Y1, X2, Y2, Data);
 		const float ToolStrength = FMath::Clamp<float>(UISettings->ToolStrength * Pressure, 0.0f, 1.0f);
 
 		// Apply the brush
@@ -670,8 +849,8 @@ public:
 
 						const int32 SampleX1 = X - XRadius; checkSlow(SampleX1 >= BrushInfo.GetBounds().Min.X);
 						const int32 SampleY1 = Y - YRadius; checkSlow(SampleY1 >= BrushInfo.GetBounds().Min.Y);
-						const int32 SampleX2 = X + XRadius; checkSlow(SampleX2 <  BrushInfo.GetBounds().Max.X);
-						const int32 SampleY2 = Y + YRadius; checkSlow(SampleY2 <  BrushInfo.GetBounds().Max.Y);
+						const int32 SampleX2 = X + XRadius; checkSlow(SampleX2 < BrushInfo.GetBounds().Max.X);
+						const int32 SampleY2 = Y + YRadius; checkSlow(SampleY2 < BrushInfo.GetBounds().Max.Y);
 						for (int32 SampleY = SampleY1; SampleY <= SampleY2; SampleY++)
 						{
 							const float* SampleBrushScanline = BrushInfo.GetDataPtr(FIntPoint(0, SampleY));
@@ -683,9 +862,9 @@ public:
 								// constrain sample to within the brush, symmetrically to prevent flattening bug
 								const float SampleBrushValue =
 									FMath::Min(
-										FMath::Min<float>(SampleBrushScanline [SampleX], SampleBrushScanline [X + (X - SampleX)]),
+										FMath::Min<float>(SampleBrushScanline[SampleX], SampleBrushScanline[X + (X - SampleX)]),
 										FMath::Min<float>(SampleBrushScanline2[SampleX], SampleBrushScanline2[X + (X - SampleX)])
-										);
+									);
 								if (SampleBrushValue > 0.0f)
 								{
 									FilterValue += SampleDataScanline[SampleX];
@@ -702,8 +881,7 @@ public:
 			}
 		}
 
-		this->Cache.SetCachedData(X1, Y1, X2, Y2, Data, UISettings->PaintingRestriction);
-		this->Cache.Flush();
+		LayerDataCache.Write(X1, Y1, X2, Y2, Data, UISettings->PaintingRestriction);
 	}
 };
 
@@ -733,12 +911,14 @@ class FLandscapeToolStrokeFlatten : public FLandscapeToolStrokePaintBase<ToolTar
 	float FlattenPlaneDist;
 	bool bInitializedFlattenValue;
 	bool bTargetIsHeightmap;
+	FLandscapeLayerDataCache<ToolTarget> LayerDataCache;
 
 public:
 	FLandscapeToolStrokeFlatten(FEdModeLandscape* InEdMode, FEditorViewportClient* InViewportClient, const FLandscapeToolTarget& InTarget)
 		: FLandscapeToolStrokePaintBase<ToolTarget>(InEdMode, InViewportClient, InTarget)
 		, bInitializedFlattenValue(false)
 		, bTargetIsHeightmap(InTarget.TargetType == ELandscapeToolTargetType::Heightmap)
+		, LayerDataCache(InTarget, this->Cache)
 	{
 		if (InEdMode->UISettings->bUseFlattenTarget && bTargetIsHeightmap)
 		{
@@ -752,10 +932,9 @@ public:
 	void Apply(FEditorViewportClient* ViewportClient, FLandscapeBrush* Brush, const ULandscapeEditorObject* UISettings, const TArray<FLandscapeToolInteractorPosition>& InteractorPositions)
  	{
 		if (!this->LandscapeInfo) return;
-		bool bLayerSystemFlattenHeightMode = this->LandscapeInfo->CanHaveLayersContent() && bTargetIsHeightmap;
+
 		ALandscape* Landscape = this->LandscapeInfo->LandscapeActor.Get();
-		if (!Landscape && bLayerSystemFlattenHeightMode)
-			return;
+		const bool bCombinedLayerOperation = bTargetIsHeightmap && UISettings->bCombinedLayersOperation && Landscape && Landscape->HasLayersContent();
 
 		// Can't use slope if use Flatten Target because no normal is provided
 		bool bUseSlopeFlatten = UISettings->bUseSlopeFlatten && !UISettings->bUseFlattenTarget;
@@ -773,7 +952,7 @@ public:
 			float V10 = 0.f;
 			float V01 = 0.f;
 			float V11 = 0.f;
-			if (bLayerSystemFlattenHeightMode)
+			if (bCombinedLayerOperation)
 			{
 				// Can't rely on cache in this mode
 				FScopedSetLandscapeEditingLayer Scope(Landscape, FGuid());
@@ -797,7 +976,7 @@ public:
 
 			if (bUseSlopeFlatten && bTargetIsHeightmap)
 			{
-				if (bLayerSystemFlattenHeightMode)
+				if (bCombinedLayerOperation)
 				{
 					// Can't rely on cache in this mode
 					FVector Vert00 = FVector(0.0f, 0.0f, V00);
@@ -843,72 +1022,8 @@ public:
 		}
 
 		TArray<typename ToolTarget::CacheClass::DataType> Data;
-		TArray<typename ToolTarget::CacheClass::DataType> BottomLayersData;
-		FGuid EditingLayer;
-		uint8 EditingLayerIndex = MAX_uint8;
-		TSet<ULandscapeComponent*> AffectedComponents;
-		TArray<bool> BackupLayerVisibility;
-		auto& LayerSystemDataAccess = this->Cache.DataAccess;
-
-		auto SynchronousUpdateHeightmapForComponents = [&](const TSet<ULandscapeComponent*>& InComponents)
-		{
-			for (ULandscapeComponent* Component : InComponents)
-			{
-				Component->RequestHeightmapUpdate();
-			}
-			Landscape->ForceUpdateLayersContent();
-		};
-
-		auto SetLayersVisibility = [&](const TArray<bool>& InLayerVisibility)
-		{
-			check(InLayerVisibility.Num() == Landscape->GetLayerCount());
-			for (int i = 0; i < InLayerVisibility.Num(); ++i)
-			{
-				if (FLandscapeLayer* Layer = Landscape->GetLayer(i))
-				{
-					Layer->bVisible = InLayerVisibility[i];
-				}
-			}
-		};
-
-		auto SynchronousUpdateComponentVisibilityForHeight = [&](const TSet<ULandscapeComponent*>& InComponents, const TArray<bool>& InLayerVisibility)
-		{
-			SetLayersVisibility(InLayerVisibility);
-			SynchronousUpdateHeightmapForComponents(InComponents);
-		};
-
-		if (bLayerSystemFlattenHeightMode)
-		{
-			EditingLayer = Landscape->GetEditingLayer();
-			this->LandscapeInfo->GetComponentsInRegion(X1, Y1, X2, Y2, AffectedComponents);
-			TArray<bool> NewLayerVisibility;
-			for (int i = 0; i < Landscape->GetLayerCount(); ++i)
-			{
-				FLandscapeLayer* CurrentLayer = Landscape->GetLayer(i);
-				BackupLayerVisibility.Add(CurrentLayer->bVisible);
-				NewLayerVisibility.Add((i > EditingLayerIndex) ? false : CurrentLayer->bVisible);
-				if (CurrentLayer->Guid == EditingLayer)
-				{
-					EditingLayerIndex = i;
-				}
-			}
-			check(EditingLayerIndex < Landscape->GetLayerCount());
-			
-			FScopedSetLandscapeEditingLayer Scope(Landscape, FGuid());
-			SynchronousUpdateComponentVisibilityForHeight(AffectedComponents, NewLayerVisibility);
-			Data.AddZeroed((X2 - X1 + 1) * (Y2 - Y1 + 1));
-			LayerSystemDataAccess.GetDataFast(X1, Y1, X2, Y2, Data.GetData());
-			
-			BottomLayersData.AddZeroed((X2 - X1 + 1) * (Y2 - Y1 + 1));
-			NewLayerVisibility[EditingLayerIndex] = false;
-			SynchronousUpdateComponentVisibilityForHeight(AffectedComponents, NewLayerVisibility);
-			LayerSystemDataAccess.GetDataFast(X1, Y1, X2, Y2, BottomLayersData.GetData());
-		}
-		else
-		{
-			this->Cache.CacheData(X1, Y1, X2, Y2);
-			this->Cache.GetCachedData(X1, Y1, X2, Y2, Data);
-		}
+		LayerDataCache.Initialize(LandscapeInfo, bCombinedLayerOperation);
+		LayerDataCache.Read(X1, Y1, X2, Y2, Data);
 
 		// Apply the brush
 		for (int32 Y = BrushInfo.GetBounds().Min.Y; Y < BrushInfo.GetBounds().Max.Y; Y++)
@@ -1045,27 +1160,7 @@ public:
 			}
 		}
 
-		if (bLayerSystemFlattenHeightMode)
-		{
-			const float Alpha = Landscape->GetLayerAlpha(EditingLayerIndex, bTargetIsHeightmap);
-			const float InverseAlpha = (Alpha != 0.f) ? 1.f / Alpha : 1.f;
-			for (int i = 0; i < Data.Num(); ++i)
-			{
-				float Diff = (LandscapeDataAccess::GetLocalHeight(Data[i]) - LandscapeDataAccess::GetLocalHeight(BottomLayersData[i])) * InverseAlpha;
-				Data[i] = LandscapeDataAccess::GetTexHeight(Diff);
-			}
-			check(EditingLayer == Landscape->GetEditingLayer());
-			
-			SetLayersVisibility(BackupLayerVisibility);
-			LayerSystemDataAccess.SetData(X1, Y1, X2, Y2, Data.GetData(), UISettings->PaintingRestriction);
-			LayerSystemDataAccess.Flush();
-			SynchronousUpdateHeightmapForComponents(AffectedComponents);
-		}
-		else
-		{
-			this->Cache.SetCachedData(X1, Y1, X2, Y2, Data, UISettings->PaintingRestriction);
-			this->Cache.Flush();
-		}
+		LayerDataCache.Write(X1, Y1, X2, Y2, Data, UISettings->PaintingRestriction);
 	}
 };
 
diff --git a/Engine/Source/Editor/LandscapeEditor/Private/LandscapeEdModeTools.h b/Engine/Source/Editor/LandscapeEditor/Private/LandscapeEdModeTools.h
--- a/Engine/Source/Editor/LandscapeEditor/Private/LandscapeEdModeTools.h
+++ b/Engine/Source/Editor/LandscapeEditor/Private/LandscapeEdModeTools.h
@@ -292,7 +292,7 @@ public:
 	}
 
 	// X2/Y2 Coordinates are "inclusive" max values
-	void CacheData(int32 X1, int32 Y1, int32 X2, int32 Y2)
+	void CacheData(int32 X1, int32 Y1, int32 X2, int32 Y2, bool bCacheOriginalData = true)
 	{
 		if (!Valid)
 		{
@@ -320,8 +320,10 @@ public:
 				DataAccess.GetDataFast(CachedX1, CachedY1, CachedX2, CachedY2, CachedData);
 			}
 
-			OriginalData = CachedData;
-
+			if (bCacheOriginalData)
+			{
+				OriginalData = CachedData;
+			}
 			Valid = true;
 		}
 		else
@@ -344,7 +346,10 @@ public:
 					DataAccess.GetDataFast(X1, CachedY1, CachedX1 - 1, CachedY2, CachedData);
 				}
 
-				CacheOriginalData(X1, CachedY1, CachedX1 - 1, CachedY2);
+				if (bCacheOriginalData)
+				{
+					CacheOriginalData(X1, CachedY1, CachedX1 - 1, CachedY2);
+				}
 				CachedX1 = X1;
 			}
 
@@ -364,7 +369,10 @@ public:
 				{
 					DataAccess.GetDataFast(CachedX2 + 1, CachedY1, X2, CachedY2, CachedData);
 				}
-				CacheOriginalData(CachedX2 + 1, CachedY1, X2, CachedY2);
+				if (bCacheOriginalData)
+				{
+					CacheOriginalData(CachedX2 + 1, CachedY1, X2, CachedY2);
+				}
 				CachedX2 = X2;
 			}
 
@@ -384,7 +392,10 @@ public:
 				{
 					DataAccess.GetDataFast(CachedX1, Y1, CachedX2, CachedY1 - 1, CachedData);
 				}
-				CacheOriginalData(CachedX1, Y1, CachedX2, CachedY1 - 1);
+				if (bCacheOriginalData)
+				{
+					CacheOriginalData(CachedX1, Y1, CachedX2, CachedY1 - 1);
+				}
 				CachedY1 = Y1;
 			}
 
@@ -404,8 +415,10 @@ public:
 				{
 					DataAccess.GetDataFast(CachedX1, CachedY2 + 1, CachedX2, Y2, CachedData);
 				}
-
-				CacheOriginalData(CachedX1, CachedY2 + 1, CachedX2, Y2);
+				if (bCacheOriginalData)
+				{
+					CacheOriginalData(CachedX1, CachedY2 + 1, CachedX2, Y2);
+				}
 				CachedY2 = Y2;
 			}
 		}
@@ -485,6 +498,25 @@ public:
 		return Value == 0;
 	}
 
+	bool HasCachedData(int32 X1, int32 Y1, int32 X2, int32 Y2) const
+	{
+		return (Valid && X1 >= CachedX1 && Y1 >= CachedY1 && X2 <= CachedX2 && Y2 <= CachedY2);
+	}
+
+	template<typename TGetCacheRegionFunction>
+	bool GetDataAndCache(int32 X1, int32 Y1, int32 X2, int32 Y2, TArray<AccessorType>& OutData, TGetCacheRegionFunction GetCacheRegion)
+	{
+		if (!HasCachedData(X1, Y1, X2, Y2))
+		{
+			FIntRect Bounds = GetCacheRegion();
+			check((Bounds.Min.X <= X1) && (Bounds.Min.Y <= Y1) && (Bounds.Max.X >= X2) && (Bounds.Max.Y >= Y2));
+			const bool bCacheOriginalData = false;
+			CacheData(Bounds.Min.X, Bounds.Min.Y, Bounds.Max.X, Bounds.Max.Y, bCacheOriginalData);
+		}
+		ensure(HasCachedData(X1, Y1, X2, Y2));
+		return GetCachedData(X1, Y1, X2, Y2, OutData);
+	}
+
 	// X2/Y2 Coordinates are "inclusive" max values
 	bool GetCachedData(int32 X1, int32 Y1, int32 X2, int32 Y2, TArray<AccessorType>& OutData)
 	{
@@ -521,7 +553,7 @@ public:
 	}
 
 	// X2/Y2 Coordinates are "inclusive" max values
-	void SetCachedData(int32 X1, int32 Y1, int32 X2, int32 Y2, TArray<AccessorType>& Data, ELandscapeLayerPaintingRestriction PaintingRestriction = ELandscapeLayerPaintingRestriction::None)
+	void SetCachedData(int32 X1, int32 Y1, int32 X2, int32 Y2, TArray<AccessorType>& Data, ELandscapeLayerPaintingRestriction PaintingRestriction = ELandscapeLayerPaintingRestriction::None, bool bUpdateData = true)
 	{
 		checkSlow(Data.Num() == (1 + Y2 - Y1) * (1 + X2 - X1));
 
@@ -534,8 +566,11 @@ public:
 			}
 		}
 
-		// Update real data
-		DataAccess.SetData(X1, Y1, X2, Y2, Data.GetData(), PaintingRestriction);
+		if (bUpdateData)
+		{
+			// Update real data
+			DataAccess.SetData(X1, Y1, X2, Y2, Data.GetData(), PaintingRestriction);
+		}
 	}
 
 	// Get the original data before we made any changes with the SetCachedData interface.
diff --git a/Engine/Source/Editor/LandscapeEditor/Private/LandscapeEditorObject.cpp b/Engine/Source/Editor/LandscapeEditor/Private/LandscapeEditorObject.cpp
--- a/Engine/Source/Editor/LandscapeEditor/Private/LandscapeEditorObject.cpp
+++ b/Engine/Source/Editor/LandscapeEditor/Private/LandscapeEditorObject.cpp
@@ -21,6 +21,7 @@ ULandscapeEditorObject::ULandscapeEditorObject(const FObjectInitializer& ObjectI
 	, bUseWeightTargetValue(false)
 	, WeightTargetValue(1.0f)
 	, MaximumValueRadius(10000.0f)
+	, bCombinedLayersOperation(true)
 
 	, FlattenMode(ELandscapeToolFlattenMode::Both)
 	, bUseSlopeFlatten(false)
diff --git a/Engine/Source/Editor/LandscapeEditor/Private/SLandscapeEditor.cpp b/Engine/Source/Editor/LandscapeEditor/Private/SLandscapeEditor.cpp
--- a/Engine/Source/Editor/LandscapeEditor/Private/SLandscapeEditor.cpp
+++ b/Engine/Source/Editor/LandscapeEditor/Private/SLandscapeEditor.cpp
@@ -534,6 +534,13 @@ bool SLandscapeEditor::GetIsPropertyVisible(const FPropertyAndParent& PropertyAn
 				return false;
 			}
 		}
+		if (Property.HasMetaData("ShowForLandscapeLayerSystem"))
+		{
+			if (!LandscapeEdMode->HasLandscapeLayersContent())
+			{
+				return false;
+			}
+		}
 
 		return true;
 	}
diff --git a/Engine/Source/Editor/LandscapeEditor/Public/LandscapeEditorObject.h b/Engine/Source/Editor/LandscapeEditor/Public/LandscapeEditorObject.h
--- a/Engine/Source/Editor/LandscapeEditor/Public/LandscapeEditorObject.h
+++ b/Engine/Source/Editor/LandscapeEditor/Public/LandscapeEditorObject.h
@@ -266,6 +266,9 @@ class ULandscapeEditorObject : public UObject
 	UPROPERTY(NonTransactional)
 	float MaximumValueRadius;
 
+	UPROPERTY(Category="Tool Settings", EditAnywhere, NonTransactional, meta=(ShowForTools="Flatten,Smooth", ShowForTargetTypes="Heightmap", ShowForLandscapeLayerSystem))
+	bool bCombinedLayersOperation;
+
 	// Flatten Tool:
 
 	// Whether to flatten by lowering, raising, both or terracing
