- Landscape Layer System:
New Toggles for Affects Height/Weight
Show Brush Stack always
More brush Undo/Redo fixes
New BP Tool Icons

#rnx

diff --git a/Engine/Source/Editor/EditorStyle/Private/SlateEditorStyle.cpp b/Engine/Source/Editor/EditorStyle/Private/SlateEditorStyle.cpp
--- a/Engine/Source/Editor/EditorStyle/Private/SlateEditorStyle.cpp
+++ b/Engine/Source/Editor/EditorStyle/Private/SlateEditorStyle.cpp
@@ -6669,7 +6669,7 @@ void FSlateEditorStyle::FStyle::SetupLandscapeEditorStyle()
 		Set("LandscapeEditor.ManageMode.Small", new IMAGE_BRUSH("Icons/icon_Landscape_Mode_Manage_20x", Icon20x20));
 		Set("LandscapeEditor.SculptMode.Small", new IMAGE_BRUSH("Icons/icon_Landscape_Mode_Sculpt_20x", Icon20x20));
 		Set("LandscapeEditor.PaintMode.Small",  new IMAGE_BRUSH("Icons/icon_Landscape_Mode_Paint_20x",  Icon20x20));
-
+		
 		// Tools
 		Set("LandscapeEditor.NewLandscape",       new IMAGE_BRUSH("Icons/icon_Landscape_New_Landscape_40x", Icon40x40));
 		Set("LandscapeEditor.NewLandscape.Small", new IMAGE_BRUSH("Icons/icon_Landscape_New_Landscape_20x", Icon20x20));
@@ -6687,7 +6687,7 @@ void FSlateEditorStyle::FStyle::SetupLandscapeEditorStyle()
 		Set("LandscapeEditor.NoiseTool",        new IMAGE_BRUSH("Icons/icon_Landscape_Tool_Noise_40x",            Icon40x40));
 		Set("LandscapeEditor.RetopologizeTool", new IMAGE_BRUSH("Icons/icon_Landscape_Tool_Retopologize_40x",     Icon40x40));
 		Set("LandscapeEditor.VisibilityTool",   new IMAGE_BRUSH("Icons/icon_Landscape_Tool_Visibility_40x",       Icon40x40));
-		Set("LandscapeEditor.BPCustomTool", new IMAGE_BRUSH("Icons/icon_Landscape_Tool_Visibility_40x", Icon40x40));// TODO: change icon
+		Set("LandscapeEditor.BlueprintBrushTool", new IMAGE_BRUSH("Icons/icon_Landscape_Tool_BlueprintBrush_40x", Icon40x40));
 		Set("LandscapeEditor.SculptTool.Small",       new IMAGE_BRUSH("Icons/icon_Landscape_Tool_Sculpt_20x",           Icon20x20));
 		Set("LandscapeEditor.EraseTool.Small", new IMAGE_BRUSH("Icons/icon_Landscape_Tool_Erase_20x", Icon20x20));
 		Set("LandscapeEditor.PaintTool.Small",        new IMAGE_BRUSH("Icons/icon_Landscape_Tool_Paint_20x",            Icon20x20));
@@ -6699,7 +6699,7 @@ void FSlateEditorStyle::FStyle::SetupLandscapeEditorStyle()
 		Set("LandscapeEditor.NoiseTool.Small",        new IMAGE_BRUSH("Icons/icon_Landscape_Tool_Noise_20x",            Icon20x20));
 		Set("LandscapeEditor.RetopologizeTool.Small", new IMAGE_BRUSH("Icons/icon_Landscape_Tool_Retopologize_20x",     Icon20x20));
 		Set("LandscapeEditor.VisibilityTool.Small",   new IMAGE_BRUSH("Icons/icon_Landscape_Tool_Visibility_20x",       Icon20x20));
-		Set("LandscapeEditor.BPCustomTool.Small", new IMAGE_BRUSH("Icons/icon_Landscape_Tool_Visibility_20x", Icon20x20)); // TODO: change icon
+		Set("LandscapeEditor.BlueprintBrushTool.Small", new IMAGE_BRUSH("Icons/icon_Landscape_Tool_BlueprintBrush_20x", Icon20x20));
 
 		Set("LandscapeEditor.SelectComponentTool", new IMAGE_BRUSH("Icons/icon_Landscape_Tool_Selection_40x",       Icon40x40));
 		Set("LandscapeEditor.AddComponentTool",    new IMAGE_BRUSH("Icons/icon_Landscape_Tool_AddComponent_40x",    Icon40x40));
@@ -6748,6 +6748,11 @@ void FSlateEditorStyle::FStyle::SetupLandscapeEditorStyle()
 		Set("LandscapeEditor.Brushes.Alpha.UseBChannel", new IMAGE_BRUSH("Icons/icon_Landscape_Brush_Alpha_UseBChannel_20x", Icon20x20));
 		Set("LandscapeEditor.Brushes.Alpha.UseAChannel", new IMAGE_BRUSH("Icons/icon_Landscape_Brush_Alpha_UseAChannel_20x", Icon20x20));
 
+		Set("LandscapeEditor.Brush.AffectsHeight.Enabled", new IMAGE_BRUSH("Icons/icon_Landscape_Affects_Height_Enabled_16x", Icon16x16));
+		Set("LandscapeEditor.Brush.AffectsHeight.Disabled", new IMAGE_BRUSH("Icons/icon_Landscape_Affects_Height_Disabled_16x", Icon16x16));
+		Set("LandscapeEditor.Brush.AffectsWeight.Enabled", new IMAGE_BRUSH("Icons/icon_Landscape_Affects_Weight_Enabled_16x", Icon16x16));
+		Set("LandscapeEditor.Brush.AffectsWeight.Disabled", new IMAGE_BRUSH("Icons/icon_Landscape_Affects_Weight_Disabled_16x", Icon16x16));
+
 		// Target List
 		Set("LandscapeEditor.TargetList.RowBackground",        new FSlateNoResource());
 		Set("LandscapeEditor.TargetList.RowBackgroundHovered", new BOX_BRUSH("Common/RoundedSelection_16x", 4.0f / 16.0f, FLinearColor(1.0f, 1.0f, 1.0f, 0.1f)));
diff --git a/Engine/Source/Editor/LandscapeEditor/Private/LandscapeEdMode.cpp b/Engine/Source/Editor/LandscapeEditor/Private/LandscapeEdMode.cpp
--- a/Engine/Source/Editor/LandscapeEditor/Private/LandscapeEdMode.cpp
+++ b/Engine/Source/Editor/LandscapeEditor/Private/LandscapeEdMode.cpp
@@ -340,7 +340,7 @@ FEdModeLandscape::FEdModeLandscape()
 	InitializeTool_Splines();
 	InitializeTool_Ramp();
 	InitializeTool_Mirror();
-	InitializeTool_BPCustom();
+	InitializeTool_BlueprintBrush();
 
 	// Initialize brushes
 	InitializeBrushes();
@@ -436,7 +436,7 @@ void FEdModeLandscape::UpdateToolModes()
 
 	if (CanHaveLandscapeLayersContent())
 	{
-		ToolMode_Sculpt->ValidTools.Add(TEXT("BPCustom"));
+		ToolMode_Sculpt->ValidTools.Add(TEXT("BlueprintBrush"));
 	}
 
 	ToolMode_Sculpt->ValidTools.Add(TEXT("Mask"));
@@ -452,7 +452,7 @@ void FEdModeLandscape::UpdateToolModes()
 
 	if (CanHaveLandscapeLayersContent())
 	{
-		ToolMode_Paint->ValidTools.Add(TEXT("BPCustom"));
+		ToolMode_Paint->ValidTools.Add(TEXT("BlueprintBrush"));
 	}
 
 	// Since available tools might have changed try and reset the current tool
@@ -2573,6 +2573,12 @@ bool FEdModeLandscape::CanEditCurrentTarget(FText* Reason) const
 	if (GetLandscape() == nullptr)
 	{
 		ALandscapeProxy* Proxy = CurrentToolTarget.LandscapeInfo->GetLandscapeProxy();
+		if (!Proxy)
+		{
+            LocalReason = NSLOCTEXT("UnrealEd", "LandscapeNotFound", "No Landscape found.");
+			return false;
+		}
+		
 		if (Proxy->HasLayersContent())
 		{
 			LocalReason = NSLOCTEXT("UnrealEd", "LandscapeActorNotLoaded", "Landscape actor is not loaded. It is needed to do layer editing.");
@@ -4619,7 +4625,7 @@ bool FEdModeLandscape::CanEditLayer(FText* Reason /*=nullptr*/, FLandscapeLayer*
 		}
 	}
 
-	if (CurrentToolTarget.TargetType == ELandscapeToolTargetType::Weightmap && CurrentToolTarget.LayerInfo == NULL && CurrentTool->GetToolName() != FName("BPCustom"))
+	if (CurrentToolTarget.TargetType == ELandscapeToolTargetType::Weightmap && CurrentToolTarget.LayerInfo == NULL && CurrentTool->GetToolName() != FName("BlueprintBrush"))
 	{
 		if (Reason)
 		{
@@ -4710,7 +4716,7 @@ void FEdModeLandscape::UpdateBrushList()
 	BrushList.Empty();
 	for (TObjectIterator<ALandscapeBlueprintCustomBrush> BrushIt(RF_Transient|RF_ClassDefaultObject|RF_ArchetypeObject, true, EInternalObjectFlags::PendingKill); BrushIt; ++BrushIt)
 	{
-		BrushList.AddUnique(*BrushIt);
+		BrushList.Add(*BrushIt);
 	}
 }
 
@@ -4726,7 +4732,7 @@ void FEdModeLandscape::OnLevelActorAdded(AActor* InActor)
 
 	if (Brush != nullptr && Brush->GetTypedOuter<UPackage>() != GetTransientPackage())
 	{
-		BrushList.AddUnique(Brush);
+		BrushList.Add(Brush);
 		AddBrushToCurrentLayer(Brush);
 		RefreshDetailPanel();
 	}
diff --git a/Engine/Source/Editor/LandscapeEditor/Private/LandscapeEdMode.h b/Engine/Source/Editor/LandscapeEditor/Private/LandscapeEdMode.h
--- a/Engine/Source/Editor/LandscapeEditor/Private/LandscapeEdMode.h
+++ b/Engine/Source/Editor/LandscapeEditor/Private/LandscapeEdMode.h
@@ -355,7 +355,7 @@ public:
 	void InitializeTool_Splines();
 	void InitializeTool_Ramp();
 	void InitializeTool_Mirror();
-	void InitializeTool_BPCustom();
+	void InitializeTool_BlueprintBrush();
 	void UpdateToolModes();
 
 	/** Destructor */
diff --git a/Engine/Source/Editor/LandscapeEditor/Private/LandscapeEdModeBPCustomTools.cpp b/Engine/Source/Editor/LandscapeEditor/Private/LandscapeEdModeBPCustomTools.cpp
--- a/Engine/Source/Editor/LandscapeEditor/Private/LandscapeEdModeBPCustomTools.cpp
+++ b/Engine/Source/Editor/LandscapeEditor/Private/LandscapeEdModeBPCustomTools.cpp
@@ -27,13 +27,13 @@
 #define LOCTEXT_NAMESPACE "Landscape"
 
 template<class ToolTarget>
-class FLandscapeToolBPCustom : public FLandscapeTool
+class FLandscapeToolBlueprintBrush : public FLandscapeTool
 {
 protected:
 	FEdModeLandscape* EdMode;
 
 public:
-	FLandscapeToolBPCustom(FEdModeLandscape* InEdMode)
+	FLandscapeToolBlueprintBrush(FEdModeLandscape* InEdMode)
 		: EdMode(InEdMode)
 	{
 	}
@@ -46,7 +46,7 @@ public:
 	{
 	}
 
-	virtual const TCHAR* GetToolName() override { return TEXT("BPCustom"); }
+	virtual const TCHAR* GetToolName() override { return TEXT("BlueprintBrush"); }
 	virtual FText GetDisplayName() override { return FText(); };
 
 	virtual void SetEditRenderType() override { GLandscapeEditRenderMode = ELandscapeEditRenderMode::None | (GLandscapeEditRenderMode & ELandscapeEditRenderMode::BitMaskForMask); }
@@ -63,6 +63,7 @@ public:
 
 	virtual void ExitTool() override
 	{
+		GEditor->SelectNone(true, true);
 	}
 
 	virtual void Tick(FEditorViewportClient* ViewportClient, float DeltaTime) 
@@ -194,15 +195,15 @@ void FEdModeLandscape::CenterMirrorTool()
 //
 // Toolset initialization
 //
-void FEdModeLandscape::InitializeTool_BPCustom()
+void FEdModeLandscape::InitializeTool_BlueprintBrush()
 {
-	auto Sculpt_Tool_BPCustom = MakeUnique<FLandscapeToolBPCustom<FHeightmapToolTarget>>(this);
-	Sculpt_Tool_BPCustom->ValidBrushes.Add("BrushSet_Dummy");
-	LandscapeTools.Add(MoveTemp(Sculpt_Tool_BPCustom));
+	auto Sculpt_Tool_BlueprintBrush = MakeUnique<FLandscapeToolBlueprintBrush<FHeightmapToolTarget>>(this);
+	Sculpt_Tool_BlueprintBrush->ValidBrushes.Add("BrushSet_Dummy");
+	LandscapeTools.Add(MoveTemp(Sculpt_Tool_BlueprintBrush));
 
-	auto Paint_Tool_BPCustom = MakeUnique<FLandscapeToolBPCustom<FWeightmapToolTarget>>(this);
-	Paint_Tool_BPCustom->ValidBrushes.Add("BrushSet_Dummy");
-	LandscapeTools.Add(MoveTemp(Paint_Tool_BPCustom));
+	auto Paint_Tool_BlueprintBrush = MakeUnique<FLandscapeToolBlueprintBrush<FWeightmapToolTarget>>(this);
+	Paint_Tool_BlueprintBrush->ValidBrushes.Add("BrushSet_Dummy");
+	LandscapeTools.Add(MoveTemp(Paint_Tool_BlueprintBrush));
 }
 
 #undef LOCTEXT_NAMESPACE
diff --git a/Engine/Source/Editor/LandscapeEditor/Private/LandscapeEditorCommands.cpp b/Engine/Source/Editor/LandscapeEditor/Private/LandscapeEditorCommands.cpp
--- a/Engine/Source/Editor/LandscapeEditor/Private/LandscapeEditorCommands.cpp
+++ b/Engine/Source/Editor/LandscapeEditor/Private/LandscapeEditorCommands.cpp
@@ -56,8 +56,8 @@ void FLandscapeEditorCommands::RegisterCommands()
 	UI_COMMAND(VisibilityTool, "Tool - Visibility", "", EUserInterfaceActionType::RadioButton, FInputChord());
 	NameToCommandMap.Add("Tool_Visibility", VisibilityTool);
 
-	UI_COMMAND(BPCustomTool, "Tool - BP Custom", "", EUserInterfaceActionType::RadioButton, FInputChord());
-	NameToCommandMap.Add("Tool_BPCustom", BPCustomTool);
+	UI_COMMAND(BlueprintBrushTool, "Tool - Blueprint Brush", "", EUserInterfaceActionType::RadioButton, FInputChord());
+	NameToCommandMap.Add("Tool_BlueprintBrush", BlueprintBrushTool);
 
 	UI_COMMAND(SelectComponentTool, "Tool - Component Selection", "", EUserInterfaceActionType::RadioButton, FInputChord());
 	NameToCommandMap.Add("Tool_Select", SelectComponentTool);
diff --git a/Engine/Source/Editor/LandscapeEditor/Private/LandscapeEditorCommands.h b/Engine/Source/Editor/LandscapeEditor/Private/LandscapeEditorCommands.h
--- a/Engine/Source/Editor/LandscapeEditor/Private/LandscapeEditorCommands.h
+++ b/Engine/Source/Editor/LandscapeEditor/Private/LandscapeEditorCommands.h
@@ -44,7 +44,7 @@ public:
 	TSharedPtr<FUICommandInfo> NoiseTool;
 	TSharedPtr<FUICommandInfo> RetopologizeTool;
 	TSharedPtr<FUICommandInfo> VisibilityTool;
-	TSharedPtr<FUICommandInfo> BPCustomTool;
+	TSharedPtr<FUICommandInfo> BlueprintBrushTool;
 
 	TSharedPtr<FUICommandInfo> SelectComponentTool;
 	TSharedPtr<FUICommandInfo> AddComponentTool;
diff --git a/Engine/Source/Editor/LandscapeEditor/Private/LandscapeEditorDetailCustomization_Layers.cpp b/Engine/Source/Editor/LandscapeEditor/Private/LandscapeEditorDetailCustomization_Layers.cpp
--- a/Engine/Source/Editor/LandscapeEditor/Private/LandscapeEditorDetailCustomization_Layers.cpp
+++ b/Engine/Source/Editor/LandscapeEditor/Private/LandscapeEditorDetailCustomization_Layers.cpp
@@ -144,7 +144,7 @@ void FLandscapeEditorCustomNodeBuilder_Layers::GenerateChildContent(IDetailChild
 		LayerList->SetDropIndicator_Above(*FEditorStyle::GetBrush("LandscapeEditor.TargetList.DropZone.Above"));
 		LayerList->SetDropIndicator_Below(*FEditorStyle::GetBrush("LandscapeEditor.TargetList.DropZone.Below"));
 
-		ChildrenBuilder.AddCustomRow(FText::FromString(FString(TEXT("Test"))))
+		ChildrenBuilder.AddCustomRow(FText::FromString(FString(TEXT("Layers"))))
 			.Visibility(EVisibility::Visible)
 			[
 				LayerList.ToSharedRef()
diff --git a/Engine/Source/Editor/LandscapeEditor/Private/LandscapeEditorDetailCustomization_LayersBrushStack.cpp b/Engine/Source/Editor/LandscapeEditor/Private/LandscapeEditorDetailCustomization_LayersBrushStack.cpp
--- a/Engine/Source/Editor/LandscapeEditor/Private/LandscapeEditorDetailCustomization_LayersBrushStack.cpp
+++ b/Engine/Source/Editor/LandscapeEditor/Private/LandscapeEditorDetailCustomization_LayersBrushStack.cpp
@@ -45,10 +45,7 @@ void FLandscapeEditorDetailCustomization_LayersBrushStack::CustomizeDetails(IDet
 	{
 		const FName CurrentToolName = LandscapeEdMode->CurrentTool->GetToolName();
 
-		if (LandscapeEdMode->CurrentToolMode->SupportedTargetTypes != 0 && CurrentToolName == TEXT("BPCustom"))
-		{
-			LayerCategory.AddCustomBuilder(MakeShareable(new FLandscapeEditorCustomNodeBuilder_LayersBrushStack(DetailBuilder.GetThumbnailPool().ToSharedRef())));
-		}
+		LayerCategory.AddCustomBuilder(MakeShareable(new FLandscapeEditorCustomNodeBuilder_LayersBrushStack(DetailBuilder.GetThumbnailPool().ToSharedRef())));
 	}
 }
 END_SLATE_FUNCTION_BUILD_OPTIMIZATION
@@ -150,7 +147,7 @@ TSharedPtr<SWidget> FLandscapeEditorCustomNodeBuilder_LayersBrushStack::Generate
 				.Content()
 				[
 					SNew(SImage)
-					.Image(this, &FLandscapeEditorCustomNodeBuilder_LayersBrushStack::GetVisibilityBrushForBrush, InBrushIndex)
+					.Image(this, &FLandscapeEditorCustomNodeBuilder_LayersBrushStack::GetVisibilityBrush, InBrushIndex)
 				]
 			]
 			+SHorizontalBox::Slot()
@@ -166,9 +163,44 @@ TSharedPtr<SWidget> FLandscapeEditorCustomNodeBuilder_LayersBrushStack::Generate
 					SNew(STextBlock)
 					.ColorAndOpacity(TAttribute<FSlateColor>::Create(TAttribute<FSlateColor>::FGetter::CreateSP(this, &FLandscapeEditorCustomNodeBuilder_LayersBrushStack::GetBrushTextColor, InBrushIndex)))
 					.Text(this, &FLandscapeEditorCustomNodeBuilder_LayersBrushStack::GetBrushText, InBrushIndex)
-					.IsEnabled(TAttribute<bool>::Create(TAttribute<bool>::FGetter::CreateSP(this, &FLandscapeEditorCustomNodeBuilder_LayersBrushStack::IsBrushEnabled, InBrushIndex)))
+					.IsEnabled(true)
+				]
+			]
+			+ SHorizontalBox::Slot()
+			.AutoWidth()
+			.Padding(0, 0, 4, 0)
+			.VAlign(VAlign_Center)
+			[
+				SNew(SButton)
+				.ContentPadding(0)
+				.ButtonStyle(FEditorStyle::Get(), "NoBorder")
+				.OnClicked(this, &FLandscapeEditorCustomNodeBuilder_LayersBrushStack::OnToggleAffectsHeightmap, InBrushIndex)
+				.ToolTipText(LOCTEXT("LandscapeBrushAffectsHeightmap", "Toggle Affects Heightmap"))
+				.HAlign(HAlign_Center)
+				.VAlign(VAlign_Center)
+				.Content()
+				[
+					SNew(SImage)
+					.Image(this, &FLandscapeEditorCustomNodeBuilder_LayersBrushStack::GetAffectsHeightmapBrush, InBrushIndex)
 				]
 			]
+			+ SHorizontalBox::Slot()
+			.AutoWidth()
+            .VAlign(VAlign_Center)
+            [
+				SNew(SButton)
+                .ContentPadding(0)
+                .ButtonStyle(FEditorStyle::Get(), "NoBorder")
+                .OnClicked(this, &FLandscapeEditorCustomNodeBuilder_LayersBrushStack::OnToggleAffectsWeightmap, InBrushIndex)
+                .ToolTipText(LOCTEXT("LandscapeBrushAffectsWeightmap", "Toggle Affects Weightmap"))
+                .HAlign(HAlign_Center)
+                .VAlign(VAlign_Center)
+                .Content()
+                [
+                    SNew(SImage)
+                    .Image(this, &FLandscapeEditorCustomNodeBuilder_LayersBrushStack::GetAffectsWeightmapBrush, InBrushIndex)
+                ]
+            ]
 		];
 	
 	return RowWidget;
@@ -203,6 +235,28 @@ TSharedPtr<SWidget> FLandscapeEditorCustomNodeBuilder_LayersBrushStack::OnBrushC
 	return nullptr;
 }
 
+FReply FLandscapeEditorCustomNodeBuilder_LayersBrushStack::OnToggleAffectsHeightmap(int32 InBrushIndex)
+{
+	if (ALandscapeBlueprintCustomBrush* Brush = GetBrush(InBrushIndex))
+	{
+		const FScopedTransaction Transaction(LOCTEXT("Landscape_Brush_AffectsHeightmap", "Set Brush Affects Heightmap"));
+		bool bAffectsHeightmap = Brush->IsAffectingHeightmap();
+		Brush->SetAffectsHeightmap(!bAffectsHeightmap);
+	}
+	return FReply::Handled();
+}
+
+FReply FLandscapeEditorCustomNodeBuilder_LayersBrushStack::OnToggleAffectsWeightmap(int32 InBrushIndex)
+{
+	if (ALandscapeBlueprintCustomBrush* Brush = GetBrush(InBrushIndex))
+	{
+		const FScopedTransaction Transaction(LOCTEXT("Landscape_Brush_AffectsWeightmap", "Set Brush Affects Weightmap"));
+		bool bAffectsWeightmap = Brush->IsAffectingWeightmap();
+		Brush->SetAffectsWeightmap(!bAffectsWeightmap);
+	}
+	return FReply::Handled();
+}
+
 FReply FLandscapeEditorCustomNodeBuilder_LayersBrushStack::OnToggleVisibility(int32 InBrushIndex)
 {
 	if (ALandscapeBlueprintCustomBrush* Brush = GetBrush(InBrushIndex))
@@ -214,7 +268,19 @@ FReply FLandscapeEditorCustomNodeBuilder_LayersBrushStack::OnToggleVisibility(in
 	return FReply::Handled();
 }
 
-const FSlateBrush* FLandscapeEditorCustomNodeBuilder_LayersBrushStack::GetVisibilityBrushForBrush(int32 InBrushIndex) const
+const FSlateBrush* FLandscapeEditorCustomNodeBuilder_LayersBrushStack::GetAffectsWeightmapBrush(int32 InBrushIndex) const
+{
+	ALandscapeBlueprintCustomBrush* Brush = GetBrush(InBrushIndex);
+	return Brush && Brush->IsAffectingWeightmap() ? FEditorStyle::GetBrush("LandscapeEditor.Brush.AffectsWeight.Enabled") : FEditorStyle::GetBrush("LandscapeEditor.Brush.AffectsWeight.Disabled");
+}
+
+const FSlateBrush* FLandscapeEditorCustomNodeBuilder_LayersBrushStack::GetAffectsHeightmapBrush(int32 InBrushIndex) const
+{
+	ALandscapeBlueprintCustomBrush* Brush = GetBrush(InBrushIndex);
+	return Brush && Brush->IsAffectingHeightmap() ? FEditorStyle::GetBrush("LandscapeEditor.Brush.AffectsHeight.Enabled") : FEditorStyle::GetBrush("LandscapeEditor.Brush.AffectsHeight.Disabled");
+}
+
+const FSlateBrush* FLandscapeEditorCustomNodeBuilder_LayersBrushStack::GetVisibilityBrush(int32 InBrushIndex) const
 {
 	ALandscapeBlueprintCustomBrush* Brush = GetBrush(InBrushIndex);
 	bool bIsVisible = Brush && Brush->IsVisible();
@@ -236,11 +302,22 @@ bool FLandscapeEditorCustomNodeBuilder_LayersBrushStack::IsBrushSelected(int32 I
 
 void FLandscapeEditorCustomNodeBuilder_LayersBrushStack::OnBrushSelectionChanged(int32 InBrushIndex)
 {
-	ALandscapeBlueprintCustomBrush* Brush = GetBrush(InBrushIndex);
-	if (Brush != nullptr && !Brush->IsCommited())
+	FEdModeLandscape* LandscapeEdMode = GetEditorMode();
+	if (!LandscapeEdMode || !LandscapeEdMode->CurrentToolMode)
+	{
+		return;
+	}
+
+	const FName CurrentToolName = LandscapeEdMode->CurrentTool->GetToolName();
+	if (CurrentToolName == TEXT("BlueprintBrush"))
 	{
-		GEditor->SelectNone(true, true);
-		GEditor->SelectActor(Brush, true, true);
+		ALandscapeBlueprintCustomBrush* Brush = GetBrush(InBrushIndex);
+		if (Brush != nullptr && !Brush->IsCommited())
+		{
+			FScopedTransaction Transaction(LOCTEXT("LandscapeBrushSelect", "Brush selection"));
+			GEditor->SelectNone(true, true);
+			GEditor->SelectActor(Brush, true, true);
+		}
 	}
 }
 
diff --git a/Engine/Source/Editor/LandscapeEditor/Private/LandscapeEditorDetailCustomization_LayersBrushStack.h b/Engine/Source/Editor/LandscapeEditor/Private/LandscapeEditorDetailCustomization_LayersBrushStack.h
--- a/Engine/Source/Editor/LandscapeEditor/Private/LandscapeEditorDetailCustomization_LayersBrushStack.h
+++ b/Engine/Source/Editor/LandscapeEditor/Private/LandscapeEditorDetailCustomization_LayersBrushStack.h
@@ -62,7 +62,12 @@ protected:
 	void RemoveBrush(ALandscapeBlueprintCustomBrush* Brush);
 	
 	FReply OnToggleVisibility(int32 InBrushIndex);
-	const FSlateBrush* GetVisibilityBrushForBrush(int32 InBrushIndex) const;
+	FReply OnToggleAffectsHeightmap(int32 InBrushIndex);
+	FReply OnToggleAffectsWeightmap(int32 InBrushIndex);
+
+	const FSlateBrush* GetAffectsHeightmapBrush(int32 InBrushIndex) const;
+	const FSlateBrush* GetAffectsWeightmapBrush(int32 InBrushIndex) const;
+	const FSlateBrush* GetVisibilityBrush(int32 InBrushIndex) const;
 	bool IsBrushSelected(int32 InBrushIndex) const;
 	bool IsBrushEnabled(int32 InBrushIndex) const;
 	void OnBrushSelectionChanged(int32 InBrushIndex);
diff --git a/Engine/Source/Editor/LandscapeEditor/Private/LandscapeEditorDetailCustomization_TargetLayers.cpp b/Engine/Source/Editor/LandscapeEditor/Private/LandscapeEditorDetailCustomization_TargetLayers.cpp
--- a/Engine/Source/Editor/LandscapeEditor/Private/LandscapeEditorDetailCustomization_TargetLayers.cpp
+++ b/Engine/Source/Editor/LandscapeEditor/Private/LandscapeEditorDetailCustomization_TargetLayers.cpp
@@ -98,7 +98,7 @@ bool FLandscapeEditorDetailCustomization_TargetLayers::ShouldShowTargetLayers()
 
 		//// Visible if there are possible choices
 		//if (bSupportsWeightmap || bSupportsHeightmap || bSupportsVisibility)
-		if (LandscapeEdMode->CurrentToolMode->SupportedTargetTypes != 0 && CurrentToolName != TEXT("BPCustom"))
+		if (LandscapeEdMode->CurrentToolMode->SupportedTargetTypes != 0 && CurrentToolName != TEXT("BlueprintBrush"))
 		{
 			return true;
 		}
@@ -115,7 +115,7 @@ bool FLandscapeEditorDetailCustomization_TargetLayers::ShouldShowPaintingRestric
 	{
 		const FName CurrentToolName = LandscapeEdMode->CurrentTool->GetToolName();
 
-		if ((LandscapeEdMode->CurrentToolTarget.TargetType == ELandscapeToolTargetType::Weightmap && CurrentToolName != TEXT("BPCustom"))
+		if ((LandscapeEdMode->CurrentToolTarget.TargetType == ELandscapeToolTargetType::Weightmap && CurrentToolName != TEXT("BlueprintBrush"))
 			|| LandscapeEdMode->CurrentToolTarget.TargetType == ELandscapeToolTargetType::Visibility)
 		{
 			return true;
diff --git a/Engine/Source/Editor/LandscapeEditor/Private/LandscapeEditorDetails.cpp b/Engine/Source/Editor/LandscapeEditor/Private/LandscapeEditorDetails.cpp
--- a/Engine/Source/Editor/LandscapeEditor/Private/LandscapeEditorDetails.cpp
+++ b/Engine/Source/Editor/LandscapeEditor/Private/LandscapeEditorDetails.cpp
@@ -173,7 +173,7 @@ FText FLandscapeEditorDetails::GetLocalizedName(FString Name)
 
 		if (LandscapeEdMode->CanHaveLandscapeLayersContent())
 		{
-			LOCTEXT("ToolSet_BPCustom", "Blueprint Custom");
+			LOCTEXT("ToolSet_BlueprintBrush", "Blueprint Brushes");
 		}
 
 		LOCTEXT("ToolSet_Select", "Selection");
@@ -336,7 +336,7 @@ TSharedRef<SWidget> FLandscapeEditorDetails::GetToolSelector()
 
 			if (LandscapeEdMode->CanHaveLandscapeLayersContent())
 			{
-				MenuBuilder.AddToolButton(NameToCommandMap.FindChecked("Tool_BPCustom"), NAME_None, LOCTEXT("Tool.SculptBPCustom", "Blueprint Custom"), LOCTEXT("Tool.SculptBPCustom.Tooltip", "Custom sculpting tools created using Blueprint."));
+				MenuBuilder.AddToolButton(NameToCommandMap.FindChecked("Tool_BlueprintBrush"), NAME_None, LOCTEXT("Tool.SculptBlueprintBrush", "Blueprint Brushes"), LOCTEXT("Tool.SculptBlueprintBrush.Tooltip", "Custom sculpting tools created using Blueprint."));
 			}
 
 			MenuBuilder.EndSection();
@@ -358,7 +358,7 @@ TSharedRef<SWidget> FLandscapeEditorDetails::GetToolSelector()
 
 			if (LandscapeEdMode->CanHaveLandscapeLayersContent())
 			{
-				MenuBuilder.AddToolButton(NameToCommandMap.FindChecked("Tool_BPCustom"), NAME_None, LOCTEXT("Tool.PaintBPCustom", "Blueprint Custom"), LOCTEXT("Tool.PaintBPCustom.Tooltip", "Custom painting tools created using Blueprint."));
+				MenuBuilder.AddToolButton(NameToCommandMap.FindChecked("Tool_BlueprintBrush"), NAME_None, LOCTEXT("Tool.PaintBlueprintBrush", "Blueprint Brushes"), LOCTEXT("Tool.PaintBlueprintBrush.Tooltip", "Custom painting tools created using Blueprint."));
 			}
 
 			MenuBuilder.EndSection();
diff --git a/Engine/Source/Editor/LandscapeEditor/Private/SLandscapeEditor.cpp b/Engine/Source/Editor/LandscapeEditor/Private/SLandscapeEditor.cpp
--- a/Engine/Source/Editor/LandscapeEditor/Private/SLandscapeEditor.cpp
+++ b/Engine/Source/Editor/LandscapeEditor/Private/SLandscapeEditor.cpp
@@ -104,7 +104,7 @@ void FLandscapeToolKit::Init(const TSharedPtr<IToolkitHost>& InitToolkitHost)
 	MAP_TOOL("Noise");
 	MAP_TOOL("Retopologize");
 	MAP_TOOL("Visibility");
-	MAP_TOOL("BPCustom");
+	MAP_TOOL("BlueprintBrush");
 
 	MAP_TOOL("Select");
 	MAP_TOOL("AddComponent");
@@ -525,11 +525,11 @@ bool SLandscapeEditor::GetIsPropertyVisible(const FPropertyAndParent& PropertyAn
 				return false;
 			}
 		}
-		if (Property.HasMetaData("ShowForBPCustomTool"))
+		if (Property.HasMetaData("ShowForBlueprintBrushTool"))
 		{
 			const FName CurrentToolName = LandscapeEdMode->CurrentTool->GetToolName();
 
-			if (CurrentToolName != TEXT("BPCustom"))
+			if (CurrentToolName != TEXT("BlueprintBrush"))
 			{
 				return false;
 			}
diff --git a/Engine/Source/Editor/LandscapeEditor/Public/LandscapeEditorObject.h b/Engine/Source/Editor/LandscapeEditor/Public/LandscapeEditorObject.h
--- a/Engine/Source/Editor/LandscapeEditor/Public/LandscapeEditorObject.h
+++ b/Engine/Source/Editor/LandscapeEditor/Public/LandscapeEditorObject.h
@@ -451,9 +451,9 @@ class ULandscapeEditorObject : public UObject
 	UPROPERTY(Category="Tool Settings", EditAnywhere, NonTransactional, meta=(DisplayName="Smoothing Width", ShowForTools="Mirror", ClampMin="0", UIMin="0", UIMax="20"))
 	int32 MirrorSmoothingWidth;
 
-	// BP Custom Tool
+	// Blueprint Brush Tool
 
-	UPROPERTY(Category = "Tool Settings", EditAnywhere, Transient, meta = (DisplayName = "Blueprint Brush", ShowForTools = "BPCustom"))
+	UPROPERTY(Category = "Tool Settings", EditAnywhere, Transient, meta = (DisplayName = "Blueprint Brush", ShowForTools = "BlueprintBrush"))
 	TSubclassOf<ALandscapeBlueprintCustomBrush> BlueprintCustomBrush;
 
 	// Resize Landscape Tool
diff --git a/Engine/Source/Runtime/Landscape/Classes/LandscapeBPCustomBrush.h b/Engine/Source/Runtime/Landscape/Classes/LandscapeBPCustomBrush.h
--- a/Engine/Source/Runtime/Landscape/Classes/LandscapeBPCustomBrush.h
+++ b/Engine/Source/Runtime/Landscape/Classes/LandscapeBPCustomBrush.h
@@ -14,13 +14,13 @@ class LANDSCAPE_API ALandscapeBlueprintCustomBrush : public AActor
 	GENERATED_UCLASS_BODY()
 
 private:
-	UPROPERTY(Category= "Settings", EditAnywhere, NonTransactional)
+	UPROPERTY(Category= "Settings", EditAnywhere)
 	bool AffectHeightmap;
 
-	UPROPERTY(Category= "Settings", EditAnywhere, NonTransactional)
+	UPROPERTY(Category= "Settings", EditAnywhere)
 	bool AffectWeightmap;
 
-	UPROPERTY(Category = "Settings", EditAnywhere, NonTransactional)
+	UPROPERTY(Category = "Settings", EditAnywhere)
 	TArray<FName> AffectedWeightmapLayers;
 
 #if WITH_EDITORONLY_DATA
@@ -61,6 +61,9 @@ public:
 	bool IsVisible() const { return bIsVisible; }
 	void SetIsVisible(bool bInIsVisible);
 
+	void SetAffectsHeightmap(bool bInAffectsHeightmap);
+	void SetAffectsWeightmap(bool bInAffectsWeightmap);
+
 	void SetOwningLandscape(class ALandscape* InOwningLandscape);
 	class ALandscape* GetOwningLandscape() const;
 
diff --git a/Engine/Source/Runtime/Landscape/Private/LandscapeBPCustomBrush.cpp b/Engine/Source/Runtime/Landscape/Private/LandscapeBPCustomBrush.cpp
--- a/Engine/Source/Runtime/Landscape/Private/LandscapeBPCustomBrush.cpp
+++ b/Engine/Source/Runtime/Landscape/Private/LandscapeBPCustomBrush.cpp
@@ -60,6 +60,20 @@ void ALandscapeBlueprintCustomBrush::SetIsVisible(bool bInIsVisible)
 	GetOwningLandscape()->OnBPCustomBrushChanged();
 }
 
+void ALandscapeBlueprintCustomBrush::SetAffectsHeightmap(bool bInAffectsHeightmap)
+{
+	Modify();
+	AffectHeightmap = bInAffectsHeightmap;
+	GetOwningLandscape()->OnBPCustomBrushChanged();
+}
+
+void ALandscapeBlueprintCustomBrush::SetAffectsWeightmap(bool bInAffectsWeightmap)
+{
+	Modify();
+	AffectWeightmap = bInAffectsWeightmap;
+	GetOwningLandscape()->OnBPCustomBrushChanged();
+}
+
 void ALandscapeBlueprintCustomBrush::SetOwningLandscape(ALandscape* InOwningLandscape)
 {
 	Modify();
