- Landscape Layer System: Fix some crashes when going out of landscape in a specific tool and deleting the last landscape in the scene and going back into landscape

#rnx

diff --git a/Engine/Source/Editor/LandscapeEditor/Private/LandscapeEdModeComponentTools.cpp b/Engine/Source/Editor/LandscapeEditor/Private/LandscapeEdModeComponentTools.cpp
--- a/Engine/Source/Editor/LandscapeEditor/Private/LandscapeEdModeComponentTools.cpp
+++ b/Engine/Source/Editor/LandscapeEditor/Private/LandscapeEdModeComponentTools.cpp
@@ -543,7 +543,7 @@ public:
 				{
 					ALandscape::SplitHeightmap(HeightmapUpdateComponentPair.Key, HeightmapUpdateComponentPair.Value ? LandscapeProxy : nullptr);
 				}
-
+								
 				// Delete if it is no referenced textures...
 				for (UTexture2D* Texture : OldHeightmapTextures)
 				{
@@ -971,8 +971,10 @@ public:
 	virtual void EnterTool() override
 	{
 		FLandscapeToolBase<FLandscapeToolStrokeAddComponent>::EnterTool();
-		ULandscapeInfo* LandscapeInfo = EdMode->CurrentToolTarget.LandscapeInfo.Get();
-		LandscapeInfo->UpdateAllAddCollisions(); // Todo - as this is only used by this tool, move it into this tool?
+		if(ULandscapeInfo* LandscapeInfo = EdMode->CurrentToolTarget.LandscapeInfo.Get())
+		{
+			LandscapeInfo->UpdateAllAddCollisions(); // Todo - as this is only used by this tool, move it into this tool?
+		}
 	}
 
 	virtual void ExitTool() override
@@ -1890,23 +1892,26 @@ public:
 
 	virtual void EnterTool() override
 	{
-		const int32 ComponentSizeQuads = EdMode->CurrentToolTarget.LandscapeInfo->ComponentSizeQuads;
-		int32 MinX, MinY, MaxX, MaxY;
-		if (EdMode->CurrentToolTarget.LandscapeInfo->GetLandscapeExtent(MinX, MinY, MaxX, MaxY))
-		{
-			EdMode->UISettings->ResizeLandscape_Original_ComponentCount.X = (MaxX - MinX) / ComponentSizeQuads;
-			EdMode->UISettings->ResizeLandscape_Original_ComponentCount.Y = (MaxY - MinY) / ComponentSizeQuads;
-			EdMode->UISettings->ResizeLandscape_ComponentCount = EdMode->UISettings->ResizeLandscape_Original_ComponentCount;
-		}
-		else
+		if (ULandscapeInfo* LandscapeInfo = EdMode->CurrentToolTarget.LandscapeInfo.Get())
 		{
-			EdMode->UISettings->ResizeLandscape_Original_ComponentCount = FIntPoint::ZeroValue;
-			EdMode->UISettings->ResizeLandscape_ComponentCount = FIntPoint::ZeroValue;
+			const int32 ComponentSizeQuads = LandscapeInfo->ComponentSizeQuads;
+			int32 MinX, MinY, MaxX, MaxY;
+			if (EdMode->CurrentToolTarget.LandscapeInfo->GetLandscapeExtent(MinX, MinY, MaxX, MaxY))
+			{
+				EdMode->UISettings->ResizeLandscape_Original_ComponentCount.X = (MaxX - MinX) / ComponentSizeQuads;
+				EdMode->UISettings->ResizeLandscape_Original_ComponentCount.Y = (MaxY - MinY) / ComponentSizeQuads;
+				EdMode->UISettings->ResizeLandscape_ComponentCount = EdMode->UISettings->ResizeLandscape_Original_ComponentCount;
+			}
+			else
+			{
+				EdMode->UISettings->ResizeLandscape_Original_ComponentCount = FIntPoint::ZeroValue;
+				EdMode->UISettings->ResizeLandscape_ComponentCount = FIntPoint::ZeroValue;
+			}
+			EdMode->UISettings->ResizeLandscape_Original_QuadsPerSection = EdMode->CurrentToolTarget.LandscapeInfo->SubsectionSizeQuads;
+			EdMode->UISettings->ResizeLandscape_Original_SectionsPerComponent = EdMode->CurrentToolTarget.LandscapeInfo->ComponentNumSubsections;
+			EdMode->UISettings->ResizeLandscape_QuadsPerSection = EdMode->UISettings->ResizeLandscape_Original_QuadsPerSection;
+			EdMode->UISettings->ResizeLandscape_SectionsPerComponent = EdMode->UISettings->ResizeLandscape_Original_SectionsPerComponent;
 		}
-		EdMode->UISettings->ResizeLandscape_Original_QuadsPerSection = EdMode->CurrentToolTarget.LandscapeInfo->SubsectionSizeQuads;
-		EdMode->UISettings->ResizeLandscape_Original_SectionsPerComponent = EdMode->CurrentToolTarget.LandscapeInfo->ComponentNumSubsections;
-		EdMode->UISettings->ResizeLandscape_QuadsPerSection = EdMode->UISettings->ResizeLandscape_Original_QuadsPerSection;
-		EdMode->UISettings->ResizeLandscape_SectionsPerComponent = EdMode->UISettings->ResizeLandscape_Original_SectionsPerComponent;
 	}
 
 	virtual void ExitTool() override
diff --git a/Engine/Source/Editor/LandscapeEditor/Private/LandscapeEdModePaintTools.cpp b/Engine/Source/Editor/LandscapeEditor/Private/LandscapeEdModePaintTools.cpp
--- a/Engine/Source/Editor/LandscapeEditor/Private/LandscapeEdModePaintTools.cpp
+++ b/Engine/Source/Editor/LandscapeEditor/Private/LandscapeEdModePaintTools.cpp
@@ -1256,6 +1256,10 @@ public:
 	virtual void EnterTool() override
 	{
 		FLandscapeToolPaintBase<ToolTarget, FLandscapeToolStrokeFlatten<ToolTarget>>::EnterTool();
+		if (!this->EdMode->CurrentToolTarget.LandscapeInfo.Get())
+		{
+			return;
+		}
 
 		ALandscapeProxy* LandscapeProxy = this->EdMode->CurrentToolTarget.LandscapeInfo->GetLandscapeProxy();
 		MeshComponent = NewObject<UStaticMeshComponent>(LandscapeProxy, NAME_None, RF_Transient);
@@ -1277,8 +1281,11 @@ public:
 	{
 		FLandscapeToolPaintBase<ToolTarget, FLandscapeToolStrokeFlatten<ToolTarget>>::ExitTool();
 
-		MeshComponent->DetachFromComponent(FDetachmentTransformRules::KeepRelativeTransform);
-		MeshComponent->DestroyComponent();
+		if (MeshComponent)
+		{
+			MeshComponent->DetachFromComponent(FDetachmentTransformRules::KeepRelativeTransform);
+			MeshComponent->DestroyComponent();
+		}
 	}
 };
 
