Fixed Mac compilation
#rnx

diff --git a/Engine/Source/Editor/LandscapeEditor/Private/LandscapeEdModePaintTools.cpp b/Engine/Source/Editor/LandscapeEditor/Private/LandscapeEdModePaintTools.cpp
--- a/Engine/Source/Editor/LandscapeEditor/Private/LandscapeEdModePaintTools.cpp
+++ b/Engine/Source/Editor/LandscapeEditor/Private/LandscapeEdModePaintTools.cpp
@@ -815,7 +815,7 @@ public:
 		}
 
 		TArray<typename ToolTarget::CacheClass::DataType> Data;
-		LayerDataCache.Initialize(LandscapeInfo, bCombinedLayerOperation);
+		LayerDataCache.Initialize(this->LandscapeInfo, bCombinedLayerOperation);
 		LayerDataCache.Read(X1, Y1, X2, Y2, Data);
 		const float ToolStrength = FMath::Clamp<float>(UISettings->ToolStrength * Pressure, 0.0f, 1.0f);
 
@@ -1022,7 +1022,7 @@ public:
 		}
 
 		TArray<typename ToolTarget::CacheClass::DataType> Data;
-		LayerDataCache.Initialize(LandscapeInfo, bCombinedLayerOperation);
+		LayerDataCache.Initialize(this->LandscapeInfo, bCombinedLayerOperation);
 		LayerDataCache.Read(X1, Y1, X2, Y2, Data);
 
 		// Apply the brush
