- Landscape Layer System: Erosion & Hydra Erosion now can work on Combined Layer Result

#rnx

diff --git a/Engine/Source/Editor/LandscapeEditor/Private/LandscapeEdModeErosionTools.cpp b/Engine/Source/Editor/LandscapeEditor/Private/LandscapeEdModeErosionTools.cpp
--- a/Engine/Source/Editor/LandscapeEditor/Private/LandscapeEdModeErosionTools.cpp
+++ b/Engine/Source/Editor/LandscapeEditor/Private/LandscapeEdModeErosionTools.cpp
@@ -21,6 +21,7 @@ public:
 	FLandscapeToolStrokeErosionBase(FEdModeLandscape* InEdMode, FEditorViewportClient* InViewportClient, const FLandscapeToolTarget& InTarget)
 		: FLandscapeToolStrokeBase(InEdMode, InViewportClient, InTarget)
 		, HeightCache(InTarget)
+		, LayerHeightDataCache(InTarget, this->HeightCache)
 		, WeightCache(InTarget)
 		, bWeightApplied(InTarget.TargetType != ELandscapeToolTargetType::Heightmap)
 	{
@@ -28,6 +29,7 @@ public:
 
 protected:
 	FLandscapeHeightCache HeightCache;
+	FLandscapeLayerDataCache<FHeightmapToolTarget> LayerHeightDataCache;
 	FLandscapeFullWeightCache WeightCache;
 	bool bWeightApplied;
 };
@@ -61,11 +63,14 @@ public:
 
 	void Apply(FEditorViewportClient* ViewportClient, FLandscapeBrush* Brush, const ULandscapeEditorObject* UISettings, const TArray<FLandscapeToolInteractorPosition>& InteractorPositions)
 	{
-		if (!LandscapeInfo)
+		if (!this->LandscapeInfo)
 		{
 			return;
 		}
 
+		ALandscape* Landscape = this->LandscapeInfo->LandscapeActor.Get();
+		bool bCombinedLayerOperation = UISettings->bCombinedLayersOperation && Landscape && Landscape->HasLayersContent();
+
 		// Get list of verts to update
 		FLandscapeBrushData BrushInfo = Brush->ApplyBrush(InteractorPositions);
 		if (!BrushInfo)
@@ -88,11 +93,11 @@ public:
 		const int32 NeighborNum = 4;
 		const int32 Iteration = UISettings->ErodeIterationNum;
 		const int32 Thickness = UISettings->ErodeSurfaceThickness;
-		const int32 LayerNum = LandscapeInfo->Layers.Num();
+		const int32 LayerNum = this->LandscapeInfo->Layers.Num();
 
-		HeightCache.CacheData(X1, Y1, X2, Y2);
 		TArray<uint16> HeightData;
-		HeightCache.GetCachedData(X1, Y1, X2, Y2, HeightData);
+		LayerHeightDataCache.Initialize(this->LandscapeInfo, bCombinedLayerOperation);
+		LayerHeightDataCache.Read(X1, Y1, X2, Y2, HeightData);
 
 		TArray<uint8> WeightDatas; // Weight*Layers...
 		WeightCache.CacheData(X1, Y1, X2, Y2);
@@ -154,7 +159,7 @@ public:
 							{
 								for (int32 Idx = 0; Idx < LayerNum; Idx++)
 								{
-									ULandscapeLayerInfoObject* LayerInfo = LandscapeInfo->Layers[Idx].LayerInfoObj;
+									ULandscapeLayerInfoObject* LayerInfo = this->LandscapeInfo->Layers[Idx].LayerInfoObj;
 									if (LayerInfo)
 									{
 										uint8 Weight = WeightDatas[Center*LayerNum + Idx];
@@ -256,8 +261,7 @@ public:
 			}
 		}
 
-		HeightCache.SetCachedData(X1, Y1, X2, Y2, HeightData);
-		HeightCache.Flush();
+		LayerHeightDataCache.Write(X1, Y1, X2, Y2, HeightData);
 		if (bWeightApplied)
 		{
 			WeightCache.SetCachedData(X1, Y1, X2, Y2, WeightDatas, LayerNum, ELandscapeLayerPaintingRestriction::None);
@@ -293,11 +297,14 @@ public:
 
 	void Apply(FEditorViewportClient* ViewportClient, FLandscapeBrush* Brush, const ULandscapeEditorObject* UISettings, const TArray<FLandscapeToolInteractorPosition>& InteractorPositions)
 	{
-		if (!LandscapeInfo)
+		if (!this->LandscapeInfo)
 		{
 			return;
 		}
 
+		ALandscape* Landscape = this->LandscapeInfo->LandscapeActor.Get();
+		bool bCombinedLayerOperation = UISettings->bCombinedLayersOperation && Landscape && Landscape->HasLayersContent();
+
 		// Get list of verts to update
 		FLandscapeBrushData BrushInfo = Brush->ApplyBrush(InteractorPositions);
 		if (!BrushInfo)
@@ -317,7 +324,7 @@ public:
 		X2 += 1;
 		Y2 += 1;
 
-		const int32 LayerNum = LandscapeInfo->Layers.Num();
+		const int32 LayerNum = this->LandscapeInfo->Layers.Num();
 
 		const int32 Iteration = UISettings->HErodeIterationNum;
 		const uint16 RainAmount = UISettings->RainAmount;
@@ -325,9 +332,9 @@ public:
 		const float EvaporateRatio = 0.5;
 		const float SedimentCapacity = 0.10 * UISettings->SedimentCapacity; //DissolvingRatio; //0.01;
 
-		HeightCache.CacheData(X1, Y1, X2, Y2);
 		TArray<uint16> HeightData;
-		HeightCache.GetCachedData(X1, Y1, X2, Y2, HeightData);
+		LayerHeightDataCache.Initialize(this->LandscapeInfo, bCombinedLayerOperation);
+		LayerHeightDataCache.Read(X1, Y1, X2, Y2, HeightData);
 
 		// Apply the brush
 		TArray<uint16> WaterData;
@@ -479,8 +486,7 @@ public:
 			LowPassFilter<uint16>(X1, Y1, X2, Y2, BrushInfo, HeightData, UISettings->HErosionDetailScale, 1.0f);
 		}
 
-		HeightCache.SetCachedData(X1, Y1, X2, Y2, HeightData);
-		HeightCache.Flush();
+		LayerHeightDataCache.Write(X1, Y1, X2, Y2, HeightData);
 	}
 };
 
diff --git a/Engine/Source/Editor/LandscapeEditor/Private/LandscapeEdModePaintTools.cpp b/Engine/Source/Editor/LandscapeEditor/Private/LandscapeEdModePaintTools.cpp
--- a/Engine/Source/Editor/LandscapeEditor/Private/LandscapeEdModePaintTools.cpp
+++ b/Engine/Source/Editor/LandscapeEditor/Private/LandscapeEdModePaintTools.cpp
@@ -593,180 +593,6 @@ public:
 	virtual FText GetDisplayName() override { return NSLOCTEXT("UnrealEd", "LandscapeMode_Erase", "Erase"); };
 };
 
-template<class ToolTarget>
-class FLandscapeLayerDataCache
-{
-public:
-	FLandscapeLayerDataCache(const FLandscapeToolTarget& InTarget, typename ToolTarget::CacheClass& Cache)
-		: LandscapeInfo(nullptr)
-		, Landscape(nullptr)
-		, EditingLayerIndex(MAX_uint8)
-		, bIsInitialized(false)
-		, bCombinedLayerOperation(false)
-		, bVisibilityChanged(false)
-		, bTargetIsHeightmap(InTarget.TargetType == ELandscapeToolTargetType::Heightmap)
-		, CacheUpToEditingLayer(Cache)
-		, CacheBottomLayers(InTarget)
-	{
-	}
-
-	void Initialize(ULandscapeInfo* InLandscapeInfo, bool InCombinedLayerOperation)
-	{
-		if (!bIsInitialized)
-		{
-			LandscapeInfo = InLandscapeInfo;
-			Landscape = LandscapeInfo ? LandscapeInfo->LandscapeActor.Get() : nullptr;
-			bCombinedLayerOperation = Landscape && Landscape->HasLayersContent() && InCombinedLayerOperation && bTargetIsHeightmap;
-			if (bCombinedLayerOperation)
-			{
-				EditingLayerGuid = Landscape->GetEditingLayer();
-				for (int i = 0; i < Landscape->GetLayerCount(); ++i)
-				{
-					FLandscapeLayer* CurrentLayer = Landscape->GetLayer(i);
-					BackupLayerVisibility.Add(CurrentLayer->bVisible);
-					if (CurrentLayer->Guid == EditingLayerGuid)
-					{
-						EditingLayerIndex = i;
-					}
-				}
-				check(EditingLayerIndex < Landscape->GetLayerCount());
-			}
-			bIsInitialized = true;
-		}
-	}
-
-	void Read(int32 X1, int32 Y1, int32 X2, int32 Y2, TArray<typename ToolTarget::CacheClass::DataType>& Data)
-	{
-		check(bIsInitialized);
-		if (bCombinedLayerOperation)
-		{
-			TArray<bool> NewLayerVisibility;
-			for (int i = 0; i < Landscape->GetLayerCount(); ++i)
-			{
-				FLandscapeLayer* CurrentLayer = Landscape->GetLayer(i);
-				NewLayerVisibility.Add((i > EditingLayerIndex) ? false : CurrentLayer->bVisible);
-			}
-
-			FIntRect Bounds;
-			const float BoundStep = LandscapeInfo->ComponentSizeQuads / 4;
-			Bounds.Min.X = (FMath::FloorToInt((float)X1 / BoundStep) - 1) * BoundStep;
-			Bounds.Min.Y = (FMath::FloorToInt((float)Y1 / BoundStep) - 1) * BoundStep;
-			Bounds.Max.X = (FMath::CeilToInt((float)X2 / BoundStep) + 1) * BoundStep;
-			Bounds.Max.Y = (FMath::CeilToInt((float)Y2 / BoundStep) + 1) * BoundStep;
-
-			FScopedSetLandscapeEditingLayer Scope(Landscape, FGuid());
-			CacheUpToEditingLayer.GetDataAndCache(X1, Y1, X2, Y2, Data, [&]() -> FIntRect
-			{
-				TSet<ULandscapeComponent*> AffectedComponents;
-				LandscapeInfo->GetComponentsInRegion(X1, Y1, X2, Y2, AffectedComponents);
-				SynchronousUpdateComponentVisibilityForHeight(AffectedComponents, NewLayerVisibility);
-				return Bounds;
-			});
-
-			CacheBottomLayers.GetDataAndCache(X1, Y1, X2, Y2, BottomLayersData, [&]() -> FIntRect
-			{
-				NewLayerVisibility[EditingLayerIndex] = false;
-				TSet<ULandscapeComponent*> AffectedComponents;
-				LandscapeInfo->GetComponentsInRegion(X1, Y1, X2, Y2, AffectedComponents);
-				SynchronousUpdateComponentVisibilityForHeight(AffectedComponents, NewLayerVisibility);
-				return Bounds;
-			});
-		}
-		else
-		{
-			CacheUpToEditingLayer.CacheData(X1, Y1, X2, Y2);
-			CacheUpToEditingLayer.GetCachedData(X1, Y1, X2, Y2, Data);
-		}
-	}
-
-	void Write(int32 X1, int32 Y1, int32 X2, int32 Y2, TArray<typename ToolTarget::CacheClass::DataType>& Data, ELandscapeLayerPaintingRestriction PaintingRestriction)
-	{
-		check(bIsInitialized);
-		if (bCombinedLayerOperation)
-		{
-			const float Alpha = Landscape->GetLayerAlpha(EditingLayerIndex, bTargetIsHeightmap);
-			const float InverseAlpha = (Alpha != 0.f) ? 1.f / Alpha : 1.f;
-			TArray<typename ToolTarget::CacheClass::DataType> DataContribution;
-			DataContribution.Empty(Data.Num());
-			DataContribution.AddUninitialized(Data.Num());
-			check(Data.Num() == BottomLayersData.Num());
-			for (int i = 0; i < Data.Num(); ++i)
-			{
-				float Contribution = (LandscapeDataAccess::GetLocalHeight(Data[i]) - LandscapeDataAccess::GetLocalHeight(BottomLayersData[i])) * InverseAlpha;
-				DataContribution[i] = LandscapeDataAccess::GetTexHeight(Contribution);
-			}
-			check(EditingLayerGuid == Landscape->GetEditingLayer());
-
-			// Restore layers visibility
-			SetLayersVisibility(BackupLayerVisibility);
-			// Only store data in cache
-			CacheUpToEditingLayer.SetCachedData(X1, Y1, X2, Y2, Data, PaintingRestriction, false);
-			// Effectively write the contribution
-			CacheUpToEditingLayer.DataAccess.SetData(X1, Y1, X2, Y2, DataContribution.GetData(), PaintingRestriction);
-			CacheUpToEditingLayer.DataAccess.Flush();
-			if (bVisibilityChanged)
-			{
-				TSet<ULandscapeComponent*> AffectedComponents;
-				LandscapeInfo->GetComponentsInRegion(X1, Y1, X2, Y2, AffectedComponents);
-				SynchronousUpdateHeightmapForComponents(AffectedComponents);
-				bVisibilityChanged = false;
-			}
-		}
-		else
-		{
-			CacheUpToEditingLayer.SetCachedData(X1, Y1, X2, Y2, Data, PaintingRestriction);
-			CacheUpToEditingLayer.Flush();
-		}
-	}
-
-private:
-
-	void SynchronousUpdateHeightmapForComponents(const TSet<ULandscapeComponent*>& InComponents)
-	{
-		for (ULandscapeComponent* Component : InComponents)
-		{
-			Component->RequestHeightmapUpdate();
-		}
-		Landscape->ForceUpdateLayersContent();
-	};
-
-	void SetLayersVisibility(const TArray<bool>& InLayerVisibility)
-	{
-		check(InLayerVisibility.Num() == Landscape->GetLayerCount());
-		for (int i = 0; i < InLayerVisibility.Num(); ++i)
-		{
-			if (FLandscapeLayer* Layer = Landscape->GetLayer(i))
-			{
-				if (Layer->bVisible != InLayerVisibility[i])
-				{
-					Layer->bVisible = InLayerVisibility[i];
-					bVisibilityChanged = true;
-				}
-			}
-		}
-	};
-
-	void SynchronousUpdateComponentVisibilityForHeight(const TSet<ULandscapeComponent*>& InComponents, const TArray<bool>& InLayerVisibility)
-	{
-		SetLayersVisibility(InLayerVisibility);
-		SynchronousUpdateHeightmapForComponents(InComponents);
-	};
-
-	ULandscapeInfo* LandscapeInfo;
-	ALandscape* Landscape;
-	FGuid EditingLayerGuid;
-	uint8 EditingLayerIndex;
-	TArray<bool> BackupLayerVisibility;
-	TArray<typename ToolTarget::CacheClass::DataType> BottomLayersData;
-	bool bIsInitialized;
-	bool bCombinedLayerOperation;
-	bool bVisibilityChanged;
-	bool bTargetIsHeightmap;
-
-	typename ToolTarget::CacheClass& CacheUpToEditingLayer;
-	typename ToolTarget::CacheClass CacheBottomLayers;
-};
-
 // 
 // FLandscapeToolSmooth
 //
diff --git a/Engine/Source/Editor/LandscapeEditor/Private/LandscapeEdModeTools.h b/Engine/Source/Editor/LandscapeEditor/Private/LandscapeEdModeTools.h
--- a/Engine/Source/Editor/LandscapeEditor/Private/LandscapeEdModeTools.h
+++ b/Engine/Source/Editor/LandscapeEditor/Private/LandscapeEdModeTools.h
@@ -853,6 +853,181 @@ struct FLandscapeVisCache : public TLandscapeEditCache<FAlphamapAccessorTool<fal
 	}
 };
 
+
+template<class ToolTarget>
+class FLandscapeLayerDataCache
+{
+public:
+	FLandscapeLayerDataCache(const FLandscapeToolTarget& InTarget, typename ToolTarget::CacheClass& Cache)
+		: LandscapeInfo(nullptr)
+		, Landscape(nullptr)
+		, EditingLayerIndex(MAX_uint8)
+		, bIsInitialized(false)
+		, bCombinedLayerOperation(false)
+		, bVisibilityChanged(false)
+		, bTargetIsHeightmap(InTarget.TargetType == ELandscapeToolTargetType::Heightmap)
+		, CacheUpToEditingLayer(Cache)
+		, CacheBottomLayers(InTarget)
+	{
+	}
+
+	void Initialize(ULandscapeInfo* InLandscapeInfo, bool InCombinedLayerOperation)
+	{
+		if (!bIsInitialized)
+		{
+			LandscapeInfo = InLandscapeInfo;
+			Landscape = LandscapeInfo ? LandscapeInfo->LandscapeActor.Get() : nullptr;
+			bCombinedLayerOperation = Landscape && Landscape->HasLayersContent() && InCombinedLayerOperation && bTargetIsHeightmap;
+			if (bCombinedLayerOperation)
+			{
+				EditingLayerGuid = Landscape->GetEditingLayer();
+				for (int i = 0; i < Landscape->GetLayerCount(); ++i)
+				{
+					FLandscapeLayer* CurrentLayer = Landscape->GetLayer(i);
+					BackupLayerVisibility.Add(CurrentLayer->bVisible);
+					if (CurrentLayer->Guid == EditingLayerGuid)
+					{
+						EditingLayerIndex = i;
+					}
+				}
+				check(EditingLayerIndex < Landscape->GetLayerCount());
+			}
+			bIsInitialized = true;
+		}
+	}
+
+	void Read(int32 X1, int32 Y1, int32 X2, int32 Y2, TArray<typename ToolTarget::CacheClass::DataType>& Data)
+	{
+		check(bIsInitialized);
+		if (bCombinedLayerOperation)
+		{
+			TArray<bool> NewLayerVisibility;
+			for (int i = 0; i < Landscape->GetLayerCount(); ++i)
+			{
+				FLandscapeLayer* CurrentLayer = Landscape->GetLayer(i);
+				NewLayerVisibility.Add((i > EditingLayerIndex) ? false : CurrentLayer->bVisible);
+			}
+
+			FIntRect Bounds;
+			const float BoundStep = LandscapeInfo->ComponentSizeQuads / 4;
+			Bounds.Min.X = (FMath::FloorToInt((float)X1 / BoundStep) - 1) * BoundStep;
+			Bounds.Min.Y = (FMath::FloorToInt((float)Y1 / BoundStep) - 1) * BoundStep;
+			Bounds.Max.X = (FMath::CeilToInt((float)X2 / BoundStep) + 1) * BoundStep;
+			Bounds.Max.Y = (FMath::CeilToInt((float)Y2 / BoundStep) + 1) * BoundStep;
+
+			FScopedSetLandscapeEditingLayer Scope(Landscape, FGuid());
+			CacheUpToEditingLayer.GetDataAndCache(X1, Y1, X2, Y2, Data, [&]() -> FIntRect
+			{
+				TSet<ULandscapeComponent*> AffectedComponents;
+				LandscapeInfo->GetComponentsInRegion(X1, Y1, X2, Y2, AffectedComponents);
+				SynchronousUpdateComponentVisibilityForHeight(AffectedComponents, NewLayerVisibility);
+				return Bounds;
+			});
+
+			CacheBottomLayers.GetDataAndCache(X1, Y1, X2, Y2, BottomLayersData, [&]() -> FIntRect
+			{
+				NewLayerVisibility[EditingLayerIndex] = false;
+				TSet<ULandscapeComponent*> AffectedComponents;
+				LandscapeInfo->GetComponentsInRegion(X1, Y1, X2, Y2, AffectedComponents);
+				SynchronousUpdateComponentVisibilityForHeight(AffectedComponents, NewLayerVisibility);
+				return Bounds;
+			});
+		}
+		else
+		{
+			CacheUpToEditingLayer.CacheData(X1, Y1, X2, Y2);
+			CacheUpToEditingLayer.GetCachedData(X1, Y1, X2, Y2, Data);
+		}
+	}
+
+	void Write(int32 X1, int32 Y1, int32 X2, int32 Y2, TArray<typename ToolTarget::CacheClass::DataType>& Data, ELandscapeLayerPaintingRestriction PaintingRestriction = ELandscapeLayerPaintingRestriction::None)
+	{
+		check(bIsInitialized);
+		if (bCombinedLayerOperation)
+		{
+			const float Alpha = Landscape->GetLayerAlpha(EditingLayerIndex, bTargetIsHeightmap);
+			const float InverseAlpha = (Alpha != 0.f) ? 1.f / Alpha : 1.f;
+			TArray<typename ToolTarget::CacheClass::DataType> DataContribution;
+			DataContribution.Empty(Data.Num());
+			DataContribution.AddUninitialized(Data.Num());
+			check(Data.Num() == BottomLayersData.Num());
+			for (int i = 0; i < Data.Num(); ++i)
+			{
+				float Contribution = (LandscapeDataAccess::GetLocalHeight(Data[i]) - LandscapeDataAccess::GetLocalHeight(BottomLayersData[i])) * InverseAlpha;
+				DataContribution[i] = LandscapeDataAccess::GetTexHeight(Contribution);
+			}
+			check(EditingLayerGuid == Landscape->GetEditingLayer());
+
+			// Restore layers visibility
+			SetLayersVisibility(BackupLayerVisibility);
+			// Only store data in cache
+			CacheUpToEditingLayer.SetCachedData(X1, Y1, X2, Y2, Data, PaintingRestriction, false);
+			// Effectively write the contribution
+			CacheUpToEditingLayer.DataAccess.SetData(X1, Y1, X2, Y2, DataContribution.GetData(), PaintingRestriction);
+			CacheUpToEditingLayer.DataAccess.Flush();
+			if (bVisibilityChanged)
+			{
+				TSet<ULandscapeComponent*> AffectedComponents;
+				LandscapeInfo->GetComponentsInRegion(X1, Y1, X2, Y2, AffectedComponents);
+				SynchronousUpdateHeightmapForComponents(AffectedComponents);
+				bVisibilityChanged = false;
+			}
+		}
+		else
+		{
+			CacheUpToEditingLayer.SetCachedData(X1, Y1, X2, Y2, Data, PaintingRestriction);
+			CacheUpToEditingLayer.Flush();
+		}
+	}
+
+private:
+
+	void SynchronousUpdateHeightmapForComponents(const TSet<ULandscapeComponent*>& InComponents)
+	{
+		for (ULandscapeComponent* Component : InComponents)
+		{
+			Component->RequestHeightmapUpdate();
+		}
+		Landscape->ForceUpdateLayersContent();
+	};
+
+	void SetLayersVisibility(const TArray<bool>& InLayerVisibility)
+	{
+		check(InLayerVisibility.Num() == Landscape->GetLayerCount());
+		for (int i = 0; i < InLayerVisibility.Num(); ++i)
+		{
+			if (FLandscapeLayer* Layer = Landscape->GetLayer(i))
+			{
+				if (Layer->bVisible != InLayerVisibility[i])
+				{
+					Layer->bVisible = InLayerVisibility[i];
+					bVisibilityChanged = true;
+				}
+			}
+		}
+	};
+
+	void SynchronousUpdateComponentVisibilityForHeight(const TSet<ULandscapeComponent*>& InComponents, const TArray<bool>& InLayerVisibility)
+	{
+		SetLayersVisibility(InLayerVisibility);
+		SynchronousUpdateHeightmapForComponents(InComponents);
+	};
+
+	ULandscapeInfo* LandscapeInfo;
+	ALandscape* Landscape;
+	FGuid EditingLayerGuid;
+	uint8 EditingLayerIndex;
+	TArray<bool> BackupLayerVisibility;
+	TArray<typename ToolTarget::CacheClass::DataType> BottomLayersData;
+	bool bIsInitialized;
+	bool bCombinedLayerOperation;
+	bool bVisibilityChanged;
+	bool bTargetIsHeightmap;
+
+	typename ToolTarget::CacheClass& CacheUpToEditingLayer;
+	typename ToolTarget::CacheClass CacheBottomLayers;
+};
+
 //
 // FFullWeightmapAccessor
 //
diff --git a/Engine/Source/Editor/LandscapeEditor/Public/LandscapeEditorObject.h b/Engine/Source/Editor/LandscapeEditor/Public/LandscapeEditorObject.h
--- a/Engine/Source/Editor/LandscapeEditor/Public/LandscapeEditorObject.h
+++ b/Engine/Source/Editor/LandscapeEditor/Public/LandscapeEditorObject.h
@@ -266,7 +266,7 @@ class ULandscapeEditorObject : public UObject
 	UPROPERTY(NonTransactional)
 	float MaximumValueRadius;
 
-	UPROPERTY(Category="Tool Settings", EditAnywhere, NonTransactional, meta=(ShowForTools="Flatten,Smooth", ShowForTargetTypes="Heightmap", ShowForLandscapeLayerSystem))
+	UPROPERTY(Category="Tool Settings", EditAnywhere, NonTransactional, meta=(ShowForTools="Flatten,Smooth,Erosion,HydraErosion", ShowForTargetTypes="Heightmap", ShowForLandscapeLayerSystem))
 	bool bCombinedLayersOperation;
 
 	// Flatten Tool:
