Fix Static Analysis: swap ensure for check

#rnx

diff --git a/Engine/Source/Editor/LandscapeEditor/Private/LandscapeEditorDetailCustomization_Layers.cpp b/Engine/Source/Editor/LandscapeEditor/Private/LandscapeEditorDetailCustomization_Layers.cpp
--- a/Engine/Source/Editor/LandscapeEditor/Private/LandscapeEditorDetailCustomization_Layers.cpp
+++ b/Engine/Source/Editor/LandscapeEditor/Private/LandscapeEditorDetailCustomization_Layers.cpp
@@ -388,9 +388,9 @@ void FLandscapeEditorCustomNodeBuilder_Layers::FillClearLayerMenu(FMenuBuilder&
 			{
 				TArray<ULandscapeLayerInfoObject*> UsedLayerInfos;
 				FEdModeLandscape* LandscapeEdMode = GetEditorMode();
-				ensure(LandscapeEdMode);
+				check(LandscapeEdMode);
 				ALandscape* Landscape = LandscapeEdMode->GetLandscape();
-				ensure(Landscape);
+				check(Landscape);
 
 				Landscape->GetUsedPaintLayers(InLayerIndex, UsedLayerInfos);
 				if (UsedLayerInfos.Num() > 0)
