Remove useless if that causes ensure fail. Always declare LOCTEXT.

#rnx

diff --git a/Engine/Source/Editor/LandscapeEditor/Private/LandscapeEditorDetails.cpp b/Engine/Source/Editor/LandscapeEditor/Private/LandscapeEditorDetails.cpp
--- a/Engine/Source/Editor/LandscapeEditor/Private/LandscapeEditorDetails.cpp
+++ b/Engine/Source/Editor/LandscapeEditor/Private/LandscapeEditorDetails.cpp
@@ -170,12 +170,8 @@ FText FLandscapeEditorDetails::GetLocalizedName(FString Name)
 		LOCTEXT("ToolSet_Noise", "Noise");
 		LOCTEXT("ToolSet_Retopologize", "Retopologize");
 		LOCTEXT("ToolSet_Visibility", "Visibility");
-
-		if (LandscapeEdMode->CanHaveLandscapeLayersContent())
-		{
-			LOCTEXT("ToolSet_BlueprintBrush", "Blueprint Brushes");
-		}
-
+		LOCTEXT("ToolSet_BlueprintBrush", "Blueprint Brushes");
+		
 		LOCTEXT("ToolSet_Select", "Selection");
 		LOCTEXT("ToolSet_AddComponent", "Add");
 		LOCTEXT("ToolSet_DeleteComponent", "Delete");
