Landscape Improvement: remember previous mode tool and target layer

#rnx

diff --git a/Engine/Source/Editor/LandscapeEditor/Private/LandscapeEdMode.cpp b/Engine/Source/Editor/LandscapeEditor/Private/LandscapeEdMode.cpp
--- a/Engine/Source/Editor/LandscapeEditor/Private/LandscapeEdMode.cpp
+++ b/Engine/Source/Editor/LandscapeEditor/Private/LandscapeEdMode.cpp
@@ -409,8 +409,17 @@ void FEdModeLandscape::AddReferencedObjects(FReferenceCollector& Collector)
 
 void FEdModeLandscape::UpdateToolModes()
 {
-	LandscapeToolModes.Reset();
+	// Keep mapping of CurrentTool and CurrentTargetLayer
+	TMap<FName, FName> PreviousTools;
+	TMap<FName, FName> PreviousTargetLayerNames;
+	for (const FLandscapeToolMode& Previous : LandscapeToolModes)
+	{
+		PreviousTools.Add(Previous.ToolModeName, Previous.CurrentToolName);
+		PreviousTargetLayerNames.Add(Previous.ToolModeName, Previous.CurrentTargetLayerName);
+	}
 
+	LandscapeToolModes.Reset();
+	
 	FLandscapeToolMode* ToolMode_Manage = new(LandscapeToolModes)FLandscapeToolMode(TEXT("ToolMode_Manage"), ELandscapeToolTargetTypeMask::NA);
 	ToolMode_Manage->ValidTools.Add(TEXT("NewLandscape"));
 	ToolMode_Manage->ValidTools.Add(TEXT("Select"));
@@ -419,7 +428,14 @@ void FEdModeLandscape::UpdateToolModes()
 	ToolMode_Manage->ValidTools.Add(TEXT("MoveToLevel"));
 	ToolMode_Manage->ValidTools.Add(TEXT("ResizeLandscape"));
 	ToolMode_Manage->ValidTools.Add(TEXT("Splines"));
-	ToolMode_Manage->CurrentToolName = TEXT("Select");
+	
+	// Restore
+	FName* PreviousToolName = PreviousTools.Find(ToolMode_Manage->ToolModeName);
+	ToolMode_Manage->CurrentToolName = PreviousToolName ? *PreviousToolName : TEXT("Select");
+
+	FName* PreviousTargetLayerName = PreviousTargetLayerNames.Find(ToolMode_Manage->ToolModeName);
+	ToolMode_Manage->CurrentTargetLayerName = PreviousTargetLayerName ? *PreviousTargetLayerName : NAME_None;
+		
 
 	FLandscapeToolMode* ToolMode_Sculpt = new(LandscapeToolModes)FLandscapeToolMode(TEXT("ToolMode_Sculpt"), ELandscapeToolTargetTypeMask::Heightmap | ELandscapeToolTargetTypeMask::Visibility);
 	ToolMode_Sculpt->ValidTools.Add(TEXT("Sculpt"));
@@ -444,6 +460,13 @@ void FEdModeLandscape::UpdateToolModes()
 	ToolMode_Sculpt->ValidTools.Add(TEXT("Mask"));
 	ToolMode_Sculpt->ValidTools.Add(TEXT("CopyPaste"));
 	ToolMode_Sculpt->ValidTools.Add(TEXT("Mirror"));
+	
+	// Restore
+	PreviousToolName = PreviousTools.Find(ToolMode_Sculpt->ToolModeName);
+	ToolMode_Sculpt->CurrentToolName = PreviousToolName ? *PreviousToolName : TEXT("Sculpt");
+
+	PreviousTargetLayerName = PreviousTargetLayerNames.Find(ToolMode_Sculpt->ToolModeName);
+	ToolMode_Sculpt->CurrentTargetLayerName = PreviousTargetLayerName ? *PreviousTargetLayerName : NAME_None;
 
 	FLandscapeToolMode* ToolMode_Paint = new(LandscapeToolModes)FLandscapeToolMode(TEXT("ToolMode_Paint"), ELandscapeToolTargetTypeMask::Weightmap);
 	ToolMode_Paint->ValidTools.Add(TEXT("Paint"));
@@ -457,10 +480,16 @@ void FEdModeLandscape::UpdateToolModes()
 		ToolMode_Paint->ValidTools.Add(TEXT("BlueprintBrush"));
 	}
 
+	PreviousToolName = PreviousTools.Find(ToolMode_Paint->ToolModeName);
+	ToolMode_Paint->CurrentToolName = PreviousToolName ? *PreviousToolName : TEXT("Paint");
+
+	PreviousTargetLayerName = PreviousTargetLayerNames.Find(ToolMode_Paint->ToolModeName);
+	ToolMode_Paint->CurrentTargetLayerName = PreviousTargetLayerName ? *PreviousTargetLayerName : NAME_None;
+
 	// Since available tools might have changed try and reset the current tool
 	if (CurrentToolMode && CurrentToolIndex != INDEX_NONE)
 	{
-		SetCurrentTool(CurrentToolIndex);
+		SetCurrentTool(CurrentToolIndex, CurrentToolMode->CurrentTargetLayerName);
 	}
 }
 
@@ -2189,8 +2218,9 @@ void FEdModeLandscape::SetCurrentToolMode(FName ToolModeName, bool bRestoreCurre
 					if (CurrentToolMode->CurrentToolName == NAME_None)
 					{
 						CurrentToolMode->CurrentToolName = CurrentToolMode->ValidTools[0];
+						CurrentToolMode->CurrentTargetLayerName = NAME_None;
 					}
-					SetCurrentTool(CurrentToolMode->CurrentToolName);
+					SetCurrentTool(CurrentToolMode->CurrentToolName, CurrentToolMode->CurrentTargetLayerName);
 				}
 				break;
 			}
@@ -2198,7 +2228,7 @@ void FEdModeLandscape::SetCurrentToolMode(FName ToolModeName, bool bRestoreCurre
 	}
 }
 
-void FEdModeLandscape::SetCurrentTool(FName ToolName)
+void FEdModeLandscape::SetCurrentTool(FName ToolName, FName TargetLayerName)
 {
 	// Several tools have identically named versions for sculpting and painting
 	// Prefer the one with the same target type as the current mode
@@ -2229,10 +2259,21 @@ void FEdModeLandscape::SetCurrentTool(FName ToolName)
 	}
 	check(ToolIndex != INDEX_NONE);
 
-	SetCurrentTool(ToolIndex);
+	SetCurrentTool(ToolIndex, TargetLayerName);
+}
+
+void FEdModeLandscape::SetCurrentTargetLayer(FName TargetLayerName, TWeakObjectPtr<ULandscapeLayerInfoObject> LayerInfo)
+{
+	if (CurrentToolMode)
+	{
+		// Cache current Layer Name so we can set it back when switching between Modes
+		CurrentToolMode->CurrentTargetLayerName = TargetLayerName;
+	}
+	CurrentToolTarget.LayerName = TargetLayerName;
+	CurrentToolTarget.LayerInfo = LayerInfo;
 }
 
-void FEdModeLandscape::SetCurrentTool(int32 ToolIndex)
+void FEdModeLandscape::SetCurrentTool(int32 ToolIndex, FName TargetLayerName)
 {
 	if (CurrentTool)
 	{
@@ -2272,8 +2313,7 @@ void FEdModeLandscape::SetCurrentTool(int32 ToolIndex)
 	if (CurrentTool->GetSupportedTargetTypes() == ELandscapeToolTargetTypeMask::NA)
 	{
 		CurrentToolTarget.TargetType = ELandscapeToolTargetType::Invalid;
-		CurrentToolTarget.LayerInfo = nullptr;
-		CurrentToolTarget.LayerName = NAME_None;
+		SetCurrentTargetLayer(NAME_None, nullptr);
 	}
 	else
 	{
@@ -2282,20 +2322,18 @@ void FEdModeLandscape::SetCurrentTool(int32 ToolIndex)
 
 		if ((TargetTypeMask & ELandscapeToolTargetTypeMask::FromType(CurrentToolTarget.TargetType)) == 0)
 		{
-			auto filter = [TargetTypeMask](const TSharedRef<FLandscapeTargetListInfo>& Target){ return (TargetTypeMask & ELandscapeToolTargetTypeMask::FromType(Target->TargetType)) != 0; };
+			auto filter = [TargetTypeMask, TargetLayerName](const TSharedRef<FLandscapeTargetListInfo>& Target){ return (TargetTypeMask & ELandscapeToolTargetTypeMask::FromType(Target->TargetType)) != 0 && (TargetLayerName == NAME_None || TargetLayerName == Target->LayerName); };
 			const TSharedRef<FLandscapeTargetListInfo>* Target = LandscapeTargetList.FindByPredicate(filter);
 			if (Target != nullptr)
 			{
 				check(CurrentToolTarget.LandscapeInfo == (*Target)->LandscapeInfo);
 				CurrentToolTarget.TargetType = (*Target)->TargetType;
-				CurrentToolTarget.LayerInfo = (*Target)->LayerInfoObj;
-				CurrentToolTarget.LayerName = (*Target)->LayerName;
+				SetCurrentTargetLayer((*Target)->LayerName, (*Target)->LayerInfoObj);
 			}
 			else // can happen with for example paint tools if there are no paint layers defined
 			{
 				CurrentToolTarget.TargetType = ELandscapeToolTargetType::Invalid;
-				CurrentToolTarget.LayerInfo = nullptr;
-				CurrentToolTarget.LayerName = NAME_None;
+				SetCurrentTargetLayer(NAME_None, nullptr);
 			}
 		}
 	}
@@ -2680,8 +2718,7 @@ void FEdModeLandscape::UpdateTargetList()
 			if (!bFoundSelected)
 			{
 				CurrentToolTarget.TargetType = ELandscapeToolTargetType::Invalid;
-				CurrentToolTarget.LayerInfo = nullptr;
-				CurrentToolTarget.LayerName = NAME_None;
+				SetCurrentTargetLayer(NAME_None, nullptr);
 			}
 
 			UpdateTargetLayerDisplayOrder(UISettings->TargetDisplayOrder);
diff --git a/Engine/Source/Editor/LandscapeEditor/Private/LandscapeEdMode.h b/Engine/Source/Editor/LandscapeEditor/Private/LandscapeEdMode.h
--- a/Engine/Source/Editor/LandscapeEditor/Private/LandscapeEdMode.h
+++ b/Engine/Source/Editor/LandscapeEditor/Private/LandscapeEdMode.h
@@ -49,6 +49,7 @@ struct FLandscapeToolMode
 
 	TArray<FName>			ValidTools;
 	FName					CurrentToolName;
+	FName					CurrentTargetLayerName;
 
 	FLandscapeToolMode(FName InToolModeName, int32 InSupportedTargetTypes)
 		: ToolModeName(InToolModeName)
@@ -500,8 +501,9 @@ public:
 	void SetCurrentToolMode(FName ToolModeName, bool bRestoreCurrentTool = true);
 
 	/** Change current tool */
-	void SetCurrentTool(FName ToolName);
-	void SetCurrentTool(int32 ToolIdx);
+	void SetCurrentTool(FName ToolName, FName TargetLayerName = NAME_None);
+	void SetCurrentTool(int32 ToolIdx, FName TargetLayerName = NAME_None);
+	void SetCurrentTargetLayer(FName TargetLayerName, TWeakObjectPtr<ULandscapeLayerInfoObject> LayerInfo);
 
 	void SetCurrentBrushSet(FName BrushSetName);
 	void SetCurrentBrushSet(int32 BrushSetIndex);
diff --git a/Engine/Source/Editor/LandscapeEditor/Private/LandscapeEditorDetailCustomization_NewLandscape.cpp b/Engine/Source/Editor/LandscapeEditor/Private/LandscapeEditorDetailCustomization_NewLandscape.cpp
--- a/Engine/Source/Editor/LandscapeEditor/Private/LandscapeEditorDetailCustomization_NewLandscape.cpp
+++ b/Engine/Source/Editor/LandscapeEditor/Private/LandscapeEditorDetailCustomization_NewLandscape.cpp
@@ -872,8 +872,7 @@ FReply FLandscapeEditorDetailCustomization_NewLandscape::OnCreateButtonClicked()
 		LandscapeEdMode->UpdateLandscapeList();
 		LandscapeEdMode->SetLandscapeInfo(LandscapeInfo);
 		LandscapeEdMode->CurrentToolTarget.TargetType = ELandscapeToolTargetType::Heightmap;
-		LandscapeEdMode->CurrentToolTarget.LayerInfo = nullptr;
-		LandscapeEdMode->CurrentToolTarget.LayerName = NAME_None;
+		LandscapeEdMode->SetCurrentTargetLayer(NAME_None, nullptr);
 		LandscapeEdMode->UpdateTargetList();
 
 		LandscapeEdMode->SetCurrentTool("Select"); // change tool so switching back to the manage mode doesn't give "New Landscape" again
diff --git a/Engine/Source/Editor/LandscapeEditor/Private/LandscapeEditorDetailCustomization_TargetLayers.cpp b/Engine/Source/Editor/LandscapeEditor/Private/LandscapeEditorDetailCustomization_TargetLayers.cpp
--- a/Engine/Source/Editor/LandscapeEditor/Private/LandscapeEditorDetailCustomization_TargetLayers.cpp
+++ b/Engine/Source/Editor/LandscapeEditor/Private/LandscapeEditorDetailCustomization_TargetLayers.cpp
@@ -885,13 +885,11 @@ void FLandscapeEditorCustomNodeBuilder_TargetLayers::OnTargetSelectionChanged(co
 		if (Target->TargetType == ELandscapeToolTargetType::Heightmap)
 		{
 			checkSlow(Target->LayerInfoObj == NULL);
-			LandscapeEdMode->CurrentToolTarget.LayerInfo = NULL;
-			LandscapeEdMode->CurrentToolTarget.LayerName = NAME_None;
+			LandscapeEdMode->SetCurrentTargetLayer(NAME_None, nullptr);
 		}
 		else
 		{
-			LandscapeEdMode->CurrentToolTarget.LayerInfo = Target->LayerInfoObj;
-			LandscapeEdMode->CurrentToolTarget.LayerName = Target->LayerName;
+			LandscapeEdMode->SetCurrentTargetLayer(Target->LayerName, Target->LayerInfoObj);
 		}
 	}
 }
@@ -1203,9 +1201,8 @@ void FLandscapeEditorCustomNodeBuilder_TargetLayers::OnTargetLayerSetObject(cons
 			FEdModeLandscape* LandscapeEdMode = GetEditorMode();
 			if (LandscapeEdMode)
 			{
-				LandscapeEdMode->CurrentToolTarget.LayerName = Target->LayerName;
 				LandscapeEdMode->CurrentToolTarget.TargetType = Target->TargetType;
-				LandscapeEdMode->CurrentToolTarget.LayerInfo = SelectedLayerInfo;
+				LandscapeEdMode->SetCurrentTargetLayer(Target->LayerName, SelectedLayerInfo);
 				LandscapeEdMode->UpdateTargetList();
 			}
 
@@ -1322,7 +1319,7 @@ void FLandscapeEditorCustomNodeBuilder_TargetLayers::OnTargetLayerCreateClicked(
 			if (LandscapeEdMode->CurrentToolTarget.LayerName == Target->LayerName
 				&& LandscapeEdMode->CurrentToolTarget.LayerInfo == Target->LayerInfoObj)
 			{
-				LandscapeEdMode->CurrentToolTarget.LayerInfo = LayerInfo;
+				LandscapeEdMode->SetCurrentTargetLayer(Target->LayerName, Target->LayerInfoObj);
 			}
 
 			Target->LayerInfoObj = LayerInfo;
