Copy RVT settings when inserting a new segement in a spline
Integrated from //UE4/Dev-Rendering
#rnx

diff --git a/Engine/Source/Editor/LandscapeEditor/Private/LandscapeEdModeSplineTools.cpp b/Engine/Source/Editor/LandscapeEditor/Private/LandscapeEdModeSplineTools.cpp
--- a/Engine/Source/Editor/LandscapeEditor/Private/LandscapeEdModeSplineTools.cpp
+++ b/Engine/Source/Editor/LandscapeEditor/Private/LandscapeEdModeSplineTools.cpp
@@ -318,14 +318,16 @@ public:
 
 		if (CopyFromSegment != nullptr)
 		{
-			NewSegment->LayerName         = CopyFromSegment->LayerName;
-			NewSegment->SplineMeshes      = CopyFromSegment->SplineMeshes;
+			NewSegment->LayerName = CopyFromSegment->LayerName;
+			NewSegment->SplineMeshes = CopyFromSegment->SplineMeshes;
 			NewSegment->LDMaxDrawDistance = CopyFromSegment->LDMaxDrawDistance;
-			NewSegment->bRaiseTerrain     = CopyFromSegment->bRaiseTerrain;
-			NewSegment->bLowerTerrain     = CopyFromSegment->bLowerTerrain;
+			NewSegment->bRaiseTerrain = CopyFromSegment->bRaiseTerrain;
+			NewSegment->bLowerTerrain = CopyFromSegment->bLowerTerrain;
 			NewSegment->bPlaceSplineMeshesInStreamingLevels = CopyFromSegment->bPlaceSplineMeshesInStreamingLevels;
-			NewSegment->BodyInstance  = CopyFromSegment->BodyInstance;
-			NewSegment->bCastShadow       = CopyFromSegment->bCastShadow;
+			NewSegment->BodyInstance = CopyFromSegment->BodyInstance;
+			NewSegment->bCastShadow = CopyFromSegment->bCastShadow;
+			NewSegment->TranslucencySortPriority = CopyFromSegment->TranslucencySortPriority;
+			NewSegment->RuntimeVirtualTextures = CopyFromSegment->RuntimeVirtualTextures;
 		}
 
 		Start->ConnectedSegments.Add(FLandscapeSplineConnection(NewSegment, 0));
@@ -681,6 +683,8 @@ public:
 		NewSegment->bLowerTerrain = Segment->bLowerTerrain;
 		NewSegment->BodyInstance = Segment->BodyInstance;
 		NewSegment->bCastShadow = Segment->bCastShadow;
+		NewSegment->TranslucencySortPriority = Segment->TranslucencySortPriority;
+		NewSegment->RuntimeVirtualTextures = Segment->RuntimeVirtualTextures;
 
 		Segment->Connections[0].TangentLen *= t;
 		Segment->Connections[1].ControlPoint->ConnectedSegments.Remove(FLandscapeSplineConnection(Segment, 1));
