Copy more RVT settings when inserting a new segement in a spline
Integrated from //UE4/Dev-Rendering
#rnx

diff --git a/Engine/Source/Editor/LandscapeEditor/Private/LandscapeEdModeSplineTools.cpp b/Engine/Source/Editor/LandscapeEditor/Private/LandscapeEdModeSplineTools.cpp
--- a/Engine/Source/Editor/LandscapeEditor/Private/LandscapeEdModeSplineTools.cpp
+++ b/Engine/Source/Editor/LandscapeEditor/Private/LandscapeEdModeSplineTools.cpp
@@ -328,6 +328,7 @@ public:
 			NewSegment->bCastShadow = CopyFromSegment->bCastShadow;
 			NewSegment->TranslucencySortPriority = CopyFromSegment->TranslucencySortPriority;
 			NewSegment->RuntimeVirtualTextures = CopyFromSegment->RuntimeVirtualTextures;
+			NewSegment->VirtualTextureRenderPassType = CopyFromSegment->VirtualTextureRenderPassType;
 		}
 
 		Start->ConnectedSegments.Add(FLandscapeSplineConnection(NewSegment, 0));
@@ -685,6 +686,7 @@ public:
 		NewSegment->bCastShadow = Segment->bCastShadow;
 		NewSegment->TranslucencySortPriority = Segment->TranslucencySortPriority;
 		NewSegment->RuntimeVirtualTextures = Segment->RuntimeVirtualTextures;
+		NewSegment->VirtualTextureRenderPassType = Segment->VirtualTextureRenderPassType;
 
 		Segment->Connections[0].TangentLen *= t;
 		Segment->Connections[1].ControlPoint->ConnectedSegments.Remove(FLandscapeSplineConnection(Segment, 1));
