Fix Add Component Tool not needing the Current Editing Layer

#rnx

diff --git a/Engine/Source/Editor/LandscapeEditor/Private/LandscapeEdModeComponentTools.cpp b/Engine/Source/Editor/LandscapeEditor/Private/LandscapeEdModeComponentTools.cpp
--- a/Engine/Source/Editor/LandscapeEditor/Private/LandscapeEdModeComponentTools.cpp
+++ b/Engine/Source/Editor/LandscapeEditor/Private/LandscapeEdModeComponentTools.cpp
@@ -966,13 +966,14 @@ public:
 		: FLandscapeToolBase<FLandscapeToolStrokeAddComponent>(InEdMode)
 	{
 	}
+	virtual bool ShouldUpdateEditingLayer() const override { return false; }
 
 	virtual const TCHAR* GetToolName() override { return TEXT("AddComponent"); }
 	virtual FText GetDisplayName() override { return NSLOCTEXT("UnrealEd", "LandscapeMode_AddComponent", "Add New Landscape Component"); };
 
 	virtual void SetEditRenderType() override { GLandscapeEditRenderMode = ELandscapeEditRenderMode::None | (GLandscapeEditRenderMode & ELandscapeEditRenderMode::BitMaskForMask); }
 	virtual bool SupportsMask() override { return false; }
-
+	
 	virtual void EnterTool() override
 	{
 		FLandscapeToolBase<FLandscapeToolStrokeAddComponent>::EnterTool();
