Replace AddLambda by AddRaw (removed with RemoveAll(this))

#rnx

diff --git a/Engine/Source/Editor/LandscapeEditor/Private/LandscapeEdMode.cpp b/Engine/Source/Editor/LandscapeEditor/Private/LandscapeEdMode.cpp
--- a/Engine/Source/Editor/LandscapeEditor/Private/LandscapeEdMode.cpp
+++ b/Engine/Source/Editor/LandscapeEditor/Private/LandscapeEdMode.cpp
@@ -621,7 +621,7 @@ void FEdModeLandscape::Enter()
 
 		if (ALandscape* Landscape = GetLandscape())
 		{
-			Landscape->OnBlueprintBrushChangedDelegate().AddLambda([this]() { this->RefreshDetailPanel(); });
+			Landscape->OnBlueprintBrushChangedDelegate().AddRaw(this, &FEdModeLandscape::RefreshDetailPanel);
 			if (Landscape->HasLayersContent())
 			{
 				if (Landscape->GetLandscapeSplinesReservedLayer())
@@ -2598,7 +2598,7 @@ void FEdModeLandscape::SetTargetLandscape(const TWeakObjectPtr<ULandscapeInfo>&
 		LandscapeProxy->OnMaterialChangedDelegate().AddRaw(this, &FEdModeLandscape::OnLandscapeMaterialChangedDelegate);
 		if (ALandscape* Landscape = GetLandscape())
 		{
-			Landscape->OnBlueprintBrushChangedDelegate().AddLambda([this]() { this->RefreshDetailPanel(); });
+			Landscape->OnBlueprintBrushChangedDelegate().AddRaw(this, &FEdModeLandscape::RefreshDetailPanel);
 		}
 	}
 
