PR #6033: Fixes for IntelliSense (Contributed by [REDACTED_PERSON])

#rnx

diff --git a/Engine/Plugins/Runtime/AndroidPermission/Source/AndroidPermission/Classes/AndroidPermissionFunctionLibrary.h b/Engine/Plugins/Runtime/AndroidPermission/Source/AndroidPermission/Classes/AndroidPermissionFunctionLibrary.h
--- a/Engine/Plugins/Runtime/AndroidPermission/Source/AndroidPermission/Classes/AndroidPermissionFunctionLibrary.h
+++ b/Engine/Plugins/Runtime/AndroidPermission/Source/AndroidPermission/Classes/AndroidPermissionFunctionLibrary.h
@@ -6,6 +6,8 @@
 #include "Kismet/BlueprintFunctionLibrary.h"
 #include "AndroidPermissionFunctionLibrary.generated.h"
 
+class UAndroidPermissionCallbackProxy;
+
 UCLASS()
 class ANDROIDPERMISSION_API UAndroidPermissionFunctionLibrary : public UBlueprintFunctionLibrary
 {
diff --git a/Engine/Source/Developer/AssetTools/Private/AssetTypeActions/AssetTypeActions_NodeMappingContainer.h b/Engine/Source/Developer/AssetTools/Private/AssetTypeActions/AssetTypeActions_NodeMappingContainer.h
--- a/Engine/Source/Developer/AssetTools/Private/AssetTypeActions/AssetTypeActions_NodeMappingContainer.h
+++ b/Engine/Source/Developer/AssetTools/Private/AssetTypeActions/AssetTypeActions_NodeMappingContainer.h
@@ -4,6 +4,7 @@
 
 #include "CoreMinimal.h"
 #include "Animation/NodeMappingContainer.h"
+#include "AssetTypeActions_Base.h"
 
 class FAssetTypeActions_NodeMappingContainer : public FAssetTypeActions_Base
 {
diff --git a/Engine/Source/Developer/LauncherServices/Private/Launcher/LauncherUATCommand.h b/Engine/Source/Developer/LauncherServices/Private/Launcher/LauncherUATCommand.h
--- a/Engine/Source/Developer/LauncherServices/Private/Launcher/LauncherUATCommand.h
+++ b/Engine/Source/Developer/LauncherServices/Private/Launcher/LauncherUATCommand.h
@@ -2,6 +2,9 @@
 
 #pragma once
 
+#include "Containers/UnrealString.h"
+
+struct FLauncherTaskChainState;
 
 /**
  * class for setting up UAT command arguments for various launcher commands
diff --git a/Engine/Source/Developer/MeshSimplifier/Private/Cache.h b/Engine/Source/Developer/MeshSimplifier/Private/Cache.h
--- a/Engine/Source/Developer/MeshSimplifier/Private/Cache.h
+++ b/Engine/Source/Developer/MeshSimplifier/Private/Cache.h
@@ -3,6 +3,9 @@
 
 #pragma once
 
+#include "CoreTypes.h"
+#include "HAL/UnrealMemory.h"
+
 /*
 ===============================================================================
 	Direct Mapped Cache
diff --git a/Engine/Source/Developer/ScreenShotComparison/Private/Models/ScreenShotComparisonFilter.h b/Engine/Source/Developer/ScreenShotComparison/Private/Models/ScreenShotComparisonFilter.h
--- a/Engine/Source/Developer/ScreenShotComparison/Private/Models/ScreenShotComparisonFilter.h
+++ b/Engine/Source/Developer/ScreenShotComparison/Private/Models/ScreenShotComparisonFilter.h
@@ -2,7 +2,6 @@
 #pragma once
 
 #include "CoreMinimal.h"
-#include "Interfaces/IScreenShotData.h"
 #include "Misc/IFilter.h"
 
 /*=============================================================================
diff --git a/Engine/Source/Developer/ScreenShotComparisonTools/Public/ScreenShotComparisonTools.h b/Engine/Source/Developer/ScreenShotComparisonTools/Public/ScreenShotComparisonTools.h
--- a/Engine/Source/Developer/ScreenShotComparisonTools/Public/ScreenShotComparisonTools.h
+++ b/Engine/Source/Developer/ScreenShotComparisonTools/Public/ScreenShotComparisonTools.h
@@ -21,12 +21,9 @@ MONOLITHIC_HEADER_BOILERPLATE()
 #include "Misc/IFilter.h"
 #include "Misc/TextFilter.h"
 #include "Misc/FilterCollection.h"
-#include "ScreenShotFilterCollection.h"
 
 /* Interfaces
  *****************************************************************************/
 #include "ImageComparer.h"
-#include "ScreenShotDataItem.h"
-#include "Interfaces/IScreenShotData.h"
 #include "Interfaces/IScreenShotManager.h"
 #include "Interfaces/IScreenShotToolsModule.h"
diff --git a/Engine/Source/Developer/Windows/LiveCoding/Private/External/LC_CommandMap.h b/Engine/Source/Developer/Windows/LiveCoding/Private/External/LC_CommandMap.h
--- a/Engine/Source/Developer/Windows/LiveCoding/Private/External/LC_CommandMap.h
+++ b/Engine/Source/Developer/Windows/LiveCoding/Private/External/LC_CommandMap.h
@@ -4,6 +4,7 @@
 
 #include "CoreTypes.h"
 #include "LC_Commands.h"
+#include "LC_DuplexPipe.h"
 
 
 class DuplexPipe;
diff --git a/Engine/Source/Editor/AnimationModifiers/Public/AnimationBlueprintLibrary.h b/Engine/Source/Editor/AnimationModifiers/Public/AnimationBlueprintLibrary.h
--- a/Engine/Source/Editor/AnimationModifiers/Public/AnimationBlueprintLibrary.h
+++ b/Engine/Source/Editor/AnimationModifiers/Public/AnimationBlueprintLibrary.h
@@ -5,10 +5,12 @@
 #include "Kismet/BlueprintFunctionLibrary.h"
 #include "Animation/AnimEnums.h"
 #include "Animation/AnimCurveTypes.h"
+#include "Animation/AnimMetaData.h"
 
 #include "AnimationBlueprintLibrary.generated.h"
 
 struct FRawAnimSequenceTrack;
+class UAnimCompress;
 class USkeleton;
 
 UENUM()
diff --git a/Engine/Source/Editor/CurveEditor/Public/Tree/CurveEditorTreeTraits.h b/Engine/Source/Editor/CurveEditor/Public/Tree/CurveEditorTreeTraits.h
--- a/Engine/Source/Editor/CurveEditor/Public/Tree/CurveEditorTreeTraits.h
+++ b/Engine/Source/Editor/CurveEditor/Public/Tree/CurveEditorTreeTraits.h
@@ -2,6 +2,8 @@
 
 #pragma once
 
+#include "Widgets/Views/STableRow.h"
+#include "Widgets/Views/STreeView.h"
 #include "Framework/Views/TableViewTypeTraits.h"
 #include "CurveEditorTypes.h"
 
diff --git a/Engine/Source/Editor/GraphEditor/Private/GraphEditorCommon.h b/Engine/Source/Editor/GraphEditor/Private/GraphEditorCommon.h
--- a/Engine/Source/Editor/GraphEditor/Private/GraphEditorCommon.h
+++ b/Engine/Source/Editor/GraphEditor/Private/GraphEditorCommon.h
@@ -21,6 +21,3 @@ MONOLITHIC_HEADER_BOILERPLATE()
 #include "GraphEditorDragDropAction.h"
 #include "DragConnection.h"
 #include "DragNode.h"
-
-#include "AnimGraphDefinitions.h"
-
diff --git a/Engine/Source/Editor/LandscapeEditor/Private/LandscapeEditorDetailCustomizations.h b/Engine/Source/Editor/LandscapeEditor/Private/LandscapeEditorDetailCustomizations.h
--- a/Engine/Source/Editor/LandscapeEditor/Private/LandscapeEditorDetailCustomizations.h
+++ b/Engine/Source/Editor/LandscapeEditor/Private/LandscapeEditorDetailCustomizations.h
@@ -2,10 +2,10 @@
 
 #pragma once
 
-#include "Private/LandscapeEditorDetailCustomization_Base.h"
-#include "Private/LandscapeEditorDetailCustomization_NewLandscape.h"
-#include "Private/LandscapeEditorDetailCustomization_ResizeLandscape.h"
-#include "Private/LandscapeEditorDetailCustomization_CopyPaste.h"
-#include "Private/LandscapeEditorDetailCustomization_MiscTools.h"
-#include "Private/LandscapeEditorDetailCustomization_AlphaBrush.h"
-#include "Private/LandscapeEditorDetailCustomization_TargetLayers.h"
+#include "LandscapeEditorDetailCustomization_Base.h"
+#include "LandscapeEditorDetailCustomization_NewLandscape.h"
+#include "LandscapeEditorDetailCustomization_ResizeLandscape.h"
+#include "LandscapeEditorDetailCustomization_CopyPaste.h"
+#include "LandscapeEditorDetailCustomization_MiscTools.h"
+#include "LandscapeEditorDetailCustomization_AlphaBrush.h"
+#include "LandscapeEditorDetailCustomization_TargetLayers.h"
diff --git a/Engine/Source/Editor/LocalizationDashboard/Private/TargetsTableEntry.h b/Engine/Source/Editor/LocalizationDashboard/Private/TargetsTableEntry.h
--- a/Engine/Source/Editor/LocalizationDashboard/Private/TargetsTableEntry.h
+++ b/Engine/Source/Editor/LocalizationDashboard/Private/TargetsTableEntry.h
@@ -1,6 +1,9 @@
 // Copyright 1998-2019 Epic Games, Inc. All Rights Reserved.
 #pragma once
 
+#include "CoreMinimal.h"
+#include "PropertyHandle.h"
+
 class ULocalizationTarget;
 
 struct FTargetsTableEntry
diff --git a/Engine/Source/Editor/MaterialEditor/Public/MaterialEditingLibrary.h b/Engine/Source/Editor/MaterialEditor/Public/MaterialEditingLibrary.h
--- a/Engine/Source/Editor/MaterialEditor/Public/MaterialEditingLibrary.h
+++ b/Engine/Source/Editor/MaterialEditor/Public/MaterialEditingLibrary.h
@@ -4,13 +4,14 @@
 
 #include "Kismet/BlueprintFunctionLibrary.h"
 #include "Materials/MaterialInterface.h"
+#include "Materials/MaterialExpression.h"
 #include "SceneTypes.h"
 
 #include "MaterialEditingLibrary.generated.h"
 
 class UMaterialFunction;
 class UMaterialInstance;
-
+class UMaterialInstanceConstant;
 
 /** Blueprint library for creating/editing Materials */
 UCLASS()
diff --git a/Engine/Source/Editor/UnrealEd/Classes/Factories/Factory.h b/Engine/Source/Editor/UnrealEd/Classes/Factories/Factory.h
--- a/Engine/Source/Editor/UnrealEd/Classes/Factories/Factory.h
+++ b/Engine/Source/Editor/UnrealEd/Classes/Factories/Factory.h
@@ -9,6 +9,7 @@
 #include "Misc/SecureHash.h"
 #include "Factory.generated.h"
 
+class UAssetImportTask;
 struct FUntypedBulkData;
 
 /**
diff --git a/Engine/Source/Editor/UnrealEd/Private/Factories/ActorFactoryPointLight.cpp b/Engine/Source/Editor/UnrealEd/Private/Factories/ActorFactoryPointLight.cpp
--- a/Engine/Source/Editor/UnrealEd/Private/Factories/ActorFactoryPointLight.cpp
+++ b/Engine/Source/Editor/UnrealEd/Private/Factories/ActorFactoryPointLight.cpp
@@ -7,6 +7,7 @@
 #include "ActorFactories/ActorFactoryPointLight.h"
 #include "GameFramework/Actor.h"
 #include "Components/PointLightComponent.h"
+#include "Components/RectLightComponent.h"
 
 void UActorFactoryPointLight::PostSpawnActor(UObject* Asset, AActor* NewActor)
 {
diff --git a/Engine/Source/Editor/UnrealEd/Private/Factories/ActorFactorySpotLight.cpp b/Engine/Source/Editor/UnrealEd/Private/Factories/ActorFactorySpotLight.cpp
--- a/Engine/Source/Editor/UnrealEd/Private/Factories/ActorFactorySpotLight.cpp
+++ b/Engine/Source/Editor/UnrealEd/Private/Factories/ActorFactorySpotLight.cpp
@@ -7,6 +7,7 @@
 #include "ActorFactories/ActorFactorySpotLight.h"
 #include "GameFramework/Actor.h"
 #include "Components/SpotLightComponent.h"
+#include "Components/RectLightComponent.h"
 
 void UActorFactorySpotLight::PostSpawnActor(UObject* Asset, AActor* NewActor)
 {
diff --git a/Engine/Source/Editor/UnrealEd/Public/Commandlets/AnalyzeReferencedContentStat.h b/Engine/Source/Editor/UnrealEd/Public/Commandlets/AnalyzeReferencedContentStat.h
--- a/Engine/Source/Editor/UnrealEd/Public/Commandlets/AnalyzeReferencedContentStat.h
+++ b/Engine/Source/Editor/UnrealEd/Public/Commandlets/AnalyzeReferencedContentStat.h
@@ -2,6 +2,19 @@
 
 #pragma once
 
+#include "CoreMinimal.h"
+#include "SceneTypes.h"
+
+class USoundWave;
+class USoundCue;
+class AStaticMeshActor;
+class UAnimSequence;
+class UParticleSystem;
+class UTexture;
+class UModel;
+class USkeletalMesh;
+class UStaticMesh;
+
 class FAnalyzeReferencedContentStat
 {
 public:
diff --git a/Engine/Source/Editor/UnrealEd/Public/Subsystems/ImportSubsystem.h b/Engine/Source/Editor/UnrealEd/Public/Subsystems/ImportSubsystem.h
--- a/Engine/Source/Editor/UnrealEd/Public/Subsystems/ImportSubsystem.h
+++ b/Engine/Source/Editor/UnrealEd/Public/Subsystems/ImportSubsystem.h
@@ -11,6 +11,8 @@
 
 #include "ImportSubsystem.generated.h"
 
+class UFactory;
+
 /**
  * Interface for tasks that need delayed execution
  */
diff --git a/Engine/Source/Editor/VREditor/Public/VREditorInteractor.h b/Engine/Source/Editor/VREditor/Public/VREditorInteractor.h
--- a/Engine/Source/Editor/VREditor/Public/VREditorInteractor.h
+++ b/Engine/Source/Editor/VREditor/Public/VREditorInteractor.h
@@ -12,6 +12,7 @@
 class AActor;
 class UWidgetComponent;
 class UStaticMeshSocket;
+class AVREditorTeleporter;
 
 UENUM( BlueprintType )
 enum class EControllerType : uint8
diff --git a/Engine/Source/Editor/VREditor/Teleporter/VREditorTeleporter.h b/Engine/Source/Editor/VREditor/Teleporter/VREditorTeleporter.h
--- a/Engine/Source/Editor/VREditor/Teleporter/VREditorTeleporter.h
+++ b/Engine/Source/Editor/VREditor/Teleporter/VREditorTeleporter.h
@@ -9,6 +9,7 @@
 #include "Engine/EngineBaseTypes.h"
 #include "VREditorTeleporter.generated.h"
 
+class UVREditorInteractor;
 class UViewportInteractor;
 class UVREditorMode;
 
diff --git a/Engine/Source/Programs/PixelStreaming/WebRTCProxy/src/WebRTCProxy.cpp b/Engine/Source/Programs/PixelStreaming/WebRTCProxy/src/WebRTCProxy.cpp
--- a/Engine/Source/Programs/PixelStreaming/WebRTCProxy/src/WebRTCProxy.cpp
+++ b/Engine/Source/Programs/PixelStreaming/WebRTCProxy/src/WebRTCProxy.cpp
@@ -9,6 +9,7 @@
 #include "StringUtils.h"
 #include "ScopeGuard.h"
 #include "CrashDetection.h"
+#include "Modules/ModuleManager.h"
 
 const char* Help =
     "\
diff --git a/Engine/Source/Programs/SymbolDebugger/Private/SymbolDebuggerApp.cpp b/Engine/Source/Programs/SymbolDebugger/Private/SymbolDebuggerApp.cpp
--- a/Engine/Source/Programs/SymbolDebugger/Private/SymbolDebuggerApp.cpp
+++ b/Engine/Source/Programs/SymbolDebugger/Private/SymbolDebuggerApp.cpp
@@ -9,6 +9,8 @@
 #include "ISourceControlModule.h"
 #include "EditorStyleSet.h"
 #include "Interfaces/IEditorStyleModule.h"
+#include "Widgets/SWindow.h"
+#include "Framework/Application/SlateApplication.h"
 
 IMPLEMENT_APPLICATION(SymbolDebugger, "SymbolDebugger");
 
diff --git a/Engine/Source/Runtime/AIModule/Classes/Blueprint/AIBlueprintHelperLibrary.h b/Engine/Source/Runtime/AIModule/Classes/Blueprint/AIBlueprintHelperLibrary.h
--- a/Engine/Source/Runtime/AIModule/Classes/Blueprint/AIBlueprintHelperLibrary.h
+++ b/Engine/Source/Runtime/AIModule/Classes/Blueprint/AIBlueprintHelperLibrary.h
@@ -19,6 +19,7 @@ class UAIAsyncTaskBlueprintProxy;
 class UAnimInstance;
 class UBehaviorTree;
 class UBlackboardComponent;
+class UNavigationPath;
 
 UCLASS(meta=(ScriptName="AIHelperLibrary"))
 class AIMODULE_API UAIBlueprintHelperLibrary : public UBlueprintFunctionLibrary
diff --git a/Engine/Source/Runtime/AIModule/Public/BVTree.h b/Engine/Source/Runtime/AIModule/Public/BVTree.h
--- a/Engine/Source/Runtime/AIModule/Public/BVTree.h
+++ b/Engine/Source/Runtime/AIModule/Public/BVTree.h
@@ -2,6 +2,8 @@
 
 #pragma once
 
+#include "Containers/ContainersFwd.h"
+#include "Math/Box.h"
 
 template<typename InElementType, typename InAllocator = FDefaultAllocator>
 struct TBVTree
diff --git a/Engine/Source/Runtime/AssetRegistry/Public/AssetRegistryHelpers.h b/Engine/Source/Runtime/AssetRegistry/Public/AssetRegistryHelpers.h
--- a/Engine/Source/Runtime/AssetRegistry/Public/AssetRegistryHelpers.h
+++ b/Engine/Source/Runtime/AssetRegistry/Public/AssetRegistryHelpers.h
@@ -5,9 +5,9 @@
 #include "CoreMinimal.h"
 #include "AssetData.h"
 #include "ARFilter.h"
+#include "IAssetRegistry.h"
 #include "AssetRegistryHelpers.generated.h"
 
-
 USTRUCT(BlueprintType)
 struct FTagAndValue
 {
diff --git a/Engine/Source/Runtime/AugmentedReality/Public/ARSkyLight.h b/Engine/Source/Runtime/AugmentedReality/Public/ARSkyLight.h
--- a/Engine/Source/Runtime/AugmentedReality/Public/ARSkyLight.h
+++ b/Engine/Source/Runtime/AugmentedReality/Public/ARSkyLight.h
@@ -7,6 +7,8 @@
 
 #include "ARSkyLight.generated.h"
 
+class UAREnvironmentCaptureProbe;
+
 /** This sky light class forces a refresh of the cube map data when an AR environment probe changes */
 UCLASS()
 class AUGMENTEDREALITY_API AARSkyLight :
diff --git a/Engine/Source/Runtime/CoreUObject/Public/Misc/StringAssetReference.h b/Engine/Source/Runtime/CoreUObject/Public/Misc/StringAssetReference.h
--- a/Engine/Source/Runtime/CoreUObject/Public/Misc/StringAssetReference.h
+++ b/Engine/Source/Runtime/CoreUObject/Public/Misc/StringAssetReference.h
@@ -2,5 +2,7 @@
 
 #pragma once
 
+#include "CoreTypes.h"
+
 COMPILE_WARNING("FStringAssetReference has been renamed to FSoftObjectPath, change to #include \"UObject/SoftObjectPath.h\" and rename references")
 #include "UObject/SoftObjectPath.h"
diff --git a/Engine/Source/Runtime/CoreUObject/Public/Misc/StringClassReference.h b/Engine/Source/Runtime/CoreUObject/Public/Misc/StringClassReference.h
--- a/Engine/Source/Runtime/CoreUObject/Public/Misc/StringClassReference.h
+++ b/Engine/Source/Runtime/CoreUObject/Public/Misc/StringClassReference.h
@@ -2,5 +2,7 @@
 
 #pragma once
 
+#include "CoreTypes.h"
+
 COMPILE_WARNING("FStringClassReference has been renamed to FSoftClassPath and merged headers, change to #include \"UObject/SoftObjectPath.h\" and rename references")
 #include "UObject/SoftObjectPath.h"
\ No newline at end of file
diff --git a/Engine/Source/Runtime/CoreUObject/Public/UObject/AssetPtr.h b/Engine/Source/Runtime/CoreUObject/Public/UObject/AssetPtr.h
--- a/Engine/Source/Runtime/CoreUObject/Public/UObject/AssetPtr.h
+++ b/Engine/Source/Runtime/CoreUObject/Public/UObject/AssetPtr.h
@@ -2,5 +2,7 @@
 
 #pragma once
 
+#include "CoreTypes.h"
+
 COMPILE_WARNING("TAssetPtr has been renamed to TSoftObjectPtr, change to #include \"UObject/SoftObjectPtr.h\" and rename references")
 #include "UObject/SoftObjectPtr.h"
\ No newline at end of file
diff --git a/Engine/Source/Runtime/Online/BuildPatchServices/Private/Core/RingBuffer.h b/Engine/Source/Runtime/Online/BuildPatchServices/Private/Core/RingBuffer.h
--- a/Engine/Source/Runtime/Online/BuildPatchServices/Private/Core/RingBuffer.h
+++ b/Engine/Source/Runtime/Online/BuildPatchServices/Private/Core/RingBuffer.h
@@ -146,7 +146,7 @@ public:
 
 
 private:
-	TRingBuffer(){};
+	TRingBuffer() = delete;
 
 	// Buffer size
 	const uint32 BufferDataSize;
diff --git a/Engine/Source/Runtime/Online/BuildPatchServices/Private/Generation/BuildDataVerifier.h b/Engine/Source/Runtime/Online/BuildPatchServices/Private/Generation/BuildDataVerifier.h
--- a/Engine/Source/Runtime/Online/BuildPatchServices/Private/Generation/BuildDataVerifier.h
+++ b/Engine/Source/Runtime/Online/BuildPatchServices/Private/Generation/BuildDataVerifier.h
@@ -1,6 +1,15 @@
 // Copyright 1998-2019 Epic Games, Inc. All Rights Reserved.
 #pragma once
 
+#include "Core/BlockRange.h"
+#include "Core/BlockStructure.h"
+#include "Common/FileSystem.h"
+#include "Generation/ChunkSearch.h"
+#include "Generation/ChunkWriter.h"
+#include "Generation/DeltaEnumeration.h"
+#include "Generation/DataScanner.h"
+#include "BuildPatchUtil.h"
+
 namespace BuildPatchServices
 {
 	/**
diff --git a/Engine/Source/Runtime/Online/BuildPatchServices/Private/Tests/Mock/FileConstructorStat.mock.h b/Engine/Source/Runtime/Online/BuildPatchServices/Private/Tests/Mock/FileConstructorStat.mock.h
--- a/Engine/Source/Runtime/Online/BuildPatchServices/Private/Tests/Mock/FileConstructorStat.mock.h
+++ b/Engine/Source/Runtime/Online/BuildPatchServices/Private/Tests/Mock/FileConstructorStat.mock.h
@@ -2,6 +2,7 @@
 #pragma once
 
 #include "BuildPatchFileConstructor.h"
+#include "Common/StatsCollector.h"
 
 #if WITH_DEV_AUTOMATION_TESTS
 
diff --git a/Engine/Source/Runtime/UMG/Public/Blueprint/IUserListEntry.h b/Engine/Source/Runtime/UMG/Public/Blueprint/IUserListEntry.h
--- a/Engine/Source/Runtime/UMG/Public/Blueprint/IUserListEntry.h
+++ b/Engine/Source/Runtime/UMG/Public/Blueprint/IUserListEntry.h
@@ -7,6 +7,7 @@
 #include "IUserListEntry.generated.h"
 
 class UUserWidget;
+class UListViewBase;
 class IObjectTableRow;
 
 /**
