Fix for VR controllers disappearing in landscape mode

diff --git a/Engine/Source/Editor/LandscapeEditor/Private/LandscapeEdMode.cpp b/Engine/Source/Editor/LandscapeEditor/Private/LandscapeEdMode.cpp
--- a/Engine/Source/Editor/LandscapeEditor/Private/LandscapeEdMode.cpp
+++ b/Engine/Source/Editor/LandscapeEditor/Private/LandscapeEdMode.cpp
@@ -3022,11 +3022,11 @@ void FEdModeLandscape::ForceRealTimeViewports(const bool bEnable, const bool bSt
 					UVREditorMode* VREditorMode = Cast<UVREditorMode>( GEditor->GetEditorWorldExtensionsManager()->GetEditorWorldExtensions( GetWorld() )->FindExtension( UVREditorMode::StaticClass() ) );
 					if( VREditorMode != nullptr && VREditorMode->IsActive())
 					{
-						Viewport.SetGameView(true);
+						Viewport.SetVREditView(true);
 					} 
 					else
 					{
-						Viewport.SetGameView(false);
+						Viewport.SetVREditView(false);
 					}
 				}
 				else
