Landscape Flatten Tool - terrace flatten mode added

diff --git a/Engine/Source/Editor/LandscapeEditor/Private/LandscapeEdModePaintTools.cpp b/Engine/Source/Editor/LandscapeEditor/Private/LandscapeEdModePaintTools.cpp
--- a/Engine/Source/Editor/LandscapeEditor/Private/LandscapeEdModePaintTools.cpp
+++ b/Engine/Source/Editor/LandscapeEditor/Private/LandscapeEdModePaintTools.cpp
@@ -816,6 +816,36 @@ public:
 						int32 Delta = DataScanline[X] - FlattenHeight;
 						switch (UISettings->FlattenMode)
 						{
+						case ELandscapeToolFlattenMode::Terrace:
+							{
+								const FTransform& LocalToWorld = Target.LandscapeInfo->GetLandscapeProxy()->ActorToWorld();
+								float ScaleZ = LocalToWorld.GetScale3D().Z;
+								float TranslateZ = LocalToWorld.GetTranslation().Z;
+								float TerraceInterval = UISettings->TerraceInterval;
+								float Smoothness = UISettings->TerraceSmooth;								
+								float WorldHeight = LandscapeDataAccess::GetLocalHeight(DataScanline[X]);
+								
+								//move into world space
+								WorldHeight = (WorldHeight * ScaleZ) + TranslateZ;
+								float CurrentHeight = WorldHeight;
+
+								//smoothing part
+								float CurrentLevel = WorldHeight / TerraceInterval;								
+								Smoothness = 1.0f / FMath::Max(Smoothness, 0.0001f);
+								float CurrentPhase = FMath::Frac(CurrentLevel);
+								float Halfmask = FMath::Clamp(FMath::CeilToFloat(CurrentPhase - 0.5f), 0.0f, 1.0f);
+								CurrentLevel = FMath::FloorToFloat(WorldHeight / TerraceInterval);
+								float SCurve = FMath::Lerp(CurrentPhase, (1.0f - CurrentPhase), Halfmask) * 2.0f;
+								SCurve = FMath::Pow(SCurve, Smoothness) * 0.5f;
+								SCurve = FMath::Lerp(SCurve, 1.0f - SCurve, Halfmask) * TerraceInterval;
+								WorldHeight = (CurrentLevel * TerraceInterval)  + SCurve;
+								//end of smoothing part
+
+								float FinalHeight = FMath::Lerp(CurrentHeight, WorldHeight , Strength);
+								FinalHeight = (FinalHeight - TranslateZ) / ScaleZ;
+								DataScanline[X] = LandscapeDataAccess::GetTexHeight(FinalHeight);	
+							}
+							break;
 						case ELandscapeToolFlattenMode::Raise:
 							if (Delta < 0)
 							{
diff --git a/Engine/Source/Editor/LandscapeEditor/Public/LandscapeEditorObject.h b/Engine/Source/Editor/LandscapeEditor/Public/LandscapeEditorObject.h
--- a/Engine/Source/Editor/LandscapeEditor/Public/LandscapeEditorObject.h
+++ b/Engine/Source/Editor/LandscapeEditor/Public/LandscapeEditorObject.h
@@ -28,6 +28,9 @@ enum class ELandscapeToolFlattenMode : int8
 
 	/** Flatten may only lower values, values below the clicked point will be left unchanged */
 	Lower = 2,
+
+	/** Flatten to specific terrace height intervals */
+	Terrace = 3
 };
 
 UENUM()
@@ -262,7 +265,7 @@ class ULandscapeEditorObject : public UObject
 
 	// Flatten Tool:
 
-	// Whether to flatten by lowering, raising, or both
+	// Whether to flatten by lowering, raising, both or terracing
 	UPROPERTY(Category="Tool Settings", EditAnywhere, NonTransactional, meta=(ShowForTools="Flatten"))
 	ELandscapeToolFlattenMode FlattenMode;
 
@@ -286,6 +289,14 @@ class ULandscapeEditorObject : public UObject
 	UPROPERTY(Category = "Tool Settings", EditAnywhere, NonTransactional, AdvancedDisplay, meta = (DisplayName = "Show Preview Grid", ShowForTools = "Flatten", ShowForTargetTypes = "Heightmap", EditCondition = "bUseFlattenTarget", HideEditConditionToggle, UIMin = "-32768", UIMax = "32768"))
 	bool bShowFlattenTargetPreview;
 
+	// Height of the terrace intervals in unreal units, for the terrace flatten mode 
+	UPROPERTY(Category = "Tool Settings", EditAnywhere, NonTransactional, AdvancedDisplay, meta = (DisplayName = "Terrace Interval", ShowForTools = "Flatten", ShowForTargetTypes = "Heightmap", UIMin = "1", UIMax = "32768"))
+	float TerraceInterval;
+
+	// Smoothing value for terrace flatten mode
+	UPROPERTY(Category = "Tool Settings", EditAnywhere, NonTransactional, AdvancedDisplay, meta = (DisplayName = "Terrace Smoothing", ShowForTools = "Flatten", ShowForTargetTypes = "Heightmap", UIMin = "0.0001", UIMax = "1.0"))
+	float TerraceSmooth;
+
 	// Whether the Eye Dropper mode is activated
 	UPROPERTY(NonTransactional, Transient)
 	bool bFlattenEyeDropperModeActivated;
