[NO NON-PII COMMIT MESSAGE]

diff --git a/Engine/Source/Editor/LandscapeEditor/Private/LandscapeEdModePaintTools.cpp b/Engine/Source/Editor/LandscapeEditor/Private/LandscapeEdModePaintTools.cpp
--- a/Engine/Source/Editor/LandscapeEditor/Private/LandscapeEdModePaintTools.cpp
+++ b/Engine/Source/Editor/LandscapeEditor/Private/LandscapeEdModePaintTools.cpp
@@ -818,7 +818,7 @@ public:
 						{
 						case ELandscapeToolFlattenMode::Terrace:
 							{
-								const FTransform& LocalToWorld = Target.LandscapeInfo->GetLandscapeProxy()->ActorToWorld();
+								const FTransform& LocalToWorld = this->Target.LandscapeInfo->GetLandscapeProxy()->ActorToWorld();
 								float ScaleZ = LocalToWorld.GetScale3D().Z;
 								float TranslateZ = LocalToWorld.GetTranslation().Z;
 								float TerraceInterval = UISettings->TerraceInterval;
