[NO NON-PII COMMIT MESSAGE]

diff --git a/Engine/Source/Editor/LandscapeEditor/Private/LandscapeEdModeComponentTools.cpp b/Engine/Source/Editor/LandscapeEditor/Private/LandscapeEdModeComponentTools.cpp
--- a/Engine/Source/Editor/LandscapeEditor/Private/LandscapeEdModeComponentTools.cpp
+++ b/Engine/Source/Editor/LandscapeEditor/Private/LandscapeEdModeComponentTools.cpp
@@ -217,7 +217,7 @@ public:
 
 			ALandscape* Landscape = LandscapeInfo->LandscapeActor.Get();
 
-			if (Landscape->HasProceduralContent && !GetMutableDefault<UEditorExperimentalSettings>()->bProceduralLandscape)
+			if (Landscape != nullptr && Landscape->HasProceduralContent && !GetMutableDefault<UEditorExperimentalSettings>()->bProceduralLandscape)
 			{
 				FMessageLog("MapCheck").Warning()->AddToken(FTextToken::Create(LOCTEXT("LandscapeProcedural_ChangingDataWithoutSettings", "This map contains landscape procedural content, modifying the landscape data will result in data loss when the map is reopened with Landscape Procedural settings on. Please enable Landscape Procedural settings before modifying the data.")));
 				FMessageLog("MapCheck").Open(EMessageSeverity::Warning);
@@ -304,7 +304,7 @@ public:
 
 			ALandscape* Landscape = LandscapeInfo->LandscapeActor.Get();
 
-			if (Landscape->HasProceduralContent && !GetMutableDefault<UEditorExperimentalSettings>()->bProceduralLandscape)
+			if (Landscape != nullptr && Landscape->HasProceduralContent && !GetMutableDefault<UEditorExperimentalSettings>()->bProceduralLandscape)
 			{
 				FMessageLog("MapCheck").Warning()->AddToken(FTextToken::Create(LOCTEXT("LandscapeProcedural_ChangingDataWithoutSettings", "This map contains landscape procedural content, modifying the landscape data will result in data loss when the map is reopened with Landscape Procedural settings on. Please enable Landscape Procedural settings before modifying the data.")));
 				FMessageLog("MapCheck").Open(EMessageSeverity::Warning);
@@ -1736,7 +1736,7 @@ public:
 
 			ALandscape* Landscape = LandscapeInfo->LandscapeActor.Get();
 
-			if (Landscape->HasProceduralContent && !GetMutableDefault<UEditorExperimentalSettings>()->bProceduralLandscape)
+			if (Landscape != nullptr && Landscape->HasProceduralContent && !GetMutableDefault<UEditorExperimentalSettings>()->bProceduralLandscape)
 			{
 				FMessageLog("MapCheck").Warning()->AddToken(FTextToken::Create(LOCTEXT("LandscapeProcedural_ChangingDataWithoutSettings", "This map contains landscape procedural content, modifying the landscape data will result in data loss when the map is reopened with Landscape Procedural settings on. Please enable Landscape Procedural settings before modifying the data.")));
 				FMessageLog("MapCheck").Open(EMessageSeverity::Warning);
diff --git a/Engine/Source/Editor/LandscapeEditor/Private/LandscapeEdModeErosionTools.cpp b/Engine/Source/Editor/LandscapeEditor/Private/LandscapeEdModeErosionTools.cpp
--- a/Engine/Source/Editor/LandscapeEditor/Private/LandscapeEdModeErosionTools.cpp
+++ b/Engine/Source/Editor/LandscapeEditor/Private/LandscapeEdModeErosionTools.cpp
@@ -259,7 +259,7 @@ public:
 
 		ALandscape* Landscape = LandscapeInfo->LandscapeActor.Get();
 
-		if (Landscape->HasProceduralContent && !GetMutableDefault<UEditorExperimentalSettings>()->bProceduralLandscape)
+		if (Landscape != nullptr && Landscape->HasProceduralContent && !GetMutableDefault<UEditorExperimentalSettings>()->bProceduralLandscape)
 		{
 			FMessageLog("MapCheck").Warning()->AddToken(FTextToken::Create(LOCTEXT("LandscapeProcedural_ChangingDataWithoutSettings", "This map contains landscape procedural content, modifying the landscape data will result in data loss when the map is reopened with Landscape Procedural settings on. Please enable Landscape Procedural settings before modifying the data.")));
 			FMessageLog("MapCheck").Open(EMessageSeverity::Warning);
@@ -490,7 +490,7 @@ public:
 
 		ALandscape* Landscape = LandscapeInfo->LandscapeActor.Get();
 
-		if (Landscape->HasProceduralContent && !GetMutableDefault<UEditorExperimentalSettings>()->bProceduralLandscape)
+		if (Landscape != nullptr && Landscape->HasProceduralContent && !GetMutableDefault<UEditorExperimentalSettings>()->bProceduralLandscape)
 		{
 			FMessageLog("MapCheck").Warning()->AddToken(FTextToken::Create(LOCTEXT("LandscapeProcedural_ChangingDataWithoutSettings", "This map contains landscape procedural content, modifying the landscape data will result in data loss when the map is reopened with Landscape Procedural settings on. Please enable Landscape Procedural settings before modifying the data.")));
 			FMessageLog("MapCheck").Open(EMessageSeverity::Warning);
diff --git a/Engine/Source/Editor/LandscapeEditor/Private/LandscapeEdModePaintTools.cpp b/Engine/Source/Editor/LandscapeEditor/Private/LandscapeEdModePaintTools.cpp
--- a/Engine/Source/Editor/LandscapeEditor/Private/LandscapeEdModePaintTools.cpp
+++ b/Engine/Source/Editor/LandscapeEditor/Private/LandscapeEdModePaintTools.cpp
@@ -313,7 +313,7 @@ public:
 
 		ALandscape* Landscape = LandscapeInfo->LandscapeActor.Get();
 
-		if (Landscape->HasProceduralContent && !GetMutableDefault<UEditorExperimentalSettings>()->bProceduralLandscape)
+		if (Landscape != nullptr && Landscape->HasProceduralContent && !GetMutableDefault<UEditorExperimentalSettings>()->bProceduralLandscape)
 		{
 			FMessageLog("MapCheck").Warning()->AddToken(FTextToken::Create(LOCTEXT("LandscapeProcedural_ChangingDataWithoutSettings", "This map contains landscape procedural content, modifying the landscape data will result in data loss when the map is reopened with Landscape Procedural settings on. Please enable Landscape Procedural settings before modifying the data.")));
 			FMessageLog("MapCheck").Open(EMessageSeverity::Warning);
@@ -556,7 +556,7 @@ public:
 
 		ALandscape* Landscape = LandscapeInfo->LandscapeActor.Get();
 
-		if (Landscape->HasProceduralContent && !GetMutableDefault<UEditorExperimentalSettings>()->bProceduralLandscape)
+		if (Landscape != nullptr && Landscape->HasProceduralContent && !GetMutableDefault<UEditorExperimentalSettings>()->bProceduralLandscape)
 		{
 			FMessageLog("MapCheck").Warning()->AddToken(FTextToken::Create(LOCTEXT("LandscapeProcedural_ChangingDataWithoutSettings", "This map contains landscape procedural content, modifying the landscape data will result in data loss when the map is reopened with Landscape Procedural settings on. Please enable Landscape Procedural settings before modifying the data.")));
 			FMessageLog("MapCheck").Open(EMessageSeverity::Warning);
@@ -690,7 +690,7 @@ public:
 
 		ALandscape* Landscape = this->LandscapeInfo->LandscapeActor.Get();
 
-		if (Landscape->HasProceduralContent && !GetMutableDefault<UEditorExperimentalSettings>()->bProceduralLandscape)
+		if (Landscape != nullptr && Landscape->HasProceduralContent && !GetMutableDefault<UEditorExperimentalSettings>()->bProceduralLandscape)
 		{
 			FMessageLog("MapCheck").Warning()->AddToken(FTextToken::Create(LOCTEXT("LandscapeProcedural_ChangingDataWithoutSettings", "This map contains landscape procedural content, modifying the landscape data will result in data loss when the map is reopened with Landscape Procedural settings on. Please enable Landscape Procedural settings before modifying the data.")));
 			FMessageLog("MapCheck").Open(EMessageSeverity::Warning);
@@ -910,7 +910,7 @@ public:
 
 		ALandscape* Landscape = this->LandscapeInfo->LandscapeActor.Get();
 
-		if (Landscape->HasProceduralContent && !GetMutableDefault<UEditorExperimentalSettings>()->bProceduralLandscape)
+		if (Landscape != nullptr && Landscape->HasProceduralContent && !GetMutableDefault<UEditorExperimentalSettings>()->bProceduralLandscape)
 		{
 			FMessageLog("MapCheck").Warning()->AddToken(FTextToken::Create(LOCTEXT("LandscapeProcedural_ChangingDataWithoutSettings", "This map contains landscape procedural content, modifying the landscape data will result in data loss when the map is reopened with Landscape Procedural settings on. Please enable Landscape Procedural settings before modifying the data.")));
 			FMessageLog("MapCheck").Open(EMessageSeverity::Warning);
@@ -1140,7 +1140,7 @@ public:
 
 		ALandscape* Landscape = this->LandscapeInfo->LandscapeActor.Get();
 
-		if (Landscape->HasProceduralContent && !GetMutableDefault<UEditorExperimentalSettings>()->bProceduralLandscape)
+		if (Landscape != nullptr && Landscape->HasProceduralContent && !GetMutableDefault<UEditorExperimentalSettings>()->bProceduralLandscape)
 		{
 			FMessageLog("MapCheck").Warning()->AddToken(FTextToken::Create(LOCTEXT("LandscapeProcedural_ChangingDataWithoutSettings", "This map contains landscape procedural content, modifying the landscape data will result in data loss when the map is reopened with Landscape Procedural settings on. Please enable Landscape Procedural settings before modifying the data.")));
 			FMessageLog("MapCheck").Open(EMessageSeverity::Warning);
diff --git a/Engine/Source/Editor/LandscapeEditor/Private/LandscapeEdModeXYOffset.cpp b/Engine/Source/Editor/LandscapeEditor/Private/LandscapeEdModeXYOffset.cpp
--- a/Engine/Source/Editor/LandscapeEditor/Private/LandscapeEdModeXYOffset.cpp
+++ b/Engine/Source/Editor/LandscapeEditor/Private/LandscapeEdModeXYOffset.cpp
@@ -543,7 +543,7 @@ public:
 
 		ALandscape* Landscape = LandscapeInfo->LandscapeActor.Get();
 
-		if (Landscape->HasProceduralContent && !GetMutableDefault<UEditorExperimentalSettings>()->bProceduralLandscape)
+		if (Landscape != nullptr && Landscape->HasProceduralContent && !GetMutableDefault<UEditorExperimentalSettings>()->bProceduralLandscape)
 		{
 			FMessageLog("MapCheck").Warning()->AddToken(FTextToken::Create(LOCTEXT("LandscapeProcedural_ChangingDataWithoutSettings", "This map contains landscape procedural content, modifying the landscape data will result in data loss when the map is reopened with Landscape Procedural settings on. Please enable Landscape Procedural settings before modifying the data.")));
 			FMessageLog("MapCheck").Open(EMessageSeverity::Warning);
