Tried to share the FMaterialUpdateContext & FRecreateRenderState as much as possible, to prevent many recreation of the graphic scene. This was done to improve the performance of Change Component Settings, and would also have some impact on other things like DeleteComponent, moving component to another level, creating a new landscape, etc.

The change also include a merge from DevMain to improve the perf of UpdateMaterialInstance

diff --git a/Engine/Source/Editor/LandscapeEditor/Private/LandscapeEdMode.cpp b/Engine/Source/Editor/LandscapeEditor/Private/LandscapeEdMode.cpp
--- a/Engine/Source/Editor/LandscapeEditor/Private/LandscapeEdMode.cpp
+++ b/Engine/Source/Editor/LandscapeEditor/Private/LandscapeEdMode.cpp
@@ -55,6 +55,7 @@
 #include "LandscapeBPCustomBrush.h"
 #include "Engine/TextureRenderTarget2D.h"
 #include "Settings/EditorExperimentalSettings.h"
+#include "ComponentRecreateRenderStateContext.h"
 
 #define LOCTEXT_NAMESPACE "Landscape"
 
@@ -96,8 +97,9 @@ struct HNewLandscapeGrabHandleProxy : public HHitProxy
 
 IMPLEMENT_HIT_PROXY(HNewLandscapeGrabHandleProxy, HHitProxy)
 
+ENGINE_API extern bool GDisableAutomaticTextureMaterialUpdateDependencies;
 
-void ALandscape::SplitHeightmap(ULandscapeComponent* Comp, bool bMoveToCurrentLevel /*= false*/)
+void ALandscape::SplitHeightmap(ULandscapeComponent* Comp, bool bMoveToCurrentLevel, FMaterialUpdateContext* InOutUpdateContext, TArray<FComponentRecreateRenderStateContext>* InOutRecreateRenderStateContext, bool InReregisterComponent)
 {
 	ULandscapeInfo* Info = Comp->GetLandscapeInfo();
 	ALandscape* Landscape = Info->LandscapeActor.Get();
@@ -247,19 +249,74 @@ void ALandscape::SplitHeightmap(ULandscapeComponent* Comp, bool bMoveToCurrentLe
 #endif
 
 		Comp->SetHeightmap(HeightmapTexture);
-		Comp->UpdateMaterialInstances();
 
 		for (int32 i = 0; i < HeightmapTextureMipData.Num(); i++)
 		{
 			HeightmapTexture->Source.UnlockMip(i);
-		}
+		}		
+
 		LandscapeEdit.SetHeightData(Comp->GetSectionBase().X, Comp->GetSectionBase().Y, Comp->GetSectionBase().X + Comp->ComponentSizeQuads, Comp->GetSectionBase().Y + Comp->ComponentSizeQuads, (uint16*)HeightData.GetData(), 0, false, (uint16*)NormalData.GetData());
 	}
 
+	// End material update
+	if (InOutUpdateContext != nullptr && InOutRecreateRenderStateContext != nullptr)
+	{
+		Comp->UpdateMaterialInstances(*InOutUpdateContext, *InOutRecreateRenderStateContext);
+	}
+	else
+	{
+		Comp->UpdateMaterialInstances();
+	}
+
 	// End of LandscapeEdit interface
-	HeightmapTexture->PostEditChange();
-	// Reregister
-	FComponentReregisterContext ReregisterContext(Comp);
+
+	// We disable automatic material update context, to manage it manually if we have a custom update context specified
+	GDisableAutomaticTextureMaterialUpdateDependencies = InOutUpdateContext != nullptr;
+
+	HeightmapTexture->PostEditChange();	
+
+	if (InOutUpdateContext != nullptr)
+	{
+		// Build a list of all unique materials the landscape uses
+		TArray<UMaterialInterface*> LandscapeMaterials;
+
+		int8 MaxLOD = FMath::CeilLogTwo(Comp->SubsectionSizeQuads + 1) - 1;
+
+		for (int8 LODIndex = 0; LODIndex < MaxLOD; ++LODIndex)
+		{
+			UMaterialInterface* Material = Comp->GetLandscapeMaterial(LODIndex);
+			LandscapeMaterials.AddUnique(Material);
+		}
+
+		TSet<UMaterial*> BaseMaterialsThatUseThisTexture;
+
+		for (UMaterialInterface* MaterialInterface : LandscapeMaterials)
+		{
+			if (DoesMaterialUseTexture(MaterialInterface, HeightmapTexture))
+			{
+				UMaterial* Material = MaterialInterface->GetMaterial();
+				bool MaterialAlreadyCompute = false;
+				BaseMaterialsThatUseThisTexture.Add(Material, &MaterialAlreadyCompute);
+
+				if (!MaterialAlreadyCompute)
+				{
+					if (Material->IsTextureForceRecompileCacheRessource(HeightmapTexture))
+					{
+						InOutUpdateContext->AddMaterial(Material);
+						Material->UpdateMaterialShaderCacheAndTextureReferences();
+					}
+				}
+			}
+		}
+	}
+	
+	GDisableAutomaticTextureMaterialUpdateDependencies = false;
+
+		// Reregister
+	if (InReregisterComponent)
+	{
+		FComponentReregisterContext ReregisterContext(Comp);
+	}
 }
 
 
@@ -3531,12 +3588,23 @@ void FEdModeLandscape::DeleteLandscapeComponents(ULandscapeInfo* LandscapeInfo,
 		}
 	}
 
-	// Changing Heightmap format for selected components
-	for (ULandscapeComponent* Component : HeightmapUpdateComponents)
+	TArray<FComponentRecreateRenderStateContext> RecreateRenderStateContexts;
+
 	{
-		ALandscape::SplitHeightmap(Component, false);
+		TArray<UActorComponent*> ComponentsToReregister(HeightmapUpdateComponents.Array());
+		FMaterialUpdateContext MaterialUpdateContext(FMaterialUpdateContext::EOptions::Default & ~FMaterialUpdateContext::EOptions::RecreateRenderStates);
+
+		// Changing Heightmap format for selected components
+		for (ULandscapeComponent* Component : HeightmapUpdateComponents)
+		{
+			ALandscape::SplitHeightmap(Component, false, &MaterialUpdateContext, &RecreateRenderStateContexts, false);
+		}
+
+		FMultiComponentReregisterContext RegisterContext(ComponentsToReregister);
 	}
 
+	RecreateRenderStateContexts.Empty();
+
 	// Remove attached foliage
 	for (ULandscapeComponent* Component : ComponentsToDelete)
 	{
@@ -3547,6 +3615,8 @@ void FEdModeLandscape::DeleteLandscapeComponents(ULandscapeInfo* LandscapeInfo,
 		}
 	}
 
+	TArray<UActorComponent*> NeighborsComponentToReregister;
+
 	// Check which ones are need for height map change
 	for (ULandscapeComponent* Component : ComponentsToDelete)
 	{
@@ -3572,8 +3642,7 @@ void FEdModeLandscape::DeleteLandscapeComponents(ULandscapeInfo* LandscapeInfo,
 				NeighborComp->Modify();
 				NeighborComp->InvalidateLightingCache();
 
-				// is this really needed? It can happen multiple times per component!
-				FComponentReregisterContext ReregisterContext(NeighborComp);
+				NeighborsComponentToReregister.AddUnique(NeighborComp);
 			}
 		}
 
@@ -3620,6 +3689,10 @@ void FEdModeLandscape::DeleteLandscapeComponents(ULandscapeInfo* LandscapeInfo,
 		Component->DestroyComponent();
 	}
 
+	{
+		FMultiComponentReregisterContext RegisterContext(NeighborsComponentToReregister);
+	}
+
 	// Remove Selection
 	LandscapeInfo->ClearSelectedRegion(true);
 	//EdMode->SetMaskEnable(Landscape->SelectedRegion.Num());
diff --git a/Engine/Source/Runtime/Engine/Classes/Materials/Material.h b/Engine/Source/Runtime/Engine/Classes/Materials/Material.h
--- a/Engine/Source/Runtime/Engine/Classes/Materials/Material.h
+++ b/Engine/Source/Runtime/Engine/Classes/Materials/Material.h
@@ -1457,13 +1457,13 @@ private:
 	void GetForceRecompileTextureIdsHash(FSHAHash &TextureReferencesHash);
 
 public:
-	bool IsTextureForceRecompileCacheRessource(UTexture *Texture);
+	ENGINE_API bool IsTextureForceRecompileCacheRessource(UTexture *Texture);
 
 #if WITH_EDITOR
 	/* Recompute the ddc cache key and reload the material in case the key is not the same.
 	 * It will also make sure lightmass texture reference are up to date
 	 */
-	void UpdateMaterialShaderCacheAndTextureReferences();
+	ENGINE_API void UpdateMaterialShaderCacheAndTextureReferences();
 #endif
 
 	/**
diff --git a/Engine/Source/Runtime/Engine/Private/Texture.cpp b/Engine/Source/Runtime/Engine/Private/Texture.cpp
--- a/Engine/Source/Runtime/Engine/Private/Texture.cpp
+++ b/Engine/Source/Runtime/Engine/Private/Texture.cpp
@@ -48,6 +48,10 @@ FName FTextureResource::TextureGroupStatFNames[TEXTUREGROUP_MAX] =
 	};
 #endif
 
+// This is used to prevent the PostEditChange to automatically update the material depedencies & material context, in some case we want to manually control this
+// to be more efficient.
+ENGINE_API bool GDisableAutomaticTextureMaterialUpdateDependencies = false;
+
 UTexture::FOnTextureSaved UTexture::PreSaveEvent;
 
 UTexture::UTexture(const FObjectInitializer& ObjectInitializer)
@@ -192,7 +196,7 @@ void UTexture::PostEditChangeProperty(FPropertyChangedEvent& PropertyChangedEven
 			SRGB = false;
 		}
 	}
-	else
+	else if (!GDisableAutomaticTextureMaterialUpdateDependencies)
 	{
 		FMaterialUpdateContext UpdateContext;
 		// Update any material that uses this texture and must force a recompile of cache ressource
@@ -202,7 +206,7 @@ void UTexture::PostEditChangeProperty(FPropertyChangedEvent& PropertyChangedEven
 			UMaterialInterface* MaterialInterface = *It;
 			if (DoesMaterialUseTexture(MaterialInterface, this))
 			{
-				UMaterial *Material = MaterialInterface->GetMaterial();
+				UMaterial* Material = MaterialInterface->GetMaterial();
 				bool MaterialAlreadyCompute = false;
 				BaseMaterialsThatUseThisTexture.Add(Material, &MaterialAlreadyCompute);
 				if (!MaterialAlreadyCompute)
diff --git a/Engine/Source/Runtime/Landscape/Classes/Landscape.h b/Engine/Source/Runtime/Landscape/Classes/Landscape.h
--- a/Engine/Source/Runtime/Landscape/Classes/Landscape.h
+++ b/Engine/Source/Runtime/Landscape/Classes/Landscape.h
@@ -173,7 +173,7 @@ public:
 	LANDSCAPE_API static void CalcComponentIndicesNoOverlap(const int32 X1, const int32 Y1, const int32 X2, const int32 Y2, const int32 ComponentSizeQuads,
 		int32& ComponentIndexX1, int32& ComponentIndexY1, int32& ComponentIndexX2, int32& ComponentIndexY2);
 
-	static void SplitHeightmap(ULandscapeComponent* Comp, bool bMoveToCurrentLevel = false);
+	static void SplitHeightmap(ULandscapeComponent* Comp, bool bMoveToCurrentLevel = false, FMaterialUpdateContext* InOutUpdateContext = nullptr, TArray<FComponentRecreateRenderStateContext>* InOutRecreateRenderStateContext = nullptr, bool InReregisterComponent = true);
 	
 	//~ Begin UObject Interface.
 	virtual void PreSave(const class ITargetPlatform* TargetPlatform) override;
diff --git a/Engine/Source/Runtime/Landscape/Classes/LandscapeComponent.h b/Engine/Source/Runtime/Landscape/Classes/LandscapeComponent.h
--- a/Engine/Source/Runtime/Landscape/Classes/LandscapeComponent.h
+++ b/Engine/Source/Runtime/Landscape/Classes/LandscapeComponent.h
@@ -640,6 +640,7 @@ public:
 	 * Creates the MaterialInstance if it doesn't exist.
 	 */
 	LANDSCAPE_API void UpdateMaterialInstances();
+	LANDSCAPE_API void UpdateMaterialInstances(FMaterialUpdateContext& InOutMaterialContext, TArray<FComponentRecreateRenderStateContext>& InOutRecreateRenderStateContext);
 
 	// Internal implementation of UpdateMaterialInstances, not safe to call directly
 	void UpdateMaterialInstances_Internal(FMaterialUpdateContext& Context);
diff --git a/Engine/Source/Runtime/Landscape/Classes/LandscapeProxy.h b/Engine/Source/Runtime/Landscape/Classes/LandscapeProxy.h
--- a/Engine/Source/Runtime/Landscape/Classes/LandscapeProxy.h
+++ b/Engine/Source/Runtime/Landscape/Classes/LandscapeProxy.h
@@ -891,6 +891,7 @@ public:
 
 	/** Update the material instances for all the landscape components */
 	LANDSCAPE_API void UpdateAllComponentMaterialInstances();
+	LANDSCAPE_API void UpdateAllComponentMaterialInstances(FMaterialUpdateContext& InOutMaterialContext, TArray<FComponentRecreateRenderStateContext>& InOutRecreateRenderStateContext);
 
 	/** Create a thumbnail material for a given layer */
 	LANDSCAPE_API static ULandscapeMaterialInstanceConstant* GetLayerThumbnailMIC(UMaterialInterface* LandscapeMaterial, FName LayerName, UTexture2D* ThumbnailWeightmap, UTexture2D* ThumbnailHeightmap, ALandscapeProxy* Proxy);
diff --git a/Engine/Source/Runtime/Landscape/Private/LandscapeEdit.cpp b/Engine/Source/Runtime/Landscape/Private/LandscapeEdit.cpp
--- a/Engine/Source/Runtime/Landscape/Private/LandscapeEdit.cpp
+++ b/Engine/Source/Runtime/Landscape/Private/LandscapeEdit.cpp
@@ -161,8 +161,13 @@ ULandscapeMaterialInstanceConstant* ALandscapeProxy::GetLayerThumbnailMIC(UMater
 			LayerParameter.WeightmapIndex = INDEX_NONE;
 		}
 	}
-	MaterialInstance->UpdateStaticPermutation(StaticParameters);
 
+	// Don't recreate the render state of everything, only update the materials context
+	{
+		FMaterialUpdateContext MaterialUpdateContext(FMaterialUpdateContext::EOptions::Default & ~FMaterialUpdateContext::EOptions::RecreateRenderStates);
+		MaterialInstance->UpdateStaticPermutation(StaticParameters, &MaterialUpdateContext);
+	}
+	
 	FLinearColor Mask(1.0f, 0.0f, 0.0f, 0.0f);
 	MaterialInstance->SetVectorParameterValueEditorOnly(FName(*FString::Printf(TEXT("LayerMask_%s"), *LayerName.ToString())), Mask);
 	MaterialInstance->SetTextureParameterValueEditorOnly(FName(TEXT("Weightmap0")), ThumbnailWeightmap);
@@ -436,6 +441,21 @@ void ULandscapeComponent::UpdateMaterialInstances()
 	RecreateRenderStateContext.Reset();
 }
 
+void ULandscapeComponent::UpdateMaterialInstances(FMaterialUpdateContext& InOutMaterialContext, TArray<FComponentRecreateRenderStateContext>& InOutRecreateRenderStateContext)
+{
+	InOutRecreateRenderStateContext.Add(this);
+	UpdateMaterialInstances_Internal(InOutMaterialContext);
+}
+
+void ALandscapeProxy::UpdateAllComponentMaterialInstances(FMaterialUpdateContext& InOutMaterialContext, TArray<FComponentRecreateRenderStateContext>& InOutRecreateRenderStateContext)
+{
+	for (ULandscapeComponent* Component : LandscapeComponents)
+	{
+		Component->UpdateMaterialInstances(InOutMaterialContext, InOutRecreateRenderStateContext);
+	}
+
+}
+
 void ALandscapeProxy::UpdateAllComponentMaterialInstances()
 {
 	// we're not having the material update context recreate render states because we will manually do it for only our components
@@ -2140,6 +2160,8 @@ ULandscapeLayerInfoObject* ALandscapeProxy::CreateLayerInfo(const TCHAR* LayerNa
 }
 
 #define HEIGHTDATA(X,Y) (HeightData[ FMath::Clamp<int32>(Y,0,VertsY) * VertsX + FMath::Clamp<int32>(X,0,VertsX) ])
+ENGINE_API extern bool GDisableAutomaticTextureMaterialUpdateDependencies;
+
 LANDSCAPE_API void ALandscapeProxy::Import(
 	const FGuid Guid,
 	const int32 MinX, const int32 MinY, const int32 MaxX, const int32 MaxY,
@@ -2697,14 +2719,64 @@ LANDSCAPE_API void ALandscapeProxy::Import(
 		PendingTexturePlatformDataCreation.Add(HeightmapInfo.HeightmapTexture);
 	}
 
-	for (UTexture2D* Texture : PendingTexturePlatformDataCreation)
+	// Build a list of all unique materials the landscape uses
+	TArray<UMaterialInterface*> LandscapeMaterials;
+
+	for (ULandscapeComponent* Component : LandscapeComponents)
 	{
-		Texture->FinishCachePlatformData();
-		Texture->PostEditChange();
+		int8 MaxLOD = FMath::CeilLogTwo(Component->SubsectionSizeQuads + 1) - 1;
+
+		for (int8 LODIndex = 0; LODIndex < MaxLOD; ++LODIndex)
+		{
+			UMaterialInterface* Material = Component->GetLandscapeMaterial(LODIndex);
+			LandscapeMaterials.AddUnique(Material);
+		}
 	}
 
-	// Update MaterialInstances (must be done after textures are fully initialized)
-	UpdateAllComponentMaterialInstances();
+	// Update all materials and recreate render state of all landscape components
+	TArray<FComponentRecreateRenderStateContext> RecreateRenderStateContexts;
+
+	{
+		// We disable automatic material update context, to manage it manually
+		GDisableAutomaticTextureMaterialUpdateDependencies = true;
+	
+		FMaterialUpdateContext UpdateContext(FMaterialUpdateContext::EOptions::Default & ~FMaterialUpdateContext::EOptions::RecreateRenderStates);
+
+		for (UTexture2D* Texture : PendingTexturePlatformDataCreation)
+		{
+			Texture->FinishCachePlatformData();
+			Texture->PostEditChange();
+			
+			TSet<UMaterial*> BaseMaterialsThatUseThisTexture;
+
+			for (UMaterialInterface* MaterialInterface : LandscapeMaterials)
+			{
+				if (DoesMaterialUseTexture(MaterialInterface, Texture))
+				{
+					UMaterial* Material = MaterialInterface->GetMaterial();
+					bool MaterialAlreadyCompute = false;
+					BaseMaterialsThatUseThisTexture.Add(Material, &MaterialAlreadyCompute);
+
+					if (!MaterialAlreadyCompute)
+					{
+						if (Material->IsTextureForceRecompileCacheRessource(Texture))
+						{
+							UpdateContext.AddMaterial(Material);
+							Material->UpdateMaterialShaderCacheAndTextureReferences();
+						}
+					}
+				}
+			}
+		}
+		
+		GDisableAutomaticTextureMaterialUpdateDependencies = false;
+
+		// Update MaterialInstances (must be done after textures are fully initialized)
+		UpdateAllComponentMaterialInstances(UpdateContext, RecreateRenderStateContexts);
+
+	// Recreate the render state for this component, needed to update the static drawlist which has cached the MaterialRenderProxies
+	// Must be after the FMaterialUpdateContext is destroyed
+	RecreateRenderStateContexts.Reset();
 
 	if (GetMutableDefault<UEditorExperimentalSettings>()->bProceduralLandscape)
 	{
