Fixed Landscape's Copy/Paste Tool
#rnx

diff --git a/Engine/Source/Editor/LandscapeEditor/Private/LandscapeEdMode.cpp b/Engine/Source/Editor/LandscapeEditor/Private/LandscapeEdMode.cpp
--- a/Engine/Source/Editor/LandscapeEditor/Private/LandscapeEdMode.cpp
+++ b/Engine/Source/Editor/LandscapeEditor/Private/LandscapeEdMode.cpp
@@ -1494,18 +1494,17 @@ bool FEdModeLandscape::ProcessEditCopy()
 
 		if (!Result)
 		{
-			TArray<ALandscapeBlueprintCustomBrush*> CurrentlySelectedBPBrush;
+			ALandscapeBlueprintCustomBrush* CurrentlySelectedBPBrush = nullptr;
 			for (FSelectionIterator It(GEditor->GetSelectedActorIterator()); It; ++It)
 			{
-				ALandscapeBlueprintCustomBrush* Actor = StaticCast<ALandscapeBlueprintCustomBrush*>(*It);
-
-				if (Actor != nullptr)
+				CurrentlySelectedBPBrush = Cast<ALandscapeBlueprintCustomBrush>(*It);
+				if (CurrentlySelectedBPBrush)
 				{
-					CurrentlySelectedBPBrush.Add(Actor);
+					break;
 				}
 			}
 
-			if (CurrentlySelectedBPBrush.Num() == 0)
+			if (!CurrentlySelectedBPBrush)
 			{
 				bool IsSlowTask = IsSlowSelect(CurrentGizmoActor->TargetLandscapeInfo);
 				if (IsSlowTask)
@@ -1550,18 +1549,17 @@ bool FEdModeLandscape::ProcessEditPaste()
 
 		if (!Result)
 		{
-			TArray<ALandscapeBlueprintCustomBrush*> CurrentlySelectedBPBrush;
+			ALandscapeBlueprintCustomBrush* CurrentlySelectedBPBrush = nullptr;
 			for (FSelectionIterator It(GEditor->GetSelectedActorIterator()); It; ++It)
 			{
-				ALandscapeBlueprintCustomBrush* Actor = StaticCast<ALandscapeBlueprintCustomBrush*>(*It);
-
-				if (Actor != nullptr)
+				CurrentlySelectedBPBrush = Cast<ALandscapeBlueprintCustomBrush>(*It);
+				if (CurrentlySelectedBPBrush)
 				{
-					CurrentlySelectedBPBrush.Add(Actor);
+					break;
 				}
 			}
 
-			if (CurrentlySelectedBPBrush.Num() == 0)
+			if (!CurrentlySelectedBPBrush)
 			{
 				bool IsSlowTask = IsSlowSelect(CurrentGizmoActor->TargetLandscapeInfo);
 				if (IsSlowTask)
