Fix crash creating new landscape spline.

diff --git a/Engine/Source/Editor/LandscapeEditor/Private/LandscapeEdModeSplineTools.cpp b/Engine/Source/Editor/LandscapeEditor/Private/LandscapeEdModeSplineTools.cpp
--- a/Engine/Source/Editor/LandscapeEditor/Private/LandscapeEdModeSplineTools.cpp
+++ b/Engine/Source/Editor/LandscapeEditor/Private/LandscapeEdModeSplineTools.cpp
@@ -450,7 +450,7 @@ public:
 			}
 			else
 			{
-				FVector NewSegmentDirection = (NewControlPoint->Location - FirstPoint->Location) * (FirstPoint->ConnectedSegments[0].End ? 1.0f : -1.0f);
+				FVector NewSegmentDirection = (NewControlPoint->Location - FirstPoint->Location) * (FirstPoint->ConnectedSegments.Num() == 0 || FirstPoint->ConnectedSegments[0].End ? 1.0f : -1.0f);
 				NewControlPoint->Rotation = NewSegmentDirection.Rotation();
 			}
 
@@ -469,7 +469,7 @@ public:
 
 			for (ULandscapeSplineControlPoint* ControlPoint : SelectedSplineControlPoints)
 			{
-				if (ControlPoint->ConnectedSegments[0].End)
+				if (ControlPoint->ConnectedSegments.Num() == 0 || ControlPoint->ConnectedSegments[0].End)
 				{
 					AddSegment(ControlPoint, NewControlPoint, bAutoRotateOnJoin, !bDuplicatingControlPoint);
 				}
@@ -698,6 +698,7 @@ public:
 
 	bool UpdateSplitSegment(ULandscapeSplineControlPoint* ControlPoint, const FVector& LocalLocation)
 	{
+		check(ControlPoint->ConnectedSegments.Num() == 2);
 		ULandscapeSplineSegment* Segment0 = ControlPoint->ConnectedSegments[0].Segment;
 		ULandscapeSplineSegment* Segment1 = ControlPoint->ConnectedSegments[1].Segment;
 		ULandscapeSplineSegment *Segment, *NewSegment;
