Fixed Landscape Spline Tangent handles.

diff --git a/Engine/Source/Editor/LandscapeEditor/Private/LandscapeEdMode.cpp b/Engine/Source/Editor/LandscapeEditor/Private/LandscapeEdMode.cpp
--- a/Engine/Source/Editor/LandscapeEditor/Private/LandscapeEdMode.cpp
+++ b/Engine/Source/Editor/LandscapeEditor/Private/LandscapeEdMode.cpp
@@ -1055,6 +1055,63 @@ bool FEdModeLandscape::GetCursor(EMouseCursor::Type& OutCursor) const
 	return Result;
 }
 
+bool FEdModeLandscape::GetOverrideCursorVisibility(bool& bWantsOverride, bool& bHardwareCursorVisible, bool bSoftwareCursorVisible) const
+{
+	if (!IsEditingEnabled())
+	{
+		return false;
+	}
+
+	bool Result = false;
+	if (NewLandscapePreviewMode == ENewLandscapePreviewMode::None)
+	{
+		if (CurrentTool)
+		{
+			Result = CurrentTool->GetOverrideCursorVisibility(bWantsOverride, bHardwareCursorVisible, bSoftwareCursorVisible);
+		}
+	}
+
+	return Result;
+}
+
+bool FEdModeLandscape::PreConvertMouseMovement(FEditorViewportClient* InViewportClient)
+{
+	if (!IsEditingEnabled())
+	{
+		return false;
+	}
+
+	bool Result = false;
+	if (NewLandscapePreviewMode == ENewLandscapePreviewMode::None)
+	{
+		if (CurrentTool)
+		{
+			Result = CurrentTool->PreConvertMouseMovement(InViewportClient);
+		}
+	}
+
+	return Result;
+}
+
+bool FEdModeLandscape::PostConvertMouseMovement(FEditorViewportClient* InViewportClient)
+{
+	if (!IsEditingEnabled())
+	{
+		return false;
+	}
+
+	bool Result = false;
+	if (NewLandscapePreviewMode == ENewLandscapePreviewMode::None)
+	{
+		if (CurrentTool)
+		{
+			Result = CurrentTool->PostConvertMouseMovement(InViewportClient);
+		}
+	}
+
+	return Result;
+}
+
 bool FEdModeLandscape::DisallowMouseDeltaTracking() const
 {
 	// We never want to use the mouse delta tracker while painting
diff --git a/Engine/Source/Editor/LandscapeEditor/Private/LandscapeEdMode.h b/Engine/Source/Editor/LandscapeEditor/Private/LandscapeEdMode.h
--- a/Engine/Source/Editor/LandscapeEditor/Private/LandscapeEdMode.h
+++ b/Engine/Source/Editor/LandscapeEditor/Private/LandscapeEdMode.h
@@ -418,6 +418,15 @@ public:
 
 	virtual bool GetCursor(EMouseCursor::Type& OutCursor) const override;
 
+	/** Get override cursor visibility settings */	
+	virtual bool GetOverrideCursorVisibility(bool& bWantsOverride, bool& bHardwareCursorVisible, bool bSoftwareCursorVisible) const override;
+
+	/** Called before mouse movement is converted to drag/rot */
+	virtual bool PreConvertMouseMovement(FEditorViewportClient* InViewportClient) override;
+
+	/** Called after mouse movement is converted to drag/rot */
+	virtual bool PostConvertMouseMovement(FEditorViewportClient* InViewportClient) override;
+
 	/** Forces real-time perspective viewports */
 	void ForceRealTimeViewports(const bool bEnable, const bool bStoreCurrentState);
 
diff --git a/Engine/Source/Editor/LandscapeEditor/Private/LandscapeEdModeSplineTools.cpp b/Engine/Source/Editor/LandscapeEditor/Private/LandscapeEdModeSplineTools.cpp
--- a/Engine/Source/Editor/LandscapeEditor/Private/LandscapeEdModeSplineTools.cpp
+++ b/Engine/Source/Editor/LandscapeEditor/Private/LandscapeEdModeSplineTools.cpp
@@ -49,6 +49,8 @@ public:
 		, SelectedSplineControlPoints()
 		, SelectedSplineSegments()
 		, DraggingTangent_Segment(NULL)
+		, DraggingTangent_Length(0.0f)
+		, DraggingTangent_CacheCoordSpace(ECoordSystem::COORD_None)
 		, DraggingTangent_End(false)
 		, bMovingControlPoint(false)
 		, bAutoRotateOnJoin(true)
@@ -1442,6 +1444,12 @@ public:
 						HLandscapeSplineProxy_Tangent* SplineProxy = (HLandscapeSplineProxy_Tangent*)HitProxy;
 						DraggingTangent_Segment = SplineProxy->SplineSegment;
 						DraggingTangent_End = SplineProxy->End;
+						DraggingTangent_Length = DraggingTangent_Segment->Connections[DraggingTangent_End].TangentLen;
+
+						// Coord system MUST be set here, even if widget coord system space claims to already be in local space.
+						DraggingTangent_CacheCoordSpace = InViewportClient->GetWidgetCoordSystemSpace();
+						InViewportClient->SetWidgetCoordSystemSpace(ECoordSystem::COORD_Local);
+						InViewportClient->SetRequiredCursorOverride(true, EMouseCursor::GrabHandClosed);
 
 						GEditor->BeginTransaction(LOCTEXT("LandscapeSpline_ModifyTangent", "Modify Landscape Spline Tangent"));
 						ULandscapeSplinesComponent* SplinesComponent = DraggingTangent_Segment->GetOuterULandscapeSplinesComponent();
@@ -1480,6 +1488,9 @@ public:
 
 					DraggingTangent_Segment = NULL;
 
+					InViewportClient->SetWidgetCoordSystemSpace(DraggingTangent_CacheCoordSpace);
+					InViewportClient->SetRequiredCursorOverride(false);
+
 					GEditor->EndTransaction();
 
 					return false; // false to let FEditorViewportClient.InputKey end mouse tracking
@@ -1593,23 +1604,69 @@ public:
 		return UpdateSplitSegment(SelectedControlPoint, Location);
 	}
 
+	virtual bool GetOverrideCursorVisibility(bool& bWantsOverride, bool& bHardwareCursorVisible, bool bSoftwareCursorVisible) const override
+	{
+		if (DraggingTangent_Segment)
+		{
+			bWantsOverride = true;
+			bHardwareCursorVisible = true;
+			bSoftwareCursorVisible = false;
+			return true;
+		}
+
+		bWantsOverride = false;
+		
+		return false;
+	}
+
+	virtual bool PreConvertMouseMovement(FEditorViewportClient* InViewportClient) override
+	{
+		if (DraggingTangent_Segment)
+		{
+			InViewportClient->SetWidgetModeOverride(FWidget::WM_Translate);
+			InViewportClient->SetCurrentWidgetAxis(EAxisList::X);
+			return true;
+		}
+
+		return false;
+	}
+
+	virtual bool PostConvertMouseMovement(FEditorViewportClient* InViewportClient) override
+	{
+		if (DraggingTangent_Segment)
+		{
+			InViewportClient->SetWidgetModeOverride(FWidget::WM_Scale);
+			InViewportClient->SetCurrentWidgetAxis(EAxisList::None);
+			return true;
+		}
+
+		return false;
+	}
+
 	virtual bool InputDelta(FEditorViewportClient* InViewportClient, FViewport* InViewport, FVector& InDrag, FRotator& InRot, FVector& InScale) override
 	{
 		FVector Drag = InDrag;
 
 		if (DraggingTangent_Segment)
 		{
+			InViewportClient->SetRequiredCursorOverride(true, EMouseCursor::GrabHandClosed);
+
 			const ULandscapeSplinesComponent* SplinesComponent = DraggingTangent_Segment->GetOuterULandscapeSplinesComponent();
 			FLandscapeSplineSegmentConnection& Connection = DraggingTangent_Segment->Connections[DraggingTangent_End];
 
 			FVector StartLocation; FRotator StartRotation;
 			Connection.ControlPoint->GetConnectionLocationAndRotation(Connection.SocketName, StartLocation, StartRotation);
+			FVector ForwardVector = FQuatRotationMatrix(StartRotation.Quaternion()).TransformVector(FVector(1.0f, 0.0f, 0.0f));
 
+			FVector DragLocal = SplinesComponent->GetComponentTransform().InverseTransformVector(Drag);
+			float Angle = FMath::Acos(FVector::DotProduct(DragLocal, ForwardVector) / DragLocal.Size());
 			float OldTangentLen = Connection.TangentLen;
-			Connection.TangentLen += SplinesComponent->GetComponentTransform().InverseTransformVector(Drag) | StartRotation.Vector();
+			Connection.TangentLen = DraggingTangent_Length + (Angle < HALF_PI ? 2.0 : -2.0) * DragLocal.Size();
 
-			// Disallow a tangent of exactly 0
-			if (Connection.TangentLen == 0)
+			// Disallow a tangent of exactly 0 and don't allow tangents to flip
+			if ((Connection.TangentLen > 0 && OldTangentLen < 0) || 
+				(Connection.TangentLen < 0 && OldTangentLen > 0) ||
+				 Connection.TangentLen == 0)
 			{
 				if (OldTangentLen > 0)
 				{
@@ -1790,7 +1847,7 @@ public:
 				FVector HandlePos1 = SplinesComponent->GetComponentTransform().TransformPosition(ControlPoint->Location + ControlPoint->Rotation.Vector() * 20);
 				DrawDashedLine(PDI, HandlePos0, HandlePos1, FColor::White, 20, SDPG_Foreground);
 
-				if (GLevelEditorModeTools().GetWidgetMode() == FWidget::WM_Scale)
+				if (GLevelEditorModeTools().GetWidgetMode() == FWidget::WM_Scale && !Viewport->GetClient()->IsOrtho())
 				{
 					for (const FLandscapeSplineConnection& Connection : ControlPoint->ConnectedSegments)
 					{
@@ -1799,16 +1856,17 @@ public:
 
 						FVector StartPos = SplinesComponent->GetComponentTransform().TransformPosition(StartLocation);
 						FVector HandlePos = SplinesComponent->GetComponentTransform().TransformPosition(StartLocation + StartRotation.Vector() * Connection.GetNearConnection().TangentLen / 2);
-						PDI->DrawLine(StartPos, HandlePos, FColor::White, SDPG_Foreground);
 
+						FColor TangentColor = (Connection.Segment == DraggingTangent_Segment && Connection.End == DraggingTangent_End) ? FColor::Yellow : FColor::White;
+						PDI->DrawLine(StartPos, HandlePos, TangentColor, SDPG_Foreground);
 						if (PDI->IsHitTesting()) PDI->SetHitProxy(new HLandscapeSplineProxy_Tangent(Connection.Segment, Connection.End));
-						PDI->DrawPoint(HandlePos, FColor::White, 10.0f, SDPG_Foreground);
+						PDI->DrawPoint(HandlePos, TangentColor, 10.0f, SDPG_Foreground);
 						if (PDI->IsHitTesting()) PDI->SetHitProxy(NULL);
 					}
 				}
 			}
 
-			if (GLevelEditorModeTools().GetWidgetMode() == FWidget::WM_Scale)
+			if (GLevelEditorModeTools().GetWidgetMode() == FWidget::WM_Scale && !Viewport->GetClient()->IsOrtho())
 			{
 				for (ULandscapeSplineSegment* Segment : SelectedSplineSegments)
 				{
@@ -1823,9 +1881,10 @@ public:
 						FVector EndPos = SplinesComponent->GetComponentTransform().TransformPosition(StartLocation);
 						FVector EndHandlePos = SplinesComponent->GetComponentTransform().TransformPosition(StartLocation + StartRotation.Vector() * Connection.TangentLen / 2);
 
-						PDI->DrawLine(EndPos, EndHandlePos, FColor::White, SDPG_Foreground);
+						FColor TangentColor = (Segment == DraggingTangent_Segment && End == DraggingTangent_End) ? FColor::Yellow : FColor::White;
+						PDI->DrawLine(EndPos, EndHandlePos, TangentColor, SDPG_Foreground);
 						if (PDI->IsHitTesting()) PDI->SetHitProxy(new HLandscapeSplineProxy_Tangent(Segment, !!End));
-						PDI->DrawPoint(EndHandlePos, FColor::White, 10.0f, SDPG_Foreground);
+						PDI->DrawPoint(EndHandlePos, TangentColor, 10.0f, SDPG_Foreground);
 						if (PDI->IsHitTesting()) PDI->SetHitProxy(NULL);
 					}
 				}
@@ -1851,7 +1910,7 @@ public:
 
 	virtual bool UsesTransformWidget() const override
 	{
-		if (SelectedSplineControlPoints.Num() > 0)
+		if (SelectedSplineControlPoints.Num() > 0 || DraggingTangent_Segment)
 		{
 			// The editor can try to render the transform widget before the landscape editor ticks and realizes that the landscape has been hidden/deleted
 			const ALandscapeProxy* LandscapeProxy = EdMode->CurrentToolTarget.LandscapeInfo->GetLandscapeProxy();
@@ -1889,10 +1948,22 @@ public:
 
 	virtual FVector GetWidgetLocation() const override
 	{
-		if (SelectedSplineControlPoints.Num() > 0)
+		const ALandscapeProxy* LandscapeProxy = EdMode->CurrentToolTarget.LandscapeInfo->GetLandscapeProxy();
+		if (LandscapeProxy)
 		{
-			const ALandscapeProxy* LandscapeProxy = EdMode->CurrentToolTarget.LandscapeInfo->GetLandscapeProxy();
-			if (LandscapeProxy)
+			if (DraggingTangent_Segment)
+			{
+				const FLandscapeSplineSegmentConnection& Connection = DraggingTangent_Segment->Connections[DraggingTangent_End];
+				ULandscapeSplineControlPoint* ControlPoint = Connection.ControlPoint;
+				ULandscapeSplinesComponent* SplinesComponent = ControlPoint->GetOuterULandscapeSplinesComponent();
+				FVector StartLocation; FRotator StartRotation;
+				ControlPoint->GetConnectionLocationAndRotation(Connection.SocketName, StartLocation, StartRotation);
+
+				// Return tangent handle location.
+				return SplinesComponent->GetComponentTransform().TransformPosition(StartLocation + StartRotation.Vector() * DraggingTangent_Length / 2);
+
+			}
+			else if (SelectedSplineControlPoints.Num() > 0)
 			{
 				ULandscapeSplineControlPoint* FirstPoint = *SelectedSplineControlPoints.CreateConstIterator();
 				ULandscapeSplinesComponent* SplinesComponent = FirstPoint->GetOuterULandscapeSplinesComponent();
@@ -1905,10 +1976,18 @@ public:
 
 	virtual FMatrix GetWidgetRotation() const override
 	{
-		if (SelectedSplineControlPoints.Num() > 0)
+		const ALandscapeProxy* LandscapeProxy = EdMode->CurrentToolTarget.LandscapeInfo->GetLandscapeProxy();
+		if (LandscapeProxy)
 		{
-			const ALandscapeProxy* LandscapeProxy = EdMode->CurrentToolTarget.LandscapeInfo->GetLandscapeProxy();
-			if (LandscapeProxy)
+			if (DraggingTangent_Segment)
+			{
+				const FLandscapeSplineSegmentConnection& Connection = DraggingTangent_Segment->Connections[DraggingTangent_End];
+				ULandscapeSplinesComponent* SplinesComponent = Connection.ControlPoint->GetOuterULandscapeSplinesComponent();
+				FVector StartLocation; FRotator StartRotation;
+				Connection.ControlPoint->GetConnectionLocationAndRotation(Connection.SocketName, StartLocation, StartRotation);
+				return FQuatRotationTranslationMatrix(StartRotation.Quaternion() * SplinesComponent->GetComponentTransform().GetRotation(), FVector::ZeroVector);
+			}
+			else if (SelectedSplineControlPoints.Num() > 0)
 			{
 				ULandscapeSplineControlPoint* FirstPoint = *SelectedSplineControlPoints.CreateConstIterator();
 				ULandscapeSplinesComponent* SplinesComponent = FirstPoint->GetOuterULandscapeSplinesComponent();
@@ -2150,6 +2229,8 @@ protected:
 	TSet<ULandscapeSplineSegment*> SelectedSplineSegments;
 
 	ULandscapeSplineSegment* DraggingTangent_Segment;
+	float DraggingTangent_Length;
+	ECoordSystem DraggingTangent_CacheCoordSpace;
 	uint32 DraggingTangent_End : 1;
 
 	uint32 bMovingControlPoint : 1;
diff --git a/Engine/Source/Editor/LandscapeEditor/Public/LandscapeToolInterface.h b/Engine/Source/Editor/LandscapeEditor/Public/LandscapeToolInterface.h
--- a/Engine/Source/Editor/LandscapeEditor/Public/LandscapeToolInterface.h
+++ b/Engine/Source/Editor/LandscapeEditor/Public/LandscapeToolInterface.h
@@ -231,6 +231,15 @@ public:
 	virtual FMatrix GetWidgetRotation() const { return FMatrix::Identity; }
 	virtual bool DisallowMouseDeltaTracking() const { return false; }
 
+	/** Get override cursor visibility settings */
+	virtual bool GetOverrideCursorVisibility(bool& bWantsOverride, bool& bHardwareCursorVisible, bool bSoftwareCursorVisible) const { return false; }
+
+	/** Called before mouse movement is converted to drag/rot */
+	virtual bool PreConvertMouseMovement(FEditorViewportClient* InViewportClient) { return false; }
+
+	/** Called after mouse movement is converted to drag/rot */
+	virtual bool PostConvertMouseMovement(FEditorViewportClient* InViewportClient) { return false; }
+
 	virtual void SetCanToolBeActivated(bool Value) { }
 	virtual bool CanToolBeActivated() const { return true;  }
 	virtual void SetExternalModifierPressed(const bool bPressed) {};
diff --git a/Engine/Source/Editor/UnrealEd/Private/EditorModeManager.cpp b/Engine/Source/Editor/UnrealEd/Private/EditorModeManager.cpp
--- a/Engine/Source/Editor/UnrealEd/Private/EditorModeManager.cpp
+++ b/Engine/Source/Editor/UnrealEd/Private/EditorModeManager.cpp
@@ -924,6 +924,36 @@ bool FEditorModeTools::GetCursor(EMouseCursor::Type& OutCursor) const
 	return bHandled;
 }
 
+bool FEditorModeTools::GetOverrideCursorVisibility(bool& bWantsOverride, bool& bHardwareCursorVisible, bool bSoftwareCursorVisible) const
+{
+	bool bHandled = false;
+	for (const auto& Mode : Modes)
+	{
+		bHandled |= Mode->GetOverrideCursorVisibility(bWantsOverride, bHardwareCursorVisible, bSoftwareCursorVisible);
+	}
+	return bHandled;
+}
+
+bool FEditorModeTools::PreConvertMouseMovement(FEditorViewportClient* InViewportClient)
+{
+	bool bHandled = false;
+	for (const auto& Mode : Modes)
+	{
+		bHandled |= Mode->PreConvertMouseMovement(InViewportClient);
+	}
+	return bHandled;
+}
+
+bool FEditorModeTools::PostConvertMouseMovement(FEditorViewportClient* InViewportClient)
+{
+	bool bHandled = false;
+	for (const auto& Mode : Modes)
+	{
+		bHandled |= Mode->PostConvertMouseMovement(InViewportClient);
+	}
+	return bHandled;
+}
+
 /**
  * Used to cycle widget modes
  */
diff --git a/Engine/Source/Editor/UnrealEd/Private/EditorViewportClient.cpp b/Engine/Source/Editor/UnrealEd/Private/EditorViewportClient.cpp
--- a/Engine/Source/Editor/UnrealEd/Private/EditorViewportClient.cpp
+++ b/Engine/Source/Editor/UnrealEd/Private/EditorViewportClient.cpp
@@ -2088,6 +2088,10 @@ void FEditorViewportClient::UpdateMouseDelta()
 			const bool bIsUsingTrackpad = FSlateApplication::Get().IsUsingTrackpad();
 			const bool bIsNonOrbitMiddleMouse = MiddleMouseButtonDown && !IsAltPressed();
 
+			// If a tool is overriding current widget mode behavior, it may need to 
+			// temporarily set a different widget mode while converting mouse movement.
+			ModeTools->PreConvertMouseMovement(this);
+
 			// Convert the movement delta into drag/rotation deltas
 			FVector Drag;
 			FRotator Rot;
@@ -2124,6 +2128,8 @@ void FEditorViewportClient::UpdateMouseDelta()
 				}
 			}
 
+			ModeTools->PostConvertMouseMovement(this);
+
 			const bool bInputHandledByGizmos = InputWidgetDelta( Viewport, CurrentAxis, Drag, Rot, Scale );
 					
 			if( !Rot.IsZero() )
@@ -4580,6 +4586,18 @@ void FEditorViewportClient::UpdateRequiredCursorVisibility()
 	bool ShiftDown = IsShiftPressed();
 	bool ControlDown = IsCtrlPressed();
 
+	bool bOverrideCursorVisibility = false;
+	bool bHardwareCursorVisible = false;
+	bool bSoftwareCursorVisible = false;
+	if (ModeTools->GetOverrideCursorVisibility(bOverrideCursorVisibility, bHardwareCursorVisible, bSoftwareCursorVisible))
+	{
+		if (bOverrideCursorVisibility)
+		{
+			SetRequiredCursor(bHardwareCursorVisible, bSoftwareCursorVisible);
+			return;
+		}
+	}
+
 	if (GetViewportType() == LVT_None)
 	{
 		SetRequiredCursor(true, false);
@@ -4680,6 +4698,11 @@ void FEditorViewportClient::SetRequiredCursorOverride( bool WantOverride, EMouse
 	RequiredCursorVisibiltyAndAppearance.RequiredCursor = RequiredCursor;
 }
 
+void FEditorViewportClient::SetWidgetModeOverride(FWidget::EWidgetMode InWidgetMode)
+{
+	ModeTools->SetWidgetModeOverride(InWidgetMode);
+}
+
 EAxisList::Type FEditorViewportClient::GetCurrentWidgetAxis() const
 {
 	return Widget->GetCurrentAxis();
diff --git a/Engine/Source/Editor/UnrealEd/Public/EdMode.h b/Engine/Source/Editor/UnrealEd/Public/EdMode.h
--- a/Engine/Source/Editor/UnrealEd/Public/EdMode.h
+++ b/Engine/Source/Editor/UnrealEd/Public/EdMode.h
@@ -132,6 +132,15 @@ public:
 	 */
 	virtual bool GetCursor(EMouseCursor::Type& OutCursor) const { return false; }
 
+	/** Get override cursor visibility settings */
+	virtual bool GetOverrideCursorVisibility(bool& bWantsOverride, bool& bHardwareCursorVisible, bool bSoftwareCursorVisible) const { return false; }
+
+	/** Called before mouse movement is converted to drag/rot */
+	virtual bool PreConvertMouseMovement(FEditorViewportClient* InViewportClient) { return false; }
+
+	/** Called after mouse movement is converted to drag/rot */
+	virtual bool PostConvertMouseMovement(FEditorViewportClient* InViewportClient) { return false;}
+
 	virtual bool ShouldDrawBrushWireframe( AActor* InActor ) const { return true; }
 	virtual bool GetCustomDrawingCoordinateSystem(FMatrix& InMatrix, void* InData);
 	virtual bool GetCustomInputCoordinateSystem( FMatrix& InMatrix, void* InData ) { return 0; }
diff --git a/Engine/Source/Editor/UnrealEd/Public/EditorModeManager.h b/Engine/Source/Editor/UnrealEd/Public/EditorModeManager.h
--- a/Engine/Source/Editor/UnrealEd/Public/EditorModeManager.h
+++ b/Engine/Source/Editor/UnrealEd/Public/EditorModeManager.h
@@ -262,6 +262,15 @@ public:
 	/** Get a cursor to override the default with, if any */
 	bool GetCursor(EMouseCursor::Type& OutCursor) const;
 
+	/** Get override cursor visibility settings */
+	bool GetOverrideCursorVisibility(bool& bWantsOverride, bool& bHardwareCursorVisible, bool bSoftwareCursorVisible) const;
+
+	/** Called before converting mouse movement to drag/rot */
+	bool PreConvertMouseMovement(FEditorViewportClient* InViewportClient);
+
+	/** Called after converting mouse movement to drag/rot */
+	bool PostConvertMouseMovement(FEditorViewportClient* InViewportClient);
+
 	/**
 	 * Returns a good location to draw the widget at.
 	 */
diff --git a/Engine/Source/Editor/UnrealEd/Public/EditorViewportClient.h b/Engine/Source/Editor/UnrealEd/Public/EditorViewportClient.h
--- a/Engine/Source/Editor/UnrealEd/Public/EditorViewportClient.h
+++ b/Engine/Source/Editor/UnrealEd/Public/EditorViewportClient.h
@@ -977,8 +977,12 @@ public:
 	
 	EAxisList::Type GetCurrentWidgetAxis() const;
 
+	/** Overrides current cursor. */
 	void SetRequiredCursorOverride( bool WantOverride, EMouseCursor::Type RequiredCursor = EMouseCursor::Default ); 
 
+	/** Overrides current widget mode */
+	void SetWidgetModeOverride(FWidget::EWidgetMode InWidgetMode);
+
 	/** Get the camera speed for this viewport */
 	float GetCameraSpeed() const;
 
