[NO NON-PII COMMIT MESSAGE]

diff --git a/Engine/Source/Editor/LandscapeEditor/Private/LandscapeEdModePaintTools.cpp b/Engine/Source/Editor/LandscapeEditor/Private/LandscapeEdModePaintTools.cpp
--- a/Engine/Source/Editor/LandscapeEditor/Private/LandscapeEdModePaintTools.cpp
+++ b/Engine/Source/Editor/LandscapeEditor/Private/LandscapeEdModePaintTools.cpp
@@ -822,15 +822,15 @@ public:
 								float ScaleZ = LocalToWorld.GetScale3D().Z;
 								float TranslateZ = LocalToWorld.GetTranslation().Z;
 								float TerraceInterval = UISettings->TerraceInterval;
-								float Smoothness = UISettings->TerraceSmooth;								
+								float Smoothness = UISettings->TerraceSmooth;
 								float WorldHeight = LandscapeDataAccess::GetLocalHeight(DataScanline[X]);
-								
+
 								//move into world space
 								WorldHeight = (WorldHeight * ScaleZ) + TranslateZ;
 								float CurrentHeight = WorldHeight;
 
 								//smoothing part
-								float CurrentLevel = WorldHeight / TerraceInterval;								
+								float CurrentLevel = WorldHeight / TerraceInterval;
 								Smoothness = 1.0f / FMath::Max(Smoothness, 0.0001f);
 								float CurrentPhase = FMath::Frac(CurrentLevel);
 								float Halfmask = FMath::Clamp(FMath::CeilToFloat(CurrentPhase - 0.5f), 0.0f, 1.0f);
@@ -838,12 +838,34 @@ public:
 								float SCurve = FMath::Lerp(CurrentPhase, (1.0f - CurrentPhase), Halfmask) * 2.0f;
 								SCurve = FMath::Pow(SCurve, Smoothness) * 0.5f;
 								SCurve = FMath::Lerp(SCurve, 1.0f - SCurve, Halfmask) * TerraceInterval;
-								WorldHeight = (CurrentLevel * TerraceInterval)  + SCurve;
+								WorldHeight = (CurrentLevel * TerraceInterval) + SCurve;
 								//end of smoothing part
 
-								float FinalHeight = FMath::Lerp(CurrentHeight, WorldHeight , Strength);
+								float FinalHeight = FMath::Lerp(CurrentHeight, WorldHeight, Strength);
 								FinalHeight = (FinalHeight - TranslateZ) / ScaleZ;
-								DataScanline[X] = LandscapeDataAccess::GetTexHeight(FinalHeight);	
+								DataScanline[X] = LandscapeDataAccess::GetTexHeight(FinalHeight);
+							}
+							break;
+						case ELandscapeToolFlattenMode::Interval:
+							{
+								const FTransform& LocalToWorld = this->Target.LandscapeInfo->GetLandscapeProxy()->ActorToWorld();
+								float ScaleZ = LocalToWorld.GetScale3D().Z;
+								float TranslateZ = LocalToWorld.GetTranslation().Z;
+								float TerraceInterval = UISettings->TerraceInterval;
+								float TargetHeight = LandscapeDataAccess::GetLocalHeight(FlattenHeight);
+								float CurrentHeight = LandscapeDataAccess::GetLocalHeight(DataScanline[X]);
+														
+								//move into world space
+								TargetHeight = (TargetHeight * ScaleZ) + TranslateZ;
+								CurrentHeight = (CurrentHeight * ScaleZ) + TranslateZ;
+
+								TargetHeight = (FMath::RoundToFloat(TargetHeight / TerraceInterval)) * TerraceInterval;
+								TargetHeight = FMath::Lerp(CurrentHeight, TargetHeight, BrushValue);
+							
+								//back to local space of landscape object
+								TargetHeight = (TargetHeight - TranslateZ) / ScaleZ;
+								DataScanline[X] = LandscapeDataAccess::GetTexHeight(TargetHeight);
+
 							}
 							break;
 						case ELandscapeToolFlattenMode::Raise:
diff --git a/Engine/Source/Editor/LandscapeEditor/Public/LandscapeEditorObject.h b/Engine/Source/Editor/LandscapeEditor/Public/LandscapeEditorObject.h
--- a/Engine/Source/Editor/LandscapeEditor/Public/LandscapeEditorObject.h
+++ b/Engine/Source/Editor/LandscapeEditor/Public/LandscapeEditorObject.h
@@ -29,8 +29,11 @@ enum class ELandscapeToolFlattenMode : int8
 	/** Flatten may only lower values, values below the clicked point will be left unchanged */
 	Lower = 2,
 
+	/** Flatten to closest specific terrace interval at the clicked point */
+	Interval = 3,
+
 	/** Flatten to specific terrace height intervals */
-	Terrace = 3
+	Terrace = 4
 };
 
 UENUM()
