- Landscape Layer System : Added Flatten Tool support on all layers
- Set Terrace Interval/Smooth init value to valid ones and added clamping to avoid invalid results
- Avoid flattening in Paint mode when using Terrace or Interval Flatten mode
- Avoid using Flatten slope if use Flatten Target because no normal is provided
#rnx

diff --git a/Engine/Source/Editor/LandscapeEditor/Private/LandscapeEdMode.cpp b/Engine/Source/Editor/LandscapeEditor/Private/LandscapeEdMode.cpp
--- a/Engine/Source/Editor/LandscapeEditor/Private/LandscapeEdMode.cpp
+++ b/Engine/Source/Editor/LandscapeEditor/Private/LandscapeEdMode.cpp
@@ -4596,7 +4596,7 @@ bool FEdModeLandscape::CanEditLayer(FText* Reason /*=nullptr*/, FLandscapeLayer*
 				}
 				return false;
 			}
-			else if ((TargetLayerIndex > 0) && (CurrentTool->GetToolName() == FName("Ramp") || CurrentTool->GetToolName() == FName("Flatten")))
+			else if ((TargetLayerIndex > 0) && (CurrentTool->GetToolName() == FName("Ramp")))
 			{
 				if (Reason)
 				{
diff --git a/Engine/Source/Editor/LandscapeEditor/Private/LandscapeEdModePaintTools.cpp b/Engine/Source/Editor/LandscapeEditor/Private/LandscapeEdModePaintTools.cpp
--- a/Engine/Source/Editor/LandscapeEditor/Private/LandscapeEdModePaintTools.cpp
+++ b/Engine/Source/Editor/LandscapeEditor/Private/LandscapeEdModePaintTools.cpp
@@ -727,51 +727,95 @@ public:
 template<class ToolTarget>
 class FLandscapeToolStrokeFlatten : public FLandscapeToolStrokePaintBase<ToolTarget>
 {
-	typename ToolTarget::CacheClass::DataType FlattenHeight;
+	typename ToolTarget::CacheClass::DataType FlattenValue;
 
 	FVector FlattenNormal;
 	float FlattenPlaneDist;
-	bool bInitializedFlattenHeight;
+	bool bInitializedFlattenValue;
 	bool bTargetIsHeightmap;
 
 public:
 	FLandscapeToolStrokeFlatten(FEdModeLandscape* InEdMode, FEditorViewportClient* InViewportClient, const FLandscapeToolTarget& InTarget)
 		: FLandscapeToolStrokePaintBase<ToolTarget>(InEdMode, InViewportClient, InTarget)
-		, bInitializedFlattenHeight(false)
+		, bInitializedFlattenValue(false)
 		, bTargetIsHeightmap(InTarget.TargetType == ELandscapeToolTargetType::Heightmap)
 	{
 		if (InEdMode->UISettings->bUseFlattenTarget && bTargetIsHeightmap)
 		{
 			FTransform LocalToWorld = InTarget.LandscapeInfo->GetLandscapeProxy()->ActorToWorld();
 			float Height = (InEdMode->UISettings->FlattenTarget - LocalToWorld.GetTranslation().Z) / LocalToWorld.GetScale3D().Z;
-			FlattenHeight = LandscapeDataAccess::GetTexHeight(Height);
-			bInitializedFlattenHeight = true;
+			FlattenValue = LandscapeDataAccess::GetTexHeight(Height);
+			bInitializedFlattenValue = true;
 		}
 	}
 
 	void Apply(FEditorViewportClient* ViewportClient, FLandscapeBrush* Brush, const ULandscapeEditorObject* UISettings, const TArray<FLandscapeToolInteractorPosition>& InteractorPositions)
-	{
+ 	{
 		if (!this->LandscapeInfo) return;
+		bool bLayerSystemFlattenHeightMode = this->LandscapeInfo->CanHaveLayersContent() && bTargetIsHeightmap;
+		ALandscape* Landscape = this->LandscapeInfo->LandscapeActor.Get();
+		if (!Landscape && bLayerSystemFlattenHeightMode)
+			return;
 
-		if (!bInitializedFlattenHeight || (UISettings->bPickValuePerApply && bTargetIsHeightmap))
+		// Can't use slope if use Flatten Target because no normal is provided
+		bool bUseSlopeFlatten = UISettings->bUseSlopeFlatten && !UISettings->bUseFlattenTarget;
+
+		if (!bInitializedFlattenValue || (UISettings->bPickValuePerApply && bTargetIsHeightmap))
 		{
-			bInitializedFlattenHeight = false;
+			bInitializedFlattenValue = false;
 			float FlattenX = InteractorPositions[0].Position.X;
 			float FlattenY = InteractorPositions[0].Position.Y;
 			int32 FlattenHeightX = FMath::FloorToInt(FlattenX);
 			int32 FlattenHeightY = FMath::FloorToInt(FlattenY);
 
-			this->Cache.CacheData(FlattenHeightX, FlattenHeightY, FlattenHeightX + 1, FlattenHeightY + 1);
-			float HeightValue = this->Cache.GetValue(FlattenX, FlattenY);
-			FlattenHeight = HeightValue;
+			float InterpolatedValue = 0.f;
+			float V00 = 0.f;
+			float V10 = 0.f;
+			float V01 = 0.f;
+			float V11 = 0.f;
+			if (bLayerSystemFlattenHeightMode)
+			{
+				// Can't rely on cache in this mode
+				FScopedSetLandscapeEditingLayer Scope(Landscape, FGuid());
+				ToolTarget::CacheClass::DataType P00, P10, P01, P11;
+				this->Cache.DataAccess.GetDataFast(FlattenHeightX,     FlattenHeightY,     FlattenHeightX,     FlattenHeightY,     &P00);
+				this->Cache.DataAccess.GetDataFast(FlattenHeightX + 1, FlattenHeightY,     FlattenHeightX + 1, FlattenHeightY,     &P10);
+				this->Cache.DataAccess.GetDataFast(FlattenHeightX,     FlattenHeightY + 1, FlattenHeightX,     FlattenHeightY + 1, &P01);
+				this->Cache.DataAccess.GetDataFast(FlattenHeightX + 1, FlattenHeightY + 1, FlattenHeightX + 1, FlattenHeightY + 1, &P11);
+				V00 = P00;
+				V10 = P10;
+				V01 = P01;
+				V11 = P11;
+				InterpolatedValue = FMath::Lerp(FMath::Lerp(V00, V10, FlattenX - FlattenHeightX), FMath::Lerp(V01, V11, FlattenX - FlattenHeightX), FlattenY - FlattenHeightY);
+			}
+			else
+			{
+				this->Cache.CacheData(FlattenHeightX, FlattenHeightY, FlattenHeightX + 1, FlattenHeightY + 1);
+				InterpolatedValue = this->Cache.GetValue(FlattenX, FlattenY);
+			}
+			FlattenValue = InterpolatedValue;
 
-			if (UISettings->bUseSlopeFlatten && bTargetIsHeightmap)
+			if (bUseSlopeFlatten && bTargetIsHeightmap)
 			{
-				FlattenNormal = this->Cache.GetNormal(FlattenHeightX, FlattenHeightY);
-				FlattenPlaneDist = -(FlattenNormal | FVector(FlattenX, FlattenY, HeightValue));
+				if (bLayerSystemFlattenHeightMode)
+				{
+					// Can't rely on cache in this mode
+					FVector Vert00 = FVector(0.0f, 0.0f, V00);
+					FVector Vert01 = FVector(0.0f, 1.0f, V01);
+					FVector Vert10 = FVector(1.0f, 0.0f, V10);
+					FVector Vert11 = FVector(1.0f, 1.0f, V11);
+					FVector FaceNormal1 = ((Vert00 - Vert10) ^ (Vert10 - Vert11)).GetSafeNormal();
+					FVector FaceNormal2 = ((Vert11 - Vert01) ^ (Vert01 - Vert00)).GetSafeNormal();
+					FlattenNormal = (FaceNormal1 + FaceNormal2).GetSafeNormal();
+				}
+				else
+				{
+					FlattenNormal = this->Cache.GetNormal(FlattenHeightX, FlattenHeightY);
+				}
+				FlattenPlaneDist = -(FlattenNormal | FVector(FlattenX, FlattenY, InterpolatedValue));
 			}
 
-			bInitializedFlattenHeight = true;
+			bInitializedFlattenValue = true;
 		}
 
 
@@ -798,10 +842,73 @@ public:
 			Y2 += 1;
 		}
 
-		this->Cache.CacheData(X1, Y1, X2, Y2);
-
 		TArray<typename ToolTarget::CacheClass::DataType> Data;
-		this->Cache.GetCachedData(X1, Y1, X2, Y2, Data);
+		TArray<typename ToolTarget::CacheClass::DataType> BottomLayersData;
+		FGuid EditingLayer;
+		uint8 EditingLayerIndex = MAX_uint8;
+		TSet<ULandscapeComponent*> AffectedComponents;
+		TArray<bool> BackupLayerVisibility;
+		auto& LayerSystemDataAccess = this->Cache.DataAccess;
+
+		auto SynchronousUpdateHeightmapForComponents = [&](const TSet<ULandscapeComponent*>& InComponents)
+		{
+			for (ULandscapeComponent* Component : InComponents)
+			{
+				Component->RequestHeightmapUpdate();
+			}
+			Landscape->ForceUpdateLayersContent();
+		};
+
+		auto SetLayersVisibility = [&](const TArray<bool>& InLayerVisibility)
+		{
+			check(InLayerVisibility.Num() == Landscape->GetLayerCount());
+			for (int i = 0; i < InLayerVisibility.Num(); ++i)
+			{
+				if (FLandscapeLayer* Layer = Landscape->GetLayer(i))
+				{
+					Layer->bVisible = InLayerVisibility[i];
+				}
+			}
+		};
+
+		auto SynchronousUpdateComponentVisibilityForHeight = [&](const TSet<ULandscapeComponent*>& InComponents, const TArray<bool>& InLayerVisibility)
+		{
+			SetLayersVisibility(InLayerVisibility);
+			SynchronousUpdateHeightmapForComponents(InComponents);
+		};
+
+		if (bLayerSystemFlattenHeightMode)
+		{
+			EditingLayer = Landscape->GetEditingLayer();
+			LandscapeInfo->GetComponentsInRegion(X1, Y1, X2, Y2, AffectedComponents);
+			TArray<bool> NewLayerVisibility;
+			for (int i = 0; i < Landscape->GetLayerCount(); ++i)
+			{
+				FLandscapeLayer* CurrentLayer = Landscape->GetLayer(i);
+				BackupLayerVisibility.Add(CurrentLayer->bVisible);
+				NewLayerVisibility.Add((i > EditingLayerIndex) ? false : CurrentLayer->bVisible);
+				if (CurrentLayer->Guid == EditingLayer)
+				{
+					EditingLayerIndex = i;
+				}
+			}
+			check(EditingLayerIndex < Landscape->GetLayerCount());
+			
+			FScopedSetLandscapeEditingLayer Scope(Landscape, FGuid());
+			SynchronousUpdateComponentVisibilityForHeight(AffectedComponents, NewLayerVisibility);
+			Data.AddZeroed((X2 - X1 + 1) * (Y2 - Y1 + 1));
+			LayerSystemDataAccess.GetDataFast(X1, Y1, X2, Y2, Data.GetData());
+			
+			BottomLayersData.AddZeroed((X2 - X1 + 1) * (Y2 - Y1 + 1));
+			NewLayerVisibility[EditingLayerIndex] = false;
+			SynchronousUpdateComponentVisibilityForHeight(AffectedComponents, NewLayerVisibility);
+			LayerSystemDataAccess.GetDataFast(X1, Y1, X2, Y2, BottomLayersData.GetData());
+		}
+		else
+		{
+			this->Cache.CacheData(X1, Y1, X2, Y2);
+			this->Cache.GetCachedData(X1, Y1, X2, Y2, Data);
+		}
 
 		// Apply the brush
 		for (int32 Y = BrushInfo.GetBounds().Min.Y; Y < BrushInfo.GetBounds().Max.Y; Y++)
@@ -817,17 +924,18 @@ public:
 				{
 					float Strength = FMath::Clamp<float>(BrushValue * UISettings->ToolStrength * Pressure, 0.0f, 1.0f);
 
-					if (!(UISettings->bUseSlopeFlatten && bTargetIsHeightmap))
+					if (!(bUseSlopeFlatten && bTargetIsHeightmap))
 					{
-						int32 Delta = DataScanline[X] - FlattenHeight;
+						int32 Delta = DataScanline[X] - FlattenValue;
 						switch (UISettings->FlattenMode)
 						{
 						case ELandscapeToolFlattenMode::Terrace:
+							if (bTargetIsHeightmap)
 							{
 								const FTransform& LocalToWorld = this->Target.LandscapeInfo->GetLandscapeProxy()->ActorToWorld();
 								float ScaleZ = LocalToWorld.GetScale3D().Z;
 								float TranslateZ = LocalToWorld.GetTranslation().Z;
-								float TerraceInterval = UISettings->TerraceInterval;
+								float TerraceInterval = FMath::Clamp(UISettings->TerraceInterval, 1.0f, 32768.f);
 								float Smoothness = UISettings->TerraceSmooth;
 								float WorldHeight = LandscapeDataAccess::GetLocalHeight(DataScanline[X]);
 
@@ -853,12 +961,13 @@ public:
 							}
 							break;
 						case ELandscapeToolFlattenMode::Interval:
+							if (bTargetIsHeightmap)
 							{
 								const FTransform& LocalToWorld = this->Target.LandscapeInfo->GetLandscapeProxy()->ActorToWorld();
 								float ScaleZ = LocalToWorld.GetScale3D().Z;
 								float TranslateZ = LocalToWorld.GetTranslation().Z;
 								float TerraceInterval = UISettings->TerraceInterval;
-								float TargetHeight = LandscapeDataAccess::GetLocalHeight(FlattenHeight);
+								float TargetHeight = LandscapeDataAccess::GetLocalHeight(FlattenValue);
 								float CurrentHeight = LandscapeDataAccess::GetLocalHeight(DataScanline[X]);
 														
 								//move into world space
@@ -877,24 +986,24 @@ public:
 						case ELandscapeToolFlattenMode::Raise:
 							if (Delta < 0)
 							{
-								DataScanline[X] = FMath::CeilToInt(FMath::Lerp((float)DataScanline[X], (float)FlattenHeight, Strength));
+								DataScanline[X] = FMath::CeilToInt(FMath::Lerp((float)DataScanline[X], (float)FlattenValue, Strength));
 							}
 							break;
 						case ELandscapeToolFlattenMode::Lower:
 							if (Delta > 0)
 							{
-								DataScanline[X] = FMath::FloorToInt(FMath::Lerp((float)DataScanline[X], (float)FlattenHeight, Strength));
+								DataScanline[X] = FMath::FloorToInt(FMath::Lerp((float)DataScanline[X], (float)FlattenValue, Strength));
 							}
 							break;
 						default:
 						case ELandscapeToolFlattenMode::Both:
 							if (Delta > 0)
 							{
-								DataScanline[X] = FMath::FloorToInt(FMath::Lerp((float)DataScanline[X], (float)FlattenHeight, Strength));
+								DataScanline[X] = FMath::FloorToInt(FMath::Lerp((float)DataScanline[X], (float)FlattenValue, Strength));
 							}
 							else
 							{
-								DataScanline[X] = FMath::CeilToInt(FMath::Lerp((float)DataScanline[X], (float)FlattenHeight, Strength));
+								DataScanline[X] = FMath::CeilToInt(FMath::Lerp((float)DataScanline[X], (float)FlattenValue, Strength));
 							}
 							break;
 						}
@@ -936,8 +1045,27 @@ public:
 			}
 		}
 
-		this->Cache.SetCachedData(X1, Y1, X2, Y2, Data, UISettings->PaintingRestriction);
-		this->Cache.Flush();
+		if (bLayerSystemFlattenHeightMode)
+		{
+			const float Alpha = Landscape->GetLayerAlpha(EditingLayerIndex, bTargetIsHeightmap);
+			const float InverseAlpha = (Alpha != 0.f) ? 1.f / Alpha : 1.f;
+			for (int i = 0; i < Data.Num(); ++i)
+			{
+				float Diff = (LandscapeDataAccess::GetLocalHeight(Data[i]) - LandscapeDataAccess::GetLocalHeight(BottomLayersData[i])) * InverseAlpha;
+				Data[i] = LandscapeDataAccess::GetTexHeight(Diff);
+			}
+			check(EditingLayer == Landscape->GetEditingLayer());
+			
+			SetLayersVisibility(BackupLayerVisibility);
+			LayerSystemDataAccess.SetData(X1, Y1, X2, Y2, Data.GetData(), UISettings->PaintingRestriction);
+			LayerSystemDataAccess.Flush();
+			SynchronousUpdateHeightmapForComponents(AffectedComponents);
+		}
+		else
+		{
+			this->Cache.SetCachedData(X1, Y1, X2, Y2, Data, UISettings->PaintingRestriction);
+			this->Cache.Flush();
+		}
 	}
 };
 
diff --git a/Engine/Source/Editor/LandscapeEditor/Private/LandscapeEdModeTools.h b/Engine/Source/Editor/LandscapeEditor/Private/LandscapeEdModeTools.h
--- a/Engine/Source/Editor/LandscapeEditor/Private/LandscapeEdModeTools.h
+++ b/Engine/Source/Editor/LandscapeEditor/Private/LandscapeEdModeTools.h
@@ -1156,6 +1156,7 @@ public:
 			{
 				Landscape->RequestLayersContentUpdate(GetBeginToolContentUpdateFlag());
 				Landscape->SetEditingLayer(this->EdMode->GetCurrentLayerGuid());
+				Landscape->SetGrassUpdateEnabled(false);
 			}
 		}
 
@@ -1234,6 +1235,7 @@ public:
 			{
 				Landscape->RequestLayersContentUpdate(GetEndToolContentUpdateFlag());
 				Landscape->SetEditingLayer();
+				Landscape->SetGrassUpdateEnabled(true);
 			}
 		}
 	}
diff --git a/Engine/Source/Editor/LandscapeEditor/Private/LandscapeEditorObject.cpp b/Engine/Source/Editor/LandscapeEditor/Private/LandscapeEditorObject.cpp
--- a/Engine/Source/Editor/LandscapeEditor/Private/LandscapeEditorObject.cpp
+++ b/Engine/Source/Editor/LandscapeEditor/Private/LandscapeEditorObject.cpp
@@ -28,6 +28,9 @@ ULandscapeEditorObject::ULandscapeEditorObject(const FObjectInitializer& ObjectI
 	, bUseFlattenTarget(false)
 	, FlattenTarget(0)
 	, bShowFlattenTargetPreview(true)
+	
+	, TerraceInterval(1.0f)
+	, TerraceSmooth(0.0001f)
 
 	, RampWidth(2000)
 	, RampSideFalloff(0.4f)
diff --git a/Engine/Source/Runtime/Landscape/Classes/Landscape.h b/Engine/Source/Runtime/Landscape/Classes/Landscape.h
--- a/Engine/Source/Runtime/Landscape/Classes/Landscape.h
+++ b/Engine/Source/Runtime/Landscape/Classes/Landscape.h
@@ -268,6 +268,7 @@ public:
 	LANDSCAPE_API void DeleteLayer(int32 InLayerIndex);
 	LANDSCAPE_API void DeleteLayers();
 	LANDSCAPE_API void SetEditingLayer(const FGuid& InLayerGuid = FGuid());
+	LANDSCAPE_API void SetGrassUpdateEnabled(bool bInGrassUpdateEnabled);
 	LANDSCAPE_API const FGuid& GetEditingLayer() const;
 	LANDSCAPE_API bool IsMaxLayersReached() const;
 	LANDSCAPE_API void ShowOnlySelectedLayer(int32 InLayerIndex);
@@ -294,6 +295,7 @@ public:
 	void ReleaseLayersRenderingResource();
 	
 	LANDSCAPE_API void ToggleCanHaveLayersContent();
+	LANDSCAPE_API void ForceUpdateLayersContent();
 
 private:
 	void TickLayers(float DeltaTime, ELevelTick TickType, FActorTickFunction& ThisTickFunction);
@@ -368,6 +370,9 @@ public:
 	/** Current Editing Landscape Layer*/
 	FGuid EditingLayer;
 
+	/** Used to temporarily disable Grass Update in Editor */
+	bool bGrassUpdateEnabled;
+
 	UPROPERTY(TextExportTransient)
 	TArray<FLandscapeLayer> LandscapeLayers;
 
diff --git a/Engine/Source/Runtime/Landscape/Classes/LandscapeInfo.h b/Engine/Source/Runtime/Landscape/Classes/LandscapeInfo.h
--- a/Engine/Source/Runtime/Landscape/Classes/LandscapeInfo.h
+++ b/Engine/Source/Runtime/Landscape/Classes/LandscapeInfo.h
@@ -251,5 +251,8 @@ public:
 
 	/** Postpones landscape textures baking, usually used during landscape painting to avoid hitches */
 	LANDSCAPE_API void PostponeTextureBaking();
+
+	/** Will tell if the landscape actor can have some content related to the layer system */
+	LANDSCAPE_API bool CanHaveLayersContent() const;
 #endif
 };
diff --git a/Engine/Source/Runtime/Landscape/Private/Landscape.cpp b/Engine/Source/Runtime/Landscape/Private/Landscape.cpp
--- a/Engine/Source/Runtime/Landscape/Private/Landscape.cpp
+++ b/Engine/Source/Runtime/Landscape/Private/Landscape.cpp
@@ -1091,6 +1091,7 @@ ALandscape::ALandscape(const FObjectInitializer& ObjectInitializer)
 	WeightmapScratchPackLayerTextureResource = nullptr;
 	bLandscapeLayersAreInitialized = false;
 	LandscapeEdMode = nullptr;	
+	bGrassUpdateEnabled = true;
 #endif // WITH_EDITORONLY_DATA
 }
 
diff --git a/Engine/Source/Runtime/Landscape/Private/LandscapeEdit.cpp b/Engine/Source/Runtime/Landscape/Private/LandscapeEdit.cpp
--- a/Engine/Source/Runtime/Landscape/Private/LandscapeEdit.cpp
+++ b/Engine/Source/Runtime/Landscape/Private/LandscapeEdit.cpp
@@ -3836,6 +3836,15 @@ void ULandscapeInfo::PostponeTextureBaking()
 	});
 }
 
+bool ULandscapeInfo::CanHaveLayersContent() const
+{
+	if (ALandscape* Landscape = LandscapeActor.Get())
+	{
+		return Landscape->CanHaveLayersContent();
+	}
+	return false;
+}
+
 namespace
 {
 	inline float AdjustStaticLightingResolution(float StaticLightingResolution, int32 NumSubsections, int32 SubsectionSizeQuads, int32 ComponentSizeQuads)
diff --git a/Engine/Source/Runtime/Landscape/Private/LandscapeEditInterface.cpp b/Engine/Source/Runtime/Landscape/Private/LandscapeEditInterface.cpp
--- a/Engine/Source/Runtime/Landscape/Private/LandscapeEditInterface.cpp
+++ b/Engine/Source/Runtime/Landscape/Private/LandscapeEditInterface.cpp
@@ -5517,8 +5517,7 @@ const ALandscape* FLandscapeEditDataInterface::GetTargetLandscape() const
 
 bool FLandscapeEditDataInterface::CanHaveLandscapeLayersContent() const
 {
-	const ALandscape* Landscape = GetTargetLandscape();
-	return Landscape ? Landscape->CanHaveLayersContent() : false;
+	return LandscapeInfo ? LandscapeInfo->CanHaveLayersContent() : false;
 }
 
 bool FLandscapeEditDataInterface::HasLandscapeLayersContent() const
diff --git a/Engine/Source/Runtime/Landscape/Private/LandscapeEditLayers.cpp b/Engine/Source/Runtime/Landscape/Private/LandscapeEditLayers.cpp
--- a/Engine/Source/Runtime/Landscape/Private/LandscapeEditLayers.cpp
+++ b/Engine/Source/Runtime/Landscape/Private/LandscapeEditLayers.cpp
@@ -4441,6 +4441,11 @@ void ALandscape::InitializeLayers()
 }
 
 void ALandscape::OnPreSave()
+{
+	ForceUpdateLayersContent();
+}
+
+void ALandscape::ForceUpdateLayersContent()
 {
 	const bool bWaitForStreaming = true;
 	UpdateLayersContent(bWaitForStreaming);
@@ -5152,6 +5157,13 @@ void ALandscape::SetEditingLayer(const FGuid& InLayerGuid)
 	EditingLayer = InLayerGuid;
 }
 
+void ALandscape::SetGrassUpdateEnabled(bool bInGrassUpdateEnabled)
+{
+#ifdef WITH_EDITORONLY_DATA
+	bGrassUpdateEnabled = bInGrassUpdateEnabled;
+#endif
+}
+
 const FGuid& ALandscape::GetEditingLayer() const
 {
 	return EditingLayer;
diff --git a/Engine/Source/Runtime/Landscape/Private/LandscapeGrass.cpp b/Engine/Source/Runtime/Landscape/Private/LandscapeGrass.cpp
--- a/Engine/Source/Runtime/Landscape/Private/LandscapeGrass.cpp
+++ b/Engine/Source/Runtime/Landscape/Private/LandscapeGrass.cpp
@@ -1351,7 +1351,11 @@ void ALandscapeProxy::TickGrass()
 
 	if (ALandscape* Landscape = GetLandscapeActor())
 	{
-		if (!Landscape->IsUpToDate())
+		if (!Landscape->IsUpToDate() 
+#if WITH_EDITORONLY_DATA
+			|| !Landscape->bGrassUpdateEnabled
+#endif
+			)
 		{
 			return;
 		}
diff --git a/Engine/Source/Runtime/Landscape/Public/LandscapeEdit.h b/Engine/Source/Runtime/Landscape/Public/LandscapeEdit.h
--- a/Engine/Source/Runtime/Landscape/Public/LandscapeEdit.h
+++ b/Engine/Source/Runtime/Landscape/Public/LandscapeEdit.h
@@ -343,6 +343,11 @@ struct FHeightmapAccessor
 		LandscapeEdit->GetHeightDataFast(X1, Y1, X2, Y2, Data);
 	}
 
+	void GetDataFast(int32 X1, int32 Y1, int32 X2, int32 Y2, uint16* Data)
+	{
+		LandscapeEdit->GetHeightDataFast(X1, Y1, X2, Y2, Data, 0);
+	}
+
 	void SetData(int32 X1, int32 Y1, int32 X2, int32 Y2, const uint16* Data, ELandscapeLayerPaintingRestriction PaintingRestriction = ELandscapeLayerPaintingRestriction::None)
 	{
 		TSet<ULandscapeComponent*> Components;
@@ -508,6 +513,11 @@ struct FAlphamapAccessor
 		LandscapeEdit.GetWeightDataFast(LayerInfo, X1, Y1, X2, Y2, Data);
 	}
 
+	void GetDataFast(int32 X1, int32 Y1, int32 X2, int32 Y2, uint8* Data)
+	{
+		LandscapeEdit.GetWeightDataFast(LayerInfo, X1, Y1, X2, Y2, Data, 0);
+	}
+
 	void SetData(int32 X1, int32 Y1, int32 X2, int32 Y2, const uint8* Data, ELandscapeLayerPaintingRestriction PaintingRestriction)
 	{
 		TSet<ULandscapeComponent*> Components;
