Landscape Layer System: EdMode selection minor improvements

#rnx

diff --git a/Engine/Source/Editor/LandscapeEditor/Private/LandscapeEditorDetailCustomization_Layers.cpp b/Engine/Source/Editor/LandscapeEditor/Private/LandscapeEditorDetailCustomization_Layers.cpp
--- a/Engine/Source/Editor/LandscapeEditor/Private/LandscapeEditorDetailCustomization_Layers.cpp
+++ b/Engine/Source/Editor/LandscapeEditor/Private/LandscapeEditorDetailCustomization_Layers.cpp
@@ -355,6 +355,7 @@ void FLandscapeEditorCustomNodeBuilder_Layers::SetLayerName(const FText& InText,
 	{
 		const FScopedTransaction Transaction(LOCTEXT("Landscape_Layers_Rename", "Rename Layer"));
 		LandscapeEdMode->SetLayerName(InLayerIndex, *InText.ToString());
+		OnLayerSelectionChanged(InLayerIndex);
 	}
 }
 
@@ -466,6 +467,7 @@ void FLandscapeEditorCustomNodeBuilder_Layers::SetLandscapeSplinesReservedLayer(
 				Landscape->SetLandscapeSplinesReservedLayer(InLayerIndex);
 				LandscapeEdMode->RefreshDetailPanel();
 				LandscapeEdMode->AutoUpdateDirtyLandscapeSplines();
+				OnLayerSelectionChanged(InLayerIndex);
 			}
 		}
 	}
@@ -493,6 +495,7 @@ void FLandscapeEditorCustomNodeBuilder_Layers::ClearLayer(int32 InLayerIndex)
 			{
 				const FScopedTransaction Transaction(LOCTEXT("Landscape_Layers_Clean", "Clear Layer"));
 				Landscape->ClearLayer(InLayerIndex);
+				OnLayerSelectionChanged(InLayerIndex);
 			}
 		}
 	}
@@ -528,6 +531,7 @@ void FLandscapeEditorCustomNodeBuilder_Layers::ShowOnlySelectedLayer(int32 InLay
 	{
 		const FScopedTransaction Transaction(LOCTEXT("ShowOnlySelectedLayer", "Show Only Selected Layer"));
 		Landscape->ShowOnlySelectedLayer(InLayerIndex);
+		OnLayerSelectionChanged(InLayerIndex);
 	}
 }
 
@@ -544,22 +548,23 @@ void FLandscapeEditorCustomNodeBuilder_Layers::ShowAllLayers()
 
 void FLandscapeEditorCustomNodeBuilder_Layers::CreateLayer()
 {
-	FEdModeLandscape* LandscapeEdMode = GetEditorMode();
-	ALandscape* Landscape = LandscapeEdMode ? LandscapeEdMode->GetLandscape() : nullptr;
-	if (Landscape)
+FEdModeLandscape* LandscapeEdMode = GetEditorMode();
+ALandscape* Landscape = LandscapeEdMode ? LandscapeEdMode->GetLandscape() : nullptr;
+if (Landscape)
+{
 	{
-		{
-			const FScopedTransaction Transaction(LOCTEXT("Landscape_Layers_Create", "Create Layer"));
-			Landscape->CreateLayer();
-		}
-		LandscapeEdMode->RefreshDetailPanel();
+		const FScopedTransaction Transaction(LOCTEXT("Landscape_Layers_Create", "Create Layer"));
+		Landscape->CreateLayer();
+		OnLayerSelectionChanged(Landscape->GetLayerCount() - 1);
 	}
+	LandscapeEdMode->RefreshDetailPanel();
+}
 }
 
 void FLandscapeEditorCustomNodeBuilder_Layers::OnLayerSelectionChanged(int32 InLayerIndex)
 {
 	FEdModeLandscape* LandscapeEdMode = GetEditorMode();
-	if (LandscapeEdMode)
+	if (LandscapeEdMode && LandscapeEdMode->GetCurrentLayerIndex() != InLayerIndex)
 	{
 		FScopedTransaction Transaction(LOCTEXT("Landscape_Layers_SetCurrentLayer", "Set Current Layer"));
 		LandscapeEdMode->SetCurrentLayer(InLayerIndex);
@@ -603,6 +608,10 @@ FReply FLandscapeEditorCustomNodeBuilder_Layers::OnToggleVisibility(int32 InLaye
 	{
 		const FScopedTransaction Transaction(LOCTEXT("Landscape_Layers_SetVisibility", "Set Layer Visibility"));
 		LandscapeEdMode->SetLayerVisibility(!LandscapeEdMode->IsLayerVisible(InLayerIndex), InLayerIndex);
+		if (LandscapeEdMode->IsLayerVisible(InLayerIndex))
+		{
+			OnLayerSelectionChanged(InLayerIndex);
+		}
 	}
 	return FReply::Handled();
 }
diff --git a/Engine/Source/Editor/LandscapeEditor/Private/LandscapeEditorDetailCustomization_TargetLayers.cpp b/Engine/Source/Editor/LandscapeEditor/Private/LandscapeEditorDetailCustomization_TargetLayers.cpp
--- a/Engine/Source/Editor/LandscapeEditor/Private/LandscapeEditorDetailCustomization_TargetLayers.cpp
+++ b/Engine/Source/Editor/LandscapeEditor/Private/LandscapeEditorDetailCustomization_TargetLayers.cpp
@@ -1201,15 +1201,13 @@ void FLandscapeEditorCustomNodeBuilder_TargetLayers::OnTargetLayerSetObject(cons
 					Target->LandscapeInfo->CreateLayerEditorSettingsFor(SelectedLayerInfo);
 				}
 			}
-
+						
 			FEdModeLandscape* LandscapeEdMode = GetEditorMode();
 			if (LandscapeEdMode)
 			{
-				if (LandscapeEdMode->CurrentToolTarget.LayerName == Target->LayerName
-					&& LandscapeEdMode->CurrentToolTarget.LayerInfo == Target->LayerInfoObj)
-				{
-					LandscapeEdMode->CurrentToolTarget.LayerInfo = SelectedLayerInfo;
-				}
+				LandscapeEdMode->CurrentToolTarget.LayerName = Target->LayerName;
+				LandscapeEdMode->CurrentToolTarget.TargetType = Target->TargetType;
+				LandscapeEdMode->CurrentToolTarget.LayerInfo = SelectedLayerInfo;
 				LandscapeEdMode->UpdateTargetList();
 			}
 
@@ -1545,15 +1543,15 @@ FReply SLandscapeEditorSelectableBorder::OnMouseButtonUp(const FGeometry& MyGeom
 		}
 		else if (MouseEvent.GetEffectingButton() == EKeys::RightMouseButton &&
 			OnContextMenuOpening.IsBound())
-		{
-			TSharedPtr<SWidget> Content = OnContextMenuOpening.Execute();
-			if (Content.IsValid())
 			{
-				FWidgetPath WidgetPath = MouseEvent.GetEventPath() != nullptr ? *MouseEvent.GetEventPath() : FWidgetPath();
-
-				FSlateApplication::Get().PushMenu(SharedThis(this), WidgetPath, Content.ToSharedRef(), MouseEvent.GetScreenSpacePosition(), FPopupTransitionEffect(FPopupTransitionEffect::ContextMenu));
-			}
+				TSharedPtr<SWidget> Content = OnContextMenuOpening.Execute();
+				if (Content.IsValid())
+				{
+					FWidgetPath WidgetPath = MouseEvent.GetEventPath() != nullptr ? *MouseEvent.GetEventPath() : FWidgetPath();
 
+					FSlateApplication::Get().PushMenu(SharedThis(this), WidgetPath, Content.ToSharedRef(), MouseEvent.GetScreenSpacePosition(), FPopupTransitionEffect(FPopupTransitionEffect::ContextMenu));
+				}
+			
 			return FReply::Handled().ReleaseMouseCapture();
 		}
 	}
