Fix Compile EngineTestEditor Mac

#rnx

diff --git a/Engine/Source/Editor/LandscapeEditor/Private/LandscapeEdModePaintTools.cpp b/Engine/Source/Editor/LandscapeEditor/Private/LandscapeEdModePaintTools.cpp
--- a/Engine/Source/Editor/LandscapeEditor/Private/LandscapeEdModePaintTools.cpp
+++ b/Engine/Source/Editor/LandscapeEditor/Private/LandscapeEdModePaintTools.cpp
@@ -777,7 +777,7 @@ public:
 			{
 				// Can't rely on cache in this mode
 				FScopedSetLandscapeEditingLayer Scope(Landscape, FGuid());
-				ToolTarget::CacheClass::DataType P00, P10, P01, P11;
+				typename ToolTarget::CacheClass::DataType P00, P10, P01, P11;
 				this->Cache.DataAccess.GetDataFast(FlattenHeightX,     FlattenHeightY,     FlattenHeightX,     FlattenHeightY,     &P00);
 				this->Cache.DataAccess.GetDataFast(FlattenHeightX + 1, FlattenHeightY,     FlattenHeightX + 1, FlattenHeightY,     &P10);
 				this->Cache.DataAccess.GetDataFast(FlattenHeightX,     FlattenHeightY + 1, FlattenHeightX,     FlattenHeightY + 1, &P01);
@@ -880,7 +880,7 @@ public:
 		if (bLayerSystemFlattenHeightMode)
 		{
 			EditingLayer = Landscape->GetEditingLayer();
-			LandscapeInfo->GetComponentsInRegion(X1, Y1, X2, Y2, AffectedComponents);
+			this->LandscapeInfo->GetComponentsInRegion(X1, Y1, X2, Y2, AffectedComponents);
 			TArray<bool> NewLayerVisibility;
 			for (int i = 0; i < Landscape->GetLayerCount(); ++i)
 			{
