- Landscape Layer System: Fix missing UpdateToolModes. Also UpdateToolModes now tries to set the CurrentTool in case it is no longer available after the update.

#rnx

diff --git a/Engine/Source/Editor/LandscapeEditor/Private/LandscapeEdMode.cpp b/Engine/Source/Editor/LandscapeEditor/Private/LandscapeEdMode.cpp
--- a/Engine/Source/Editor/LandscapeEditor/Private/LandscapeEdMode.cpp
+++ b/Engine/Source/Editor/LandscapeEditor/Private/LandscapeEdMode.cpp
@@ -290,6 +290,13 @@ namespace LandscapeTool
 /** Constructor */
 FEdModeLandscape::FEdModeLandscape()
 	: FEdMode()
+	, UISettings(nullptr)
+	, CurrentToolMode(nullptr)
+	, CurrentTool(nullptr)
+	, CurrentBrush(nullptr)
+	, GizmoBrush(nullptr)
+	, CurrentToolIndex(INDEX_NONE)
+	, CurrentBrushSetIndex(0)
 	, NewLandscapePreviewMode(ENewLandscapePreviewMode::None)
 	, DraggingEdge(ELandscapeEdge::None)
 	, DraggingEdge_Remainder(0)
@@ -309,12 +316,10 @@ FEdModeLandscape::FEdModeLandscape()
 	GColorMaskRegionMaterial = LandscapeTool::CreateMaterialInstance(LoadObject<UMaterialInstanceConstant>(nullptr, TEXT("/Engine/EditorLandscapeResources/ColorMaskBrushMaterial_MaskedRegion.ColorMaskBrushMaterial_MaskedRegion")));
 	GLandscapeBlackTexture   = LoadObject<UTexture2D>(nullptr, TEXT("/Engine/EngineResources/Black.Black"));
 	GLandscapeLayerUsageMaterial = LandscapeTool::CreateMaterialInstance(LoadObject<UMaterial>(nullptr, TEXT("/Engine/EditorLandscapeResources/LandscapeLayerUsageMaterial.LandscapeLayerUsageMaterial")));
-
-
+	
 	// Initialize modes
 	UpdateToolModes();
-	CurrentToolMode = nullptr;
-
+	
 	// Initialize tools.
 	InitializeTool_Paint();
 	InitializeTool_Smooth();
@@ -337,15 +342,11 @@ FEdModeLandscape::FEdModeLandscape()
 	InitializeTool_Mirror();
 	InitializeTool_BPCustom();
 
-	CurrentTool = nullptr;
-	CurrentToolIndex = INDEX_NONE;
-
 	// Initialize brushes
 	InitializeBrushes();
 
 	CurrentBrush = LandscapeBrushSets[0].Brushes[0];
-	CurrentBrushSetIndex = 0;
-
+	
 	CurrentToolTarget.LandscapeInfo = nullptr;
 	CurrentToolTarget.TargetType = ELandscapeToolTargetType::Heightmap;
 	CurrentToolTarget.LayerInfo = nullptr;
@@ -453,6 +454,12 @@ void FEdModeLandscape::UpdateToolModes()
 	{
 		ToolMode_Paint->ValidTools.Add(TEXT("BPCustom"));
 	}
+
+	// Since available tools might have changed try and reset the current tool
+	if (CurrentToolMode && CurrentToolIndex != INDEX_NONE)
+	{
+		SetCurrentTool(CurrentToolIndex);
+	}
 }
 
 bool FEdModeLandscape::UsesToolkits() const
@@ -469,6 +476,7 @@ TSharedRef<FUICommandList> FEdModeLandscape::GetUICommandList() const
 void FEdModeLandscape::OnCanHaveLayersContentChanged()
 {
 	RefreshDetailPanel();
+	UpdateToolModes();
 }
 
 ELandscapeToolTargetType::Type FEdModeLandscape::GetLandscapeToolTargetType() const
