- Landscape Layer System:
Log a Warning when a LandscapeStreamingProxy loses some Layer data on load because the actor layers changed (got removed without the streaming proxy being loaded)
Fix the FEdModeLandscape::CanEditCurrentTarget that was broken with the big flag change

#rnx

diff --git a/Engine/Source/Editor/LandscapeEditor/Private/LandscapeEdMode.cpp b/Engine/Source/Editor/LandscapeEditor/Private/LandscapeEdMode.cpp
--- a/Engine/Source/Editor/LandscapeEditor/Private/LandscapeEdMode.cpp
+++ b/Engine/Source/Editor/LandscapeEditor/Private/LandscapeEdMode.cpp
@@ -2541,10 +2541,14 @@ bool FEdModeLandscape::CanEditCurrentTarget(FText* Reason) const
 	}
 
 	// Landscape Layer Editing not available without a loaded Landscape Actor
-	if (CanHaveLandscapeLayersContent() && !CurrentToolTarget.LandscapeInfo->LandscapeActor.IsValid())
+	if (GetLandscape() == nullptr)
 	{
-		LocalReason = NSLOCTEXT("UnrealEd", "LandscapeActorNotLoaded", "Landscape actor is not loaded. It is needed to do layer editing.");
-		return false;
+		ALandscapeProxy* Proxy = CurrentToolTarget.LandscapeInfo->GetLandscapeProxy();
+		if (Proxy->HasLayersContent())
+		{
+			LocalReason = NSLOCTEXT("UnrealEd", "LandscapeActorNotLoaded", "Landscape actor is not loaded. It is needed to do layer editing.");
+			return false;
+		}
 	}
 
 	return true;
diff --git a/Engine/Source/Runtime/Landscape/Private/LandscapeEditLayers.cpp b/Engine/Source/Runtime/Landscape/Private/LandscapeEditLayers.cpp
--- a/Engine/Source/Runtime/Landscape/Private/LandscapeEditLayers.cpp
+++ b/Engine/Source/Runtime/Landscape/Private/LandscapeEditLayers.cpp
@@ -4595,11 +4595,21 @@ bool ALandscapeProxy::RemoveObsoleteLayers(const TSet<FGuid>& InExistingLayers)
 	{
 		if (!InExistingLayers.Contains(LayerGuid))
 		{
+			UE_LOG(LogLandscape, Warning, TEXT("Layer '%s' was removed from LandscapeProxy '%s' because it doesn't match any of the LandscapeActor Layers. Possible loss of data."), 
+				*LayerGuid.ToString(EGuidFormats::HexValuesInBraces), *GetPathName());
 			DeleteLayer(LayerGuid);
 			bModified = true;
 		}
 	}
 
+	if (bModified)
+	{
+		if (ALandscape* LandscapeActor = GetLandscapeActor())
+		{
+			LandscapeActor->RequestLayersContentUpdateForceAll();
+		}
+	}
+
 	return bModified;
 }
 
