Landscape BP Brushes Fixes
- Fixed bug where compiling BP Brush could change its order in the Brush stack
- Added option to add BP Brushes directly in the BP Brush stack of a Landscape Layer
#rnx

diff --git a/Engine/Source/Editor/LandscapeEditor/Private/LandscapeEdMode.cpp b/Engine/Source/Editor/LandscapeEditor/Private/LandscapeEdMode.cpp
--- a/Engine/Source/Editor/LandscapeEditor/Private/LandscapeEdMode.cpp
+++ b/Engine/Source/Editor/LandscapeEditor/Private/LandscapeEdMode.cpp
@@ -4716,7 +4716,11 @@ void FEdModeLandscape::UpdateBrushList()
 	BrushList.Empty();
 	for (TObjectIterator<ALandscapeBlueprintCustomBrush> BrushIt(RF_Transient|RF_ClassDefaultObject|RF_ArchetypeObject, true, EInternalObjectFlags::PendingKill); BrushIt; ++BrushIt)
 	{
-		BrushList.Add(*BrushIt);
+		ALandscapeBlueprintCustomBrush* Brush = *BrushIt;
+		if (Brush->GetTypedOuter<UPackage>() != GetTransientPackage())
+		{
+			BrushList.Add(Brush);
+		}
 	}
 }
 
@@ -4729,11 +4733,13 @@ void FEdModeLandscape::OnLevelActorAdded(AActor* InActor)
 	}
 
 	ALandscapeBlueprintCustomBrush* Brush = Cast<ALandscapeBlueprintCustomBrush>(InActor);
-
-	if (Brush != nullptr && Brush->GetTypedOuter<UPackage>() != GetTransientPackage())
+	if (Brush && Brush->GetTypedOuter<UPackage>() != GetTransientPackage())
 	{
-		BrushList.Add(Brush);
-		AddBrushToCurrentLayer(Brush);
+		if (!GIsReinstancing)
+		{
+			AddBrushToCurrentLayer(Brush);
+		}
+		UpdateBrushList();
 		RefreshDetailPanel();
 	}
 }
@@ -4746,11 +4752,9 @@ void FEdModeLandscape::OnLevelActorRemoved(AActor* InActor)
 	}
 
 	ALandscapeBlueprintCustomBrush* Brush = Cast<ALandscapeBlueprintCustomBrush>(InActor);
-
-	if (Brush != nullptr && Brush->GetTypedOuter<UPackage>() != GetTransientPackage())
+	if (Brush && Brush->GetTypedOuter<UPackage>() != GetTransientPackage())
 	{
-		BrushList.Remove(Brush);
-		RemoveBrushFromCurrentLayer(Brush);
+		UpdateBrushList();
 		RefreshDetailPanel();
 	}
 }
diff --git a/Engine/Source/Editor/LandscapeEditor/Private/LandscapeEditorDetailCustomization_Layers.cpp b/Engine/Source/Editor/LandscapeEditor/Private/LandscapeEditorDetailCustomization_Layers.cpp
--- a/Engine/Source/Editor/LandscapeEditor/Private/LandscapeEditorDetailCustomization_Layers.cpp
+++ b/Engine/Source/Editor/LandscapeEditor/Private/LandscapeEditorDetailCustomization_Layers.cpp
@@ -352,12 +352,9 @@ void FLandscapeEditorCustomNodeBuilder_Layers::FillAddBrushMenu(FMenuBuilder& Me
 {
 	for (ALandscapeBlueprintCustomBrush* Brush : Brushes)
 	{
-		if (Brush->GetOwningLandscape() == nullptr)
-		{
-			TSharedRef<FLandscapeEditorCustomNodeBuilder_Layers> SharedThis = AsShared();
-			FUIAction AddAction = FUIAction(FExecuteAction::CreateLambda([SharedThis, Brush]() { SharedThis->AddBrushToCurrentLayer(Brush); }));
-			MenuBuilder.AddMenuEntry(FText::FromString(Brush->GetActorLabel()), FText(), FSlateIcon(), AddAction);
-		}
+		TSharedRef<FLandscapeEditorCustomNodeBuilder_Layers> SharedThis = AsShared();
+		FUIAction AddAction = FUIAction(FExecuteAction::CreateLambda([SharedThis, Brush]() { SharedThis->AddBrushToCurrentLayer(Brush); }));
+		MenuBuilder.AddMenuEntry(FText::FromString(Brush->GetActorLabel()), FText(), FSlateIcon(), AddAction);
 	}
 }
 
@@ -448,14 +445,15 @@ TSharedPtr<SWidget> FLandscapeEditorCustomNodeBuilder_Layers::OnLayerContextMenu
 		MenuBuilder.EndSection();
 
 		const TArray<ALandscapeBlueprintCustomBrush*>& Brushes = LandscapeEdMode->GetBrushList();
-		if (Brushes.ContainsByPredicate([](ALandscapeBlueprintCustomBrush* Brush) { return Brush->GetOwningLandscape() == nullptr; }))
+		TArray<ALandscapeBlueprintCustomBrush*> FilteredBrushes = Brushes.FilterByPredicate([](ALandscapeBlueprintCustomBrush* Brush) { return Brush->GetOwningLandscape() == nullptr; });
+		if (FilteredBrushes.Num())
 		{
 			MenuBuilder.BeginSection("LandscapeEditorBrushActions", LOCTEXT("LandscapeEditorBrushActions.Heading", "Brushes"));
 			{
 				MenuBuilder.AddSubMenu(
 					LOCTEXT("LandscaeEditorBrushAddSubMenu", "Add"),
 					LOCTEXT("LandscaeEditorBrushAddSubMenuToolTip", "Add brush to current layer"),
-					FNewMenuDelegate::CreateSP(this, &FLandscapeEditorCustomNodeBuilder_Layers::FillAddBrushMenu, Brushes),
+					FNewMenuDelegate::CreateSP(this, &FLandscapeEditorCustomNodeBuilder_Layers::FillAddBrushMenu, FilteredBrushes),
 					false,
 					FSlateIcon()
 				);
diff --git a/Engine/Source/Editor/LandscapeEditor/Private/LandscapeEditorDetailCustomization_LayersBrushStack.cpp b/Engine/Source/Editor/LandscapeEditor/Private/LandscapeEditorDetailCustomization_LayersBrushStack.cpp
--- a/Engine/Source/Editor/LandscapeEditor/Private/LandscapeEditorDetailCustomization_LayersBrushStack.cpp
+++ b/Engine/Source/Editor/LandscapeEditor/Private/LandscapeEditorDetailCustomization_LayersBrushStack.cpp
@@ -207,12 +207,6 @@ TSharedPtr<SWidget> FLandscapeEditorCustomNodeBuilder_LayersBrushStack::Generate
 }
 END_SLATE_FUNCTION_BUILD_OPTIMIZATION
 
-void FLandscapeEditorCustomNodeBuilder_LayersBrushStack::RemoveBrush(ALandscapeBlueprintCustomBrush* Brush)
-{
-	const FScopedTransaction Transaction(LOCTEXT("LandscapeBrushRemoveTransaction", "Remove Brush"));
-	GetEditorMode()->RemoveBrushFromCurrentLayer(Brush);
-}
-
 TSharedPtr<SWidget> FLandscapeEditorCustomNodeBuilder_LayersBrushStack::OnBrushContextMenuOpening(int32 InBrushIndex)
 {
 	if (ALandscapeBlueprintCustomBrush* CurrentBrush = GetBrush(InBrushIndex))
@@ -222,11 +216,27 @@ TSharedPtr<SWidget> FLandscapeEditorCustomNodeBuilder_LayersBrushStack::OnBrushC
 		{
 			TSharedRef<FLandscapeEditorCustomNodeBuilder_LayersBrushStack> SharedThis = AsShared();
 
+			// Add Brush
+			const TArray<ALandscapeBlueprintCustomBrush*>& Brushes = GetEditorMode()->GetBrushList();
+			TArray<ALandscapeBlueprintCustomBrush*> FilteredBrushes = Brushes.FilterByPredicate([](ALandscapeBlueprintCustomBrush* Brush) { return Brush->GetOwningLandscape() == nullptr; });
+			if (FilteredBrushes.Num())
+			{
+				MenuBuilder.AddSubMenu(
+					LOCTEXT("LandscaeEditorBrushAddSubMenu", "Add"),
+					LOCTEXT("LandscaeEditorBrushAddSubMenuToolTip", "Add brush to current layer"),
+					FNewMenuDelegate::CreateSP(this, &FLandscapeEditorCustomNodeBuilder_LayersBrushStack::FillAddBrushMenu, FilteredBrushes),
+					false,
+					FSlateIcon()
+				);
+			}
+
 			// Remove Brush
-			FUIAction RemoveAction = FUIAction(FExecuteAction::CreateLambda([SharedThis, CurrentBrush] { SharedThis->RemoveBrush(CurrentBrush); }));
+			FUIAction RemoveAction = FUIAction(FExecuteAction::CreateLambda([SharedThis, CurrentBrush]
+			{ 
+				const FScopedTransaction Transaction(LOCTEXT("LandscapeBrushRemoveTransaction", "Remove Brush"));
+				GetEditorMode()->RemoveBrushFromCurrentLayer(CurrentBrush);
+			}));
 			MenuBuilder.AddMenuEntry(LOCTEXT("LandscapeBrushRemove", "Remove"), LOCTEXT("LandscapeBrushRemoveToolTip", "Remove brush from current layer"), FSlateIcon(), RemoveAction);
-
-		
 		}
 		MenuBuilder.EndSection();
 		return MenuBuilder.MakeWidget();
@@ -235,6 +245,20 @@ TSharedPtr<SWidget> FLandscapeEditorCustomNodeBuilder_LayersBrushStack::OnBrushC
 	return nullptr;
 }
 
+void FLandscapeEditorCustomNodeBuilder_LayersBrushStack::FillAddBrushMenu(FMenuBuilder& MenuBuilder, TArray<ALandscapeBlueprintCustomBrush*> Brushes)
+{
+	TSharedRef<FLandscapeEditorCustomNodeBuilder_LayersBrushStack> SharedThis = AsShared();
+	for (ALandscapeBlueprintCustomBrush* Brush : Brushes)
+	{
+		FUIAction AddAction = FUIAction(FExecuteAction::CreateLambda([SharedThis, Brush]()
+		{
+			const FScopedTransaction Transaction(LOCTEXT("LandscapeBrushAddToCurrentLayerTransaction", "Add brush to current layer"));
+			GetEditorMode()->AddBrushToCurrentLayer(Brush);
+		}));
+		MenuBuilder.AddMenuEntry(FText::FromString(Brush->GetActorLabel()), FText(), FSlateIcon(), AddAction);
+	}
+}
+
 FReply FLandscapeEditorCustomNodeBuilder_LayersBrushStack::OnToggleAffectsHeightmap(int32 InBrushIndex)
 {
 	if (ALandscapeBlueprintCustomBrush* Brush = GetBrush(InBrushIndex))
diff --git a/Engine/Source/Editor/LandscapeEditor/Private/LandscapeEditorDetailCustomization_LayersBrushStack.h b/Engine/Source/Editor/LandscapeEditor/Private/LandscapeEditorDetailCustomization_LayersBrushStack.h
--- a/Engine/Source/Editor/LandscapeEditor/Private/LandscapeEditorDetailCustomization_LayersBrushStack.h
+++ b/Engine/Source/Editor/LandscapeEditor/Private/LandscapeEditorDetailCustomization_LayersBrushStack.h
@@ -59,8 +59,8 @@ protected:
 	TSharedPtr<SWidget> GenerateRow(int32 InBrushIndex);
 	TSharedPtr<SWidget> OnBrushContextMenuOpening(int32 InBrushIndex);
 
-	void RemoveBrush(ALandscapeBlueprintCustomBrush* Brush);
-	
+	void FillAddBrushMenu(FMenuBuilder& MenuBuilder, TArray<ALandscapeBlueprintCustomBrush*> Brushes);
+
 	FReply OnToggleVisibility(int32 InBrushIndex);
 	FReply OnToggleAffectsHeightmap(int32 InBrushIndex);
 	FReply OnToggleAffectsWeightmap(int32 InBrushIndex);
