-Landscape Layer System: Fix Move Tool when LandscapeProxy didn't exist in destination Level

#rnx

diff --git a/Engine/Source/Editor/LandscapeEditor/Private/LandscapeEdMode.cpp b/Engine/Source/Editor/LandscapeEditor/Private/LandscapeEdMode.cpp
--- a/Engine/Source/Editor/LandscapeEditor/Private/LandscapeEdMode.cpp
+++ b/Engine/Source/Editor/LandscapeEditor/Private/LandscapeEdMode.cpp
@@ -191,6 +191,7 @@ void ALandscape::SplitHeightmap(ULandscapeComponent* Comp, ALandscapeProxy* Targ
 	GDisableAutomaticTextureMaterialUpdateDependencies = false;
 
 #if WITH_EDITORONLY_DATA
+	check(Comp->GetLandscapeProxy()->HasLayersContent() == DstProxy->CanHaveLayersContent());
 	if (Comp->GetLandscapeProxy()->HasLayersContent() && DstProxy->CanHaveLayersContent())
 	{
 		FLandscapeLayersTexture2DCPUReadBackResource* NewCPUReadBackResource = new FLandscapeLayersTexture2DCPUReadBackResource(NewHeightmapTexture->Source.GetSizeX(), NewHeightmapTexture->Source.GetSizeY(), NewHeightmapTexture->GetPixelFormat(), NewHeightmapTexture->Source.GetNumMips());
diff --git a/Engine/Source/Editor/LandscapeEditor/Private/LandscapeEdModeComponentTools.cpp b/Engine/Source/Editor/LandscapeEditor/Private/LandscapeEdModeComponentTools.cpp
--- a/Engine/Source/Editor/LandscapeEditor/Private/LandscapeEdModeComponentTools.cpp
+++ b/Engine/Source/Editor/LandscapeEditor/Private/LandscapeEdModeComponentTools.cpp
@@ -524,7 +524,8 @@ public:
 					LandscapeProxy = World->SpawnActor<ALandscapeStreamingProxy>();
 					// copy shared properties to this new proxy
 					LandscapeProxy->GetSharedProperties(Landscape);
-
+					LandscapeProxy->CreateLandscapeInfo();
+					
 					// set proxy location
 					// by default first component location
 					ULandscapeComponent* FirstComponent = *TargetSelectedComponents.CreateConstIterator();
@@ -611,10 +612,11 @@ public:
 					Component->MobileWeightmapTextures.Reset();
 					
 					Component->UpdateMaterialInstances();
-
+										
 					FFormatNamedArguments Args;
 					Args.Add(TEXT("ComponentName"), FText::FromString(Component->GetName()));
 				}
+				LandscapeProxy->UpdateCachedHasLayersContent();
 
 				for (ULandscapeHeightfieldCollisionComponent* Component : TargetSelectedCollisionComponents)
 				{
diff --git a/Engine/Source/Runtime/Landscape/Classes/LandscapeProxy.h b/Engine/Source/Runtime/Landscape/Classes/LandscapeProxy.h
--- a/Engine/Source/Runtime/Landscape/Classes/LandscapeProxy.h
+++ b/Engine/Source/Runtime/Landscape/Classes/LandscapeProxy.h
@@ -986,7 +986,7 @@ public:
 	/** Will tell if the landscape proxy can have some content related to the layer system */
 	LANDSCAPE_API bool CanHaveLayersContent() const;
 
-	void UpdateCachedHasLayersContent(bool InCheckComponentDataIntegrity = false);
+	LANDSCAPE_API void UpdateCachedHasLayersContent(bool InCheckComponentDataIntegrity = false);
 
 protected:
 	friend class ALandscape;
