Landscape Spline Ctrl+Click Improvement (Add Control Point)

#rnx

diff --git a/Engine/Source/Editor/LandscapeEditor/Private/LandscapeEdModeSplineTools.cpp b/Engine/Source/Editor/LandscapeEditor/Private/LandscapeEdModeSplineTools.cpp
--- a/Engine/Source/Editor/LandscapeEditor/Private/LandscapeEdModeSplineTools.cpp
+++ b/Engine/Source/Editor/LandscapeEditor/Private/LandscapeEdModeSplineTools.cpp
@@ -1028,17 +1028,47 @@ public:
 		if (ViewportClient->IsCtrlPressed())
 		{
 			LandscapeInfo = InTarget.LandscapeInfo.Get();
-			ALandscapeProxy* Landscape = LandscapeInfo->GetCurrentLevelLandscapeProxy(true);
-			if (!Landscape)
-			{
-				return false;
-			}
+			ALandscapeProxy* Landscape = nullptr;
 
+			// If we have a selection use the landscape of the selected spline
 			ULandscapeSplinesComponent* SplinesComponent = nullptr;
 			if (SelectedSplineControlPoints.Num() > 0)
 			{
 				ULandscapeSplineControlPoint* FirstPoint = *SelectedSplineControlPoints.CreateConstIterator();
 				SplinesComponent = FirstPoint->GetOuterULandscapeSplinesComponent();
+			
+				if (SplinesComponent)
+				{
+					Landscape = SplinesComponent->GetTypedOuter<ALandscapeProxy>();
+				}
+			}
+					
+			// Hit Test
+			if (!Landscape)
+			{
+				HHitProxy* HitProxy = ViewportClient->Viewport->GetHitProxy(ViewportClient->Viewport->GetMouseX(), ViewportClient->Viewport->GetMouseY());
+				if (HitProxy->IsA(HActor::StaticGetType()))
+				{
+					HActor* ActorProxy = (HActor*)HitProxy;
+					if (ALandscapeProxy* Proxy = Cast<ALandscapeProxy>(ActorProxy->Actor))
+					{
+						Landscape = Proxy;
+					}
+				}
+
+				
+			}
+
+			// Default to Current level Landscape
+			if (!Landscape)
+			{
+				Landscape = LandscapeInfo->GetCurrentLevelLandscapeProxy(true);
+			}
+
+			// No Landscape found
+			if (!Landscape)
+			{
+				return false;
 			}
 
 			if (!SplinesComponent)
