Landscape Layer System: Clear Option Sub menu (Height, Weight per Layer Target, All)

#rnx

diff --git a/Engine/Source/Editor/LandscapeEditor/Private/LandscapeEdMode.cpp b/Engine/Source/Editor/LandscapeEditor/Private/LandscapeEdMode.cpp
--- a/Engine/Source/Editor/LandscapeEditor/Private/LandscapeEdMode.cpp
+++ b/Engine/Source/Editor/LandscapeEditor/Private/LandscapeEdMode.cpp
@@ -309,6 +309,7 @@ FEdModeLandscape::FEdModeLandscape()
 	, ToolActiveViewport(nullptr)
 	, bIsPaintingInVR(false)
 	, InteractorPainting( nullptr )
+	, bNeedsUpdateShownLayerList(false)
 {
 	GLayerDebugColorMaterial = LandscapeTool::CreateMaterialInstance(LoadObject<UMaterial>(nullptr, TEXT("/Engine/EditorLandscapeResources/LayerVisMaterial.LayerVisMaterial")));
 	GSelectionColorMaterial  = LandscapeTool::CreateMaterialInstance(LoadObject<UMaterialInstanceConstant>(nullptr, TEXT("/Engine/EditorLandscapeResources/SelectBrushMaterial_Selected.SelectBrushMaterial_Selected")));
@@ -480,6 +481,14 @@ void FEdModeLandscape::OnCanHaveLayersContentChanged()
 	UpdateToolModes();
 }
 
+void FEdModeLandscape::PostUpdateLayerContent()
+{
+	if (bNeedsUpdateShownLayerList)
+	{
+		UpdateShownLayerList();
+	}
+}
+
 ELandscapeToolTargetType::Type FEdModeLandscape::GetLandscapeToolTargetType() const
 {
 	if (CurrentToolMode)
@@ -2800,8 +2809,20 @@ void FEdModeLandscape::OnLandscapeMaterialChangedDelegate()
 	UpdateShownLayerList();
 }
 
+void FEdModeLandscape::RequestUpdateShownLayerList()
+{
+	bNeedsUpdateShownLayerList = true;
+
+	if (CurrentToolTarget.LandscapeInfo.IsValid() && !CurrentToolTarget.LandscapeInfo->CanHaveLayersContent())
+	{
+		UpdateShownLayerList(); // do it sync when not in lanscape mode.
+	}
+}
+
 void FEdModeLandscape::UpdateShownLayerList()
 {
+	bNeedsUpdateShownLayerList = false;
+
 	if (!CurrentToolTarget.LandscapeInfo.IsValid())
 	{
 		return;
@@ -2856,9 +2877,7 @@ void FEdModeLandscape::UpdateLayerUsageInformation(TWeakObjectPtr<ULandscapeLaye
 	}
 
 	bool DetailPanelRefreshRequired = false;
-	TArray<ULandscapeComponent*> AllComponents;
-	CurrentToolTarget.LandscapeInfo->XYtoComponentMap.GenerateValueArray(AllComponents);
-
+	
 	TArray<TWeakObjectPtr<ULandscapeLayerInfoObject>> LayerInfoObjectToProcess;
 	const TArray<TSharedRef<FLandscapeTargetListInfo>>& TargetList = GetTargetList();
 
@@ -2883,43 +2902,23 @@ void FEdModeLandscape::UpdateLayerUsageInformation(TWeakObjectPtr<ULandscapeLaye
 			LayerInfoObjectToProcess.Add(TargetInfo->LayerInfoObj);
 		}
 	}
-
-
+	
+	TArray<ULandscapeLayerInfoObject*> UsedLayerInfos;
+	CurrentToolTarget.LandscapeInfo->GetUsedPaintLayers(FGuid(), UsedLayerInfos);
+		
 	for (const TWeakObjectPtr<ULandscapeLayerInfoObject>& LayerInfoObj : LayerInfoObjectToProcess)
-	{		
-		for (ULandscapeComponent* Component : AllComponents)
+	{
+		if (ULandscapeLayerInfoObject* LayerInfo = LayerInfoObj.Get())
 		{
-			TArray<uint8> WeightmapTextureData;
-			FLandscapeComponentDataInterface DataInterface(Component);
-			DataInterface.GetWeightmapTextureData(LayerInfoObj.Get(), WeightmapTextureData);
-
-			bool IsUsed = false;
-
-			for (uint8 Value : WeightmapTextureData)
-			{
-				if (Value > 0)
-				{
-					IsUsed = true;
-					break;
-				}
-			}
-
-			bool PreviousValue = LayerInfoObj->IsReferencedFromLoadedData;
-			LayerInfoObj->IsReferencedFromLoadedData = IsUsed;
-
-			if (PreviousValue != LayerInfoObj->IsReferencedFromLoadedData)
+			bool bUsed = UsedLayerInfos.Contains(LayerInfo);
+			if (LayerInfo->IsReferencedFromLoadedData != bUsed)
 			{
+				LayerInfo->IsReferencedFromLoadedData = bUsed;
 				DetailPanelRefreshRequired = true;
 			}
-
-			// Early exit as we already found a component using this layer
-			if (LayerInfoObj->IsReferencedFromLoadedData)
-			{
-				break;
-			}
 		}
 	}
-
+	
 	if (DetailPanelRefreshRequired)
 	{
 		if (Toolkit.IsValid())
diff --git a/Engine/Source/Editor/LandscapeEditor/Private/LandscapeEdMode.h b/Engine/Source/Editor/LandscapeEditor/Private/LandscapeEdMode.h
--- a/Engine/Source/Editor/LandscapeEditor/Private/LandscapeEdMode.h
+++ b/Engine/Source/Editor/LandscapeEditor/Private/LandscapeEdMode.h
@@ -362,6 +362,7 @@ public:
 	virtual ~FEdModeLandscape();
 
 	/** ILandscapeEdModeInterface */
+	virtual void PostUpdateLayerContent() override;
 	virtual ELandscapeToolTargetType::Type GetLandscapeToolTargetType() const override;
 	virtual const FLandscapeLayer* GetLandscapeSelectedLayer() const override;
 	virtual ULandscapeLayerInfoObject* GetSelectedLandscapeLayerInfo() const override;
@@ -529,6 +530,7 @@ public:
 	void MoveTargetLayerDisplayOrder(int32 IndexToMove, int32 IndexToDestination);
 
 	/** Update shown layer list */	
+	void RequestUpdateShownLayerList();
 	void UpdateShownLayerList();
 	bool ShouldShowLayer(TSharedRef<FLandscapeTargetListInfo> Target) const;
 	void UpdateLayerUsageInformation(TWeakObjectPtr<ULandscapeLayerInfoObject>* LayerInfoObjectThatChanged = nullptr);
@@ -639,4 +641,7 @@ private:
 
 	/** The interactor that is currently painting, prevents multiple interactors from sculpting when one actually is */
 	class UViewportInteractor* InteractorPainting;
+
+	/** Delayed refresh */
+	bool bNeedsUpdateShownLayerList;
 };
\ No newline at end of file
diff --git a/Engine/Source/Editor/LandscapeEditor/Private/LandscapeEditorDetailCustomization_Layers.cpp b/Engine/Source/Editor/LandscapeEditor/Private/LandscapeEditorDetailCustomization_Layers.cpp
--- a/Engine/Source/Editor/LandscapeEditor/Private/LandscapeEditorDetailCustomization_Layers.cpp
+++ b/Engine/Source/Editor/LandscapeEditor/Private/LandscapeEditorDetailCustomization_Layers.cpp
@@ -358,6 +358,61 @@ void FLandscapeEditorCustomNodeBuilder_Layers::FillAddBrushMenu(FMenuBuilder& Me
 	}
 }
 
+void FLandscapeEditorCustomNodeBuilder_Layers::FillClearPaintLayerMenu(FMenuBuilder& MenuBuilder, int32 InLayerIndex, TArray<ULandscapeLayerInfoObject*> InUsedLayerInfos)
+{
+	// Clear All Weightmap Data
+	TSharedRef<FLandscapeEditorCustomNodeBuilder_Layers> SharedThis = AsShared();
+	FUIAction ClearAction = FUIAction(FExecuteAction::CreateLambda([SharedThis, InLayerIndex]() { SharedThis->ClearLayer(InLayerIndex, ELandscapeClearMode::Clear_Weightmap); }));
+	MenuBuilder.AddMenuEntry(LOCTEXT("LandscapeClearAllWeightmap", "All"), FText(), FSlateIcon(), ClearAction);
+	MenuBuilder.AddMenuSeparator();
+
+	// Clear Per LayerInfo
+	for (ULandscapeLayerInfoObject* LayerInfo : InUsedLayerInfos)
+	{
+		FUIAction ClearLayerInfoAction = FUIAction(FExecuteAction::CreateLambda([SharedThis, InLayerIndex, LayerInfo]() { SharedThis->ClearPaintLayer(InLayerIndex, LayerInfo); }));
+		MenuBuilder.AddMenuEntry(FText::FromName(LayerInfo->LayerName), FText(), FSlateIcon(), ClearLayerInfoAction);
+	}
+}
+
+void FLandscapeEditorCustomNodeBuilder_Layers::FillClearLayerMenu(FMenuBuilder& MenuBuilder, int32 InLayerIndex)
+{
+	TSharedRef<FLandscapeEditorCustomNodeBuilder_Layers> SharedThis = AsShared();
+	const UEnum* EnumPtr = FindObject<UEnum>(ANY_PACKAGE, TEXT("ELandscapeClearMode"), true);
+	if (ensure(EnumPtr != nullptr))
+	{
+		// NumEnums()-1 to exclude Enum Max Value
+		for (int32 i = 0; i < EnumPtr->NumEnums()-1; ++i)
+		{
+			ELandscapeClearMode EnumValue = (ELandscapeClearMode)EnumPtr->GetValueByIndex(i);
+			if(EnumValue == ELandscapeClearMode::Clear_Weightmap)
+			{
+				TArray<ULandscapeLayerInfoObject*> UsedLayerInfos;
+				FEdModeLandscape* LandscapeEdMode = GetEditorMode();
+				ensure(LandscapeEdMode);
+				ALandscape* Landscape = LandscapeEdMode->GetLandscape();
+				ensure(Landscape);
+
+				Landscape->GetUsedPaintLayers(InLayerIndex, UsedLayerInfos);
+				if (UsedLayerInfos.Num() > 0)
+				{
+					MenuBuilder.AddSubMenu(
+						EnumPtr->GetDisplayNameTextByIndex(i),
+						FText(),
+						FNewMenuDelegate::CreateSP(this, &FLandscapeEditorCustomNodeBuilder_Layers::FillClearPaintLayerMenu, InLayerIndex, UsedLayerInfos),
+						false,
+						FSlateIcon()
+					);
+				}
+			}
+			else
+			{
+				FUIAction ClearAction = FUIAction(FExecuteAction::CreateLambda([SharedThis, InLayerIndex, EnumValue]() { SharedThis->ClearLayer(InLayerIndex, EnumValue); }));
+				MenuBuilder.AddMenuEntry(EnumPtr->GetDisplayNameTextByIndex(i), FText(), FSlateIcon(), ClearAction);
+			}
+		}
+	}
+}
+
 void FLandscapeEditorCustomNodeBuilder_Layers::AddBrushToCurrentLayer(ALandscapeBlueprintCustomBrush* Brush)
 {
 	const FScopedTransaction Transaction(LOCTEXT("LandscapeBrushAddToCurrentLayerTransaction", "Add brush to current layer"));
@@ -389,8 +444,13 @@ TSharedPtr<SWidget> FLandscapeEditorCustomNodeBuilder_Layers::OnLayerContextMenu
 				if (!Layer->bLocked)
 				{
 					// Clear Layer
-					FUIAction ClearLayerAction = FUIAction(FExecuteAction::CreateLambda([SharedThis, InLayerIndex] { SharedThis->ClearLayer(InLayerIndex); }));
-					MenuBuilder.AddMenuEntry(LOCTEXT("ClearLayer", "Clear..."), LOCTEXT("ClearLayerTooltip", "Clear Layer"), FSlateIcon(), ClearLayerAction);
+					MenuBuilder.AddSubMenu(
+						LOCTEXT("LandscapeEditorClearLayerSubMenu", "Clear"),
+						FText(),
+						FNewMenuDelegate::CreateSP(this, &FLandscapeEditorCustomNodeBuilder_Layers::FillClearLayerMenu, InLayerIndex),
+						false,
+						FSlateIcon()
+					);
 
 					// Delete Layer
 					FUIAction DeleteLayerAction = FUIAction(FExecuteAction::CreateLambda([SharedThis, InLayerIndex] { SharedThis->DeleteLayer(InLayerIndex); }), FCanExecuteAction::CreateLambda([Landscape] { return Landscape->LandscapeLayers.Num() > 1; }));
@@ -518,7 +578,27 @@ void FLandscapeEditorCustomNodeBuilder_Layers::RenameLayer(int32 InLayerIndex)
 	}
 }
 
-void FLandscapeEditorCustomNodeBuilder_Layers::ClearLayer(int32 InLayerIndex)
+void FLandscapeEditorCustomNodeBuilder_Layers::ClearPaintLayer(int32 InLayerIndex, ULandscapeLayerInfoObject* InLayerInfo)
+{
+	FEdModeLandscape* LandscapeEdMode = GetEditorMode();
+	ALandscape* Landscape = LandscapeEdMode ? LandscapeEdMode->GetLandscape() : nullptr;
+	if (Landscape)
+	{
+		FLandscapeLayer* Layer = LandscapeEdMode->GetLayer(InLayerIndex);
+		if (Layer)
+		{
+			EAppReturnType::Type Result = FMessageDialog::Open(EAppMsgType::YesNo, FText::Format(LOCTEXT("Landscape_ClearPaintLayer_Message", "The layer {0} : {1} content will be completely cleared.  Continue?"), FText::FromName(Layer->Name), FText::FromName(InLayerInfo->LayerName)));
+			if (Result == EAppReturnType::Yes)
+			{
+				const FScopedTransaction Transaction(LOCTEXT("Landscape_Layers_PaintClear", "Clear Paint Layer"));
+				Landscape->ClearPaintLayer(InLayerIndex, InLayerInfo);
+				LandscapeEdMode->RequestUpdateShownLayerList();
+			}
+		}
+	}
+}
+
+void FLandscapeEditorCustomNodeBuilder_Layers::ClearLayer(int32 InLayerIndex, ELandscapeClearMode InClearMode)
 {
 	FEdModeLandscape* LandscapeEdMode = GetEditorMode();
 	ALandscape* Landscape = LandscapeEdMode ? LandscapeEdMode->GetLandscape() : nullptr;
@@ -531,8 +611,12 @@ void FLandscapeEditorCustomNodeBuilder_Layers::ClearLayer(int32 InLayerIndex)
 			if (Result == EAppReturnType::Yes)
 			{
 				const FScopedTransaction Transaction(LOCTEXT("Landscape_Layers_Clean", "Clear Layer"));
-				Landscape->ClearLayer(InLayerIndex);
+				Landscape->ClearLayer(InLayerIndex, nullptr, InClearMode);
 				OnLayerSelectionChanged(InLayerIndex);
+				if (InClearMode & ELandscapeClearMode::Clear_Weightmap)
+				{
+					LandscapeEdMode->RequestUpdateShownLayerList();
+				}
 			}
 		}
 	}
diff --git a/Engine/Source/Editor/LandscapeEditor/Private/LandscapeEditorDetailCustomization_Layers.h b/Engine/Source/Editor/LandscapeEditor/Private/LandscapeEditorDetailCustomization_Layers.h
--- a/Engine/Source/Editor/LandscapeEditor/Private/LandscapeEditorDetailCustomization_Layers.h
+++ b/Engine/Source/Editor/LandscapeEditor/Private/LandscapeEditorDetailCustomization_Layers.h
@@ -24,6 +24,7 @@ class IDetailChildrenBuilder;
 class IDetailLayoutBuilder;
 class SDragAndDropVerticalBox;
 class ALandscapeBlueprintCustomBrush;
+class ULandscapeLayerInfoObject;
 
 /**
  * Slate widgets customizer for the layers list in the Landscape Editor
@@ -73,7 +74,8 @@ protected:
 	void OnLayerSelectionChanged(int32 LayerIndex);
 	TSharedPtr<SWidget> OnLayerContextMenuOpening(int32 InLayerIndex);
 	void CreateLayer();
-	void ClearLayer(int32 InLayerIndex);
+	void ClearPaintLayer(int32 InLayerIndex, ULandscapeLayerInfoObject* InLayerInfo);
+	void ClearLayer(int32 InLayerIndex, ELandscapeClearMode InClearMode);
 	void RenameLayer(int32 InLayerIndex);
 	void DeleteLayer(int32 InLayerIndex);
 	void ShowOnlySelectedLayer(int32 InLayerIndex);
@@ -99,6 +101,8 @@ protected:
 	const FSlateBrush* GetLockBrushForLayer(int32 InLayerIndex) const;
 
 	void FillAddBrushMenu(FMenuBuilder& MenuBuilder, TArray<ALandscapeBlueprintCustomBrush*> Brushes);
+	void FillClearPaintLayerMenu(FMenuBuilder& MenuBuilder, int32 InLayerIndex, TArray<ULandscapeLayerInfoObject*> InUsedLayerInfos);
+	void FillClearLayerMenu(FMenuBuilder& MenuBuilder, int32 InLayerIndex);
 	void AddBrushToCurrentLayer(ALandscapeBlueprintCustomBrush* Brush);
 
 private:
diff --git a/Engine/Source/Editor/LandscapeEditor/Private/LandscapeEditorDetailCustomization_LayersBrushStack.cpp b/Engine/Source/Editor/LandscapeEditor/Private/LandscapeEditorDetailCustomization_LayersBrushStack.cpp
--- a/Engine/Source/Editor/LandscapeEditor/Private/LandscapeEditorDetailCustomization_LayersBrushStack.cpp
+++ b/Engine/Source/Editor/LandscapeEditor/Private/LandscapeEditorDetailCustomization_LayersBrushStack.cpp
@@ -222,8 +222,8 @@ TSharedPtr<SWidget> FLandscapeEditorCustomNodeBuilder_LayersBrushStack::OnBrushC
 			if (FilteredBrushes.Num())
 			{
 				MenuBuilder.AddSubMenu(
-					LOCTEXT("LandscaeEditorBrushAddSubMenu", "Add"),
-					LOCTEXT("LandscaeEditorBrushAddSubMenuToolTip", "Add brush to current layer"),
+					LOCTEXT("LandscapeEditorBrushAddSubMenu", "Add"),
+					LOCTEXT("LandscapeEditorBrushAddSubMenuToolTip", "Add brush to current layer"),
 					FNewMenuDelegate::CreateSP(this, &FLandscapeEditorCustomNodeBuilder_LayersBrushStack::FillAddBrushMenu, FilteredBrushes),
 					false,
 					FSlateIcon()
diff --git a/Engine/Source/Runtime/Landscape/Classes/Landscape.h b/Engine/Source/Runtime/Landscape/Classes/Landscape.h
--- a/Engine/Source/Runtime/Landscape/Classes/Landscape.h
+++ b/Engine/Source/Runtime/Landscape/Classes/Landscape.h
@@ -279,8 +279,12 @@ public:
 	LANDSCAPE_API const struct FLandscapeLayer* GetLayer(int32 InLayerIndex) const;
 	LANDSCAPE_API const struct FLandscapeLayer* GetLayer(const FGuid& InLayerGuid) const;
 	LANDSCAPE_API void ForEachLayer(TFunctionRef<void(struct FLandscapeLayer&)> Fn);
-	LANDSCAPE_API void ClearLayer(int32 InLayerIndex, TSet<ULandscapeComponent*>* InComponents = nullptr);
-	LANDSCAPE_API void ClearLayer(const FGuid& InLayerGuid, TSet<ULandscapeComponent*>* InComponents = nullptr);
+	LANDSCAPE_API void GetUsedPaintLayers(int32 InLayerIndex, TArray<ULandscapeLayerInfoObject*>& OutUsedLayerInfos) const;
+	LANDSCAPE_API void GetUsedPaintLayers(const FGuid& InLayerGuid, TArray<ULandscapeLayerInfoObject*>& OutUsedLayerInfos) const;
+	LANDSCAPE_API void ClearPaintLayer(int32 InLayerIndex, ULandscapeLayerInfoObject* InLayerInfo);
+	LANDSCAPE_API void ClearPaintLayer(const FGuid& InLayerGuid, ULandscapeLayerInfoObject* InLayerInfo);
+	LANDSCAPE_API void ClearLayer(int32 InLayerIndex, TSet<ULandscapeComponent*>* InComponents = nullptr, ELandscapeClearMode InClearMode = ELandscapeClearMode::Clear_All);
+	LANDSCAPE_API void ClearLayer(const FGuid& InLayerGuid, TSet<ULandscapeComponent*>* InComponents = nullptr, ELandscapeClearMode InClearMode = ELandscapeClearMode::Clear_All);
 	LANDSCAPE_API void DeleteLayer(int32 InLayerIndex);
 	LANDSCAPE_API void DeleteLayers();
 	LANDSCAPE_API void SetEditingLayer(const FGuid& InLayerGuid = FGuid());
@@ -383,7 +387,7 @@ private:
 public:
 
 #if WITH_EDITORONLY_DATA
-	UPROPERTY(EditAnywhere, Category=Experimental)
+	UPROPERTY(EditAnywhere, Category=Landscape)
 	bool bCanHaveLayersContent = false;
 
 	DECLARE_EVENT(ALandscape, FLandscapeBPCustomBrushChangedDelegate);
diff --git a/Engine/Source/Runtime/Landscape/Classes/LandscapeComponent.h b/Engine/Source/Runtime/Landscape/Classes/LandscapeComponent.h
--- a/Engine/Source/Runtime/Landscape/Classes/LandscapeComponent.h
+++ b/Engine/Source/Runtime/Landscape/Classes/LandscapeComponent.h
@@ -327,6 +327,14 @@ enum ELandscapeLayerUpdateMode : uint32
 
 #endif
 
+UENUM()
+enum ELandscapeClearMode
+{
+	Clear_Weightmap = 1 << 0 UMETA(DisplayName = "Paint"),
+	Clear_Heightmap = 1 << 1 UMETA(DisplayName = "Sculpt"),
+	Clear_All = Clear_Weightmap | Clear_Heightmap UMETA(DisplayName = "All")
+};
+
 UCLASS(hidecategories=(Display, Attachment, Physics, Debug, Collision, Movement, Rendering, PrimitiveComponent, Object, Transform, Mobility), showcategories=("Rendering|Material"), MinimalAPI, Within=LandscapeProxy)
 class ULandscapeComponent : public UPrimitiveComponent
 {
@@ -595,6 +603,7 @@ public:
 
 	LANDSCAPE_API TArray<FWeightmapLayerAllocationInfo>& GetWeightmapLayerAllocations(bool InReturnEditingWeightmap = false);
 	LANDSCAPE_API const TArray<FWeightmapLayerAllocationInfo>& GetWeightmapLayerAllocations(bool InReturnEditingWeightmap = false) const;
+	LANDSCAPE_API const TArray<FWeightmapLayerAllocationInfo>& GetWeightmapLayerAllocations(const FGuid& InLayerGuid) const;
 
 #if WITH_EDITOR
 	LANDSCAPE_API void SetHeightmap(UTexture2D* NewHeightmap);
diff --git a/Engine/Source/Runtime/Landscape/Classes/LandscapeInfo.h b/Engine/Source/Runtime/Landscape/Classes/LandscapeInfo.h
--- a/Engine/Source/Runtime/Landscape/Classes/LandscapeInfo.h
+++ b/Engine/Source/Runtime/Landscape/Classes/LandscapeInfo.h
@@ -157,6 +157,7 @@ public:
 	LANDSCAPE_API bool IsValidPosition(int32 X, int32 Y);
 	LANDSCAPE_API void DeleteLayer(ULandscapeLayerInfoObject* LayerInfo, const FName& LayerName);
 	LANDSCAPE_API void ReplaceLayer(ULandscapeLayerInfoObject* FromLayerInfo, ULandscapeLayerInfoObject* ToLayerInfo);
+	LANDSCAPE_API void GetUsedPaintLayers(const FGuid& InLayerGuid, TArray<ULandscapeLayerInfoObject*>& OutUsedLayerInfos) const;
 
 	LANDSCAPE_API void UpdateDebugColorMaterial();
 
diff --git a/Engine/Source/Runtime/Landscape/Private/Landscape.cpp b/Engine/Source/Runtime/Landscape/Private/Landscape.cpp
--- a/Engine/Source/Runtime/Landscape/Private/Landscape.cpp
+++ b/Engine/Source/Runtime/Landscape/Private/Landscape.cpp
@@ -1520,6 +1520,21 @@ const TArray<FWeightmapLayerAllocationInfo>& ULandscapeComponent::GetWeightmapLa
 	return WeightmapLayerAllocations;
 }
 
+const TArray<FWeightmapLayerAllocationInfo>& ULandscapeComponent::GetWeightmapLayerAllocations(const FGuid& InLayerGuid) const
+{
+#if WITH_EDITORONLY_DATA
+	if (InLayerGuid.IsValid())
+	{
+		if (const FLandscapeLayerComponentData* LayerData = GetLayerData(InLayerGuid))
+		{
+			return LayerData->WeightmapData.LayerAllocations;
+		}
+	}
+#endif
+
+	return WeightmapLayerAllocations;
+}
+
 TArray<FWeightmapLayerAllocationInfo>& ULandscapeComponent::GetWeightmapLayerAllocations(bool InReturnEditingWeightmap)
 {
 #if WITH_EDITORONLY_DATA
diff --git a/Engine/Source/Runtime/Landscape/Private/LandscapeEdit.cpp b/Engine/Source/Runtime/Landscape/Private/LandscapeEdit.cpp
--- a/Engine/Source/Runtime/Landscape/Private/LandscapeEdit.cpp
+++ b/Engine/Source/Runtime/Landscape/Private/LandscapeEdit.cpp
@@ -3487,6 +3487,22 @@ void ULandscapeInfo::ReplaceLayer(ULandscapeLayerInfoObject* FromLayerInfo, ULan
 	}
 }
 
+void ULandscapeInfo::GetUsedPaintLayers(const FGuid& InLayerGuid, TArray<ULandscapeLayerInfoObject*>& OutUsedLayerInfos) const
+{
+	OutUsedLayerInfos.Empty();
+	ForAllLandscapeProxies([&](ALandscapeProxy* Proxy)
+	{
+		for (ULandscapeComponent* Component : Proxy->LandscapeComponents)
+		{
+			const TArray<FWeightmapLayerAllocationInfo>& AllocInfos = Component->GetWeightmapLayerAllocations(InLayerGuid);
+			for (const FWeightmapLayerAllocationInfo& AllocInfo : AllocInfos)
+			{
+				OutUsedLayerInfos.AddUnique(AllocInfo.LayerInfo);
+			}
+		}
+	});
+}
+
 void ALandscapeProxy::EditorApplyScale(const FVector& DeltaScale, const FVector* PivotLocation, bool bAltDown, bool bShiftDown, bool bCtrlDown)
 {
 	FVector ModifiedScale = DeltaScale;
diff --git a/Engine/Source/Runtime/Landscape/Private/LandscapeEditLayers.cpp b/Engine/Source/Runtime/Landscape/Private/LandscapeEditLayers.cpp
--- a/Engine/Source/Runtime/Landscape/Private/LandscapeEditLayers.cpp
+++ b/Engine/Source/Runtime/Landscape/Private/LandscapeEditLayers.cpp
@@ -4327,6 +4327,11 @@ void ALandscape::UpdateLayersContent(bool bInWaitForStreaming, bool bInSkipMonit
 	{
 		LayerContentUpdateModes |= ELandscapeLayerUpdateMode::Update_Client_Deferred;
 	}
+
+	if (LandscapeEdMode)
+	{
+		LandscapeEdMode->PostUpdateLayerContent();
+	}
 }
 
 // not thread safe
@@ -4981,16 +4986,67 @@ void ALandscape::DeleteLayer(int32 InLayerIndex)
 	RequestLayersContentUpdateForceAll();
 }
 
-void ALandscape::ClearLayer(int32 InLayerIndex, TSet<ULandscapeComponent*>* InComponents)
+void ALandscape::GetUsedPaintLayers(int32 InLayerIndex, TArray<ULandscapeLayerInfoObject*>& OutUsedLayerInfos) const
+{
+	const FLandscapeLayer* Layer = GetLayer(InLayerIndex);
+	if (Layer)
+	{
+		GetUsedPaintLayers(Layer->Guid, OutUsedLayerInfos);
+	}
+}
+
+void ALandscape::GetUsedPaintLayers(const FGuid& InLayerGuid, TArray<ULandscapeLayerInfoObject*>& OutUsedLayerInfos) const
+{
+	ULandscapeInfo* LandscapeInfo = GetLandscapeInfo();
+	if (!LandscapeInfo)
+	{
+		return;
+	}
+	
+	LandscapeInfo->GetUsedPaintLayers(InLayerGuid, OutUsedLayerInfos);
+}
+
+void ALandscape::ClearPaintLayer(int32 InLayerIndex, ULandscapeLayerInfoObject* InLayerInfo)
 {
 	const FLandscapeLayer* Layer = GetLayer(InLayerIndex);
 	if (Layer)
 	{
-		ClearLayer(Layer->Guid, InComponents);
+		ClearPaintLayer(Layer->Guid, InLayerInfo);
 	}
 }
 
-void ALandscape::ClearLayer(const FGuid& InLayerGuid, TSet<ULandscapeComponent*>* InComponents)
+void ALandscape::ClearPaintLayer(const FGuid& InLayerGuid, ULandscapeLayerInfoObject* InLayerInfo)
+{
+	ULandscapeInfo* LandscapeInfo = GetLandscapeInfo();
+	if (!LandscapeInfo)
+	{
+		return;
+	}
+
+	Modify();
+	FScopedSetLandscapeEditingLayer Scope(this, InLayerGuid, [=] { RequestLayersContentUpdate(ELandscapeLayerUpdateMode::Update_Weightmap_All); });
+
+	FLandscapeEditDataInterface LandscapeEdit(LandscapeInfo);
+	LandscapeInfo->ForAllLandscapeProxies([&](ALandscapeProxy* Proxy)
+	{
+		Proxy->Modify();
+		for (ULandscapeComponent* Component : Proxy->LandscapeComponents)
+		{
+			Component->DeleteLayer(InLayerInfo, LandscapeEdit);
+		}
+	});
+}
+
+void ALandscape::ClearLayer(int32 InLayerIndex, TSet<ULandscapeComponent*>* InComponents, ELandscapeClearMode InClearMode)
+{
+	const FLandscapeLayer* Layer = GetLayer(InLayerIndex);
+	if (Layer)
+	{
+		ClearLayer(Layer->Guid, InComponents, InClearMode);
+	}
+}
+
+void ALandscape::ClearLayer(const FGuid& InLayerGuid, TSet<ULandscapeComponent*>* InComponents, ELandscapeClearMode InClearMode)
 {
 	ensure(HasLayersContent());
 
@@ -5015,11 +5071,14 @@ void ALandscape::ClearLayer(const FGuid& InLayerGuid, TSet<ULandscapeComponent*>
 	TArray<uint8> NewHeightAlphaBlendData;
 	TArray<uint8> NewHeightFlagsData;
 
-	if (Layer->BlendMode == LSBM_AlphaBlend)
+	if (InClearMode & ELandscapeClearMode::Clear_Heightmap)
 	{
-		NewHeightAlphaBlendData.AddZeroed(FMath::Square(ComponentSizeQuads + 1));
-		FMemory::Memset(NewHeightAlphaBlendData.GetData(), 255, NewHeightAlphaBlendData.Num());
-		NewHeightFlagsData.AddZeroed(FMath::Square(ComponentSizeQuads + 1));
+		if (Layer->BlendMode == LSBM_AlphaBlend)
+		{
+			NewHeightAlphaBlendData.AddZeroed(FMath::Square(ComponentSizeQuads + 1));
+			FMemory::Memset(NewHeightAlphaBlendData.GetData(), 255, NewHeightAlphaBlendData.Num());
+			NewHeightFlagsData.AddZeroed(FMath::Square(ComponentSizeQuads + 1));
+		}
 	}
 
 	TArray<ULandscapeComponent*> Components;
@@ -5053,19 +5112,25 @@ void ALandscape::ClearLayer(const FGuid& InLayerGuid, TSet<ULandscapeComponent*>
 	FLandscapeEditDataInterface LandscapeEdit(LandscapeInfo);
 	for (ULandscapeComponent* Component : Components)
 	{
-		int32 MinX = MAX_int32;
-		int32 MinY = MAX_int32;
-		int32 MaxX = MIN_int32;
-		int32 MaxY = MIN_int32;
-		Component->GetComponentExtent(MinX, MinY, MaxX, MaxY);
-		check(ComponentSizeQuads == (MaxX - MinX));
-		check(ComponentSizeQuads == (MaxY - MinY));
-		LandscapeEdit.SetHeightData(MinX, MinY, MaxX, MaxY, NewHeightData.GetData(), 0, false, nullptr, NewHeightAlphaBlendData.GetData(), NewHeightFlagsData.GetData());
+		if (InClearMode & ELandscapeClearMode::Clear_Heightmap)
+		{
+			int32 MinX = MAX_int32;
+			int32 MinY = MAX_int32;
+			int32 MaxX = MIN_int32;
+			int32 MaxY = MIN_int32;
+			Component->GetComponentExtent(MinX, MinY, MaxX, MaxY);
+			check(ComponentSizeQuads == (MaxX - MinX));
+			check(ComponentSizeQuads == (MaxY - MinY));
+			LandscapeEdit.SetHeightData(MinX, MinY, MaxX, MaxY, NewHeightData.GetData(), 0, false, nullptr, NewHeightAlphaBlendData.GetData(), NewHeightFlagsData.GetData());
+		}
 
-		// Clear weight maps
-		for (FLandscapeInfoLayerSettings& LayerSettings : LandscapeInfo->Layers)
+		if (InClearMode & ELandscapeClearMode::Clear_Weightmap)
 		{
-			Component->DeleteLayer(LayerSettings.LayerInfoObj, LandscapeEdit);
+			// Clear weight maps
+			for (FLandscapeInfoLayerSettings& LayerSettings : LandscapeInfo->Layers)
+			{
+				Component->DeleteLayer(LayerSettings.LayerInfoObj, LandscapeEdit);
+			}
 		}
 	}
 }
diff --git a/Engine/Source/Runtime/Landscape/Public/LandscapeEdit.h b/Engine/Source/Runtime/Landscape/Public/LandscapeEdit.h
--- a/Engine/Source/Runtime/Landscape/Public/LandscapeEdit.h
+++ b/Engine/Source/Runtime/Landscape/Public/LandscapeEdit.h
@@ -47,6 +47,7 @@ namespace ELandscapeToolTargetType
 class ILandscapeEdModeInterface
 {
 public:
+	virtual void PostUpdateLayerContent() = 0;
 	virtual ELandscapeToolTargetType::Type GetLandscapeToolTargetType() const = 0;
 	virtual const FLandscapeLayer* GetLandscapeSelectedLayer() const = 0;
 	virtual ULandscapeLayerInfoObject* GetSelectedLandscapeLayerInfo() const = 0;
