Missing request update from previous CL

#rnx

diff --git a/Engine/Source/Editor/LandscapeEditor/Private/LandscapeEditorDetailCustomization_TargetLayers.cpp b/Engine/Source/Editor/LandscapeEditor/Private/LandscapeEditorDetailCustomization_TargetLayers.cpp
--- a/Engine/Source/Editor/LandscapeEditor/Private/LandscapeEditorDetailCustomization_TargetLayers.cpp
+++ b/Engine/Source/Editor/LandscapeEditor/Private/LandscapeEditorDetailCustomization_TargetLayers.cpp
@@ -1131,6 +1131,7 @@ void FLandscapeEditorCustomNodeBuilder_TargetLayers::OnClearLayer(const TSharedR
 			FScopedSetLandscapeEditingLayer Scope(LandscapeEdMode->GetLandscape(), LandscapeEdMode->GetCurrentLayerGuid(), [&] { LandscapeEdMode->RequestLayersContentUpdateForceAll(); });
 			FLandscapeEditDataInterface LandscapeEdit(Target->LandscapeInfo.Get());
 			LandscapeEdit.DeleteLayer(Target->LayerInfoObj.Get());
+			LandscapeEdMode->RequestUpdateShownLayerList();
 		}
 	}
 }
