Optimized landscape painting: avoid recomputing normals, caching original data when not necessary
#rnx

diff --git a/Engine/Source/Editor/LandscapeEditor/Private/LandscapeEdModePaintTools.cpp b/Engine/Source/Editor/LandscapeEditor/Private/LandscapeEdModePaintTools.cpp
--- a/Engine/Source/Editor/LandscapeEditor/Private/LandscapeEdModePaintTools.cpp
+++ b/Engine/Source/Editor/LandscapeEditor/Private/LandscapeEdModePaintTools.cpp
@@ -160,9 +160,10 @@ public:
 		// Tablet pressure
 		float Pressure = ViewportClient->Viewport->IsPenActive() ? ViewportClient->Viewport->GetTabletPressure() : 1.0f;
 
-		this->Cache.CacheData(X1, Y1, X2, Y2);
+		const bool bUseWeightTargetValue = UISettings->bUseWeightTargetValue;
+		const bool bCacheOriginalData = !bUseWeightTargetValue;
 
-		bool bUseWeightTargetValue = UISettings->bUseWeightTargetValue;
+		this->Cache.CacheData(X1, Y1, X2, Y2, bCacheOriginalData);
 
 		// The data we'll be writing to
 		TArray<ToolTarget::CacheClass::DataType> Data;
diff --git a/Engine/Source/Editor/LandscapeEditor/Private/LandscapeEdModeTools.h b/Engine/Source/Editor/LandscapeEditor/Private/LandscapeEdModeTools.h
--- a/Engine/Source/Editor/LandscapeEditor/Private/LandscapeEdModeTools.h
+++ b/Engine/Source/Editor/LandscapeEditor/Private/LandscapeEdModeTools.h
@@ -292,7 +292,7 @@ public:
 	}
 
 	// X2/Y2 Coordinates are "inclusive" max values
-	void CacheData(int32 X1, int32 Y1, int32 X2, int32 Y2, bool bCacheOriginalData = true)
+	void CacheData(int32 X1, int32 Y1, int32 X2, int32 Y2, bool bCacheOriginalData = false)
 	{
 		if (!Valid)
 		{
@@ -510,8 +510,7 @@ public:
 		{
 			FIntRect Bounds = GetCacheRegion();
 			check((Bounds.Min.X <= X1) && (Bounds.Min.Y <= Y1) && (Bounds.Max.X >= X2) && (Bounds.Max.Y >= Y2));
-			const bool bCacheOriginalData = false;
-			CacheData(Bounds.Min.X, Bounds.Min.Y, Bounds.Max.X, Bounds.Max.Y, bCacheOriginalData);
+			CacheData(Bounds.Min.X, Bounds.Min.Y, Bounds.Max.X, Bounds.Max.Y);
 		}
 		ensure(HasCachedData(X1, Y1, X2, Y2));
 		return GetCachedData(X1, Y1, X2, Y2, OutData);
@@ -729,11 +728,13 @@ struct FXYOffsetmapAccessor
 						
 			// Notify foliage to move any attached instances
 			bool bUpdateFoliage = false;
+			bool bUpdateNormals = false;
 			
 			ALandscapeProxy::InvalidateGeneratedComponentData(Components);
 
 			if (!LandscapeEdit->HasLandscapeLayersContent())
 			{
+				bUpdateNormals = true;
 				for (ULandscapeComponent* Component : Components)
 				{
 					ULandscapeHeightfieldCollisionComponent* CollisionComponent = Component->CollisionComponent.Get();
@@ -761,7 +762,7 @@ struct FXYOffsetmapAccessor
 
 				// Update landscape.
 				LandscapeEdit->SetXYOffsetData(X1, Y1, X2, Y2, Data, 0); // XY Offset always need to be update before the height update
-				LandscapeEdit->SetHeightData(X1, Y1, X2, Y2, NewHeights.GetData(), 0, true);
+				LandscapeEdit->SetHeightData(X1, Y1, X2, Y2, NewHeights.GetData(), 0, bUpdateNormals);
 
 				// Snap foliage for each component.
 				for (int32 Index = 0; Index < CollisionComponents.Num(); ++Index)
@@ -774,7 +775,7 @@ struct FXYOffsetmapAccessor
 			{
 				// No foliage, just update landscape.
 				LandscapeEdit->SetXYOffsetData(X1, Y1, X2, Y2, Data, 0); // XY Offset always need to be update before the height update
-				LandscapeEdit->SetHeightData(X1, Y1, X2, Y2, NewHeights.GetData(), 0, true);
+				LandscapeEdit->SetHeightData(X1, Y1, X2, Y2, NewHeights.GetData(), 0, bUpdateNormals);
 			}
 		}
 	}
diff --git a/Engine/Source/Runtime/Landscape/Public/LandscapeEdit.h b/Engine/Source/Runtime/Landscape/Public/LandscapeEdit.h
--- a/Engine/Source/Runtime/Landscape/Public/LandscapeEdit.h
+++ b/Engine/Source/Runtime/Landscape/Public/LandscapeEdit.h
@@ -365,12 +365,14 @@ struct FHeightmapAccessor
 						
 			// Notify foliage to move any attached instances
 			bool bUpdateFoliage = false;
+			bool bUpdateNormals = false;
 
 			ALandscapeProxy::InvalidateGeneratedComponentData(Components);
 
             // Landscape Layers are updates are delayed and done in  ALandscape::TickLayers
 			if (!LandscapeEdit->HasLandscapeLayersContent())
 			{
+				bUpdateNormals = true;
 				for (ULandscapeComponent* Component : Components)
 				{
 					ULandscapeHeightfieldCollisionComponent* CollisionComponent = Component->CollisionComponent.Get();
@@ -401,7 +403,7 @@ struct FHeightmapAccessor
 				}
 
 				// Update landscape.
-				LandscapeEdit->SetHeightData(X1, Y1, X2, Y2, Data, 0, true);
+				LandscapeEdit->SetHeightData(X1, Y1, X2, Y2, Data, 0, bUpdateNormals);
 
 				// Snap foliage for each component.
 				for (int32 Index = 0; Index < CollisionComponents.Num(); ++Index)
@@ -413,7 +415,7 @@ struct FHeightmapAccessor
 			else
 			{
 				// No foliage, just update landscape.
-				LandscapeEdit->SetHeightData(X1, Y1, X2, Y2, Data, 0, true);
+				LandscapeEdit->SetHeightData(X1, Y1, X2, Y2, Data, 0, bUpdateNormals);
 			}
 		}
 	}
