Ramp Tool now behaves like Flatten,Smooth,Erosion tool as it can be used in any layer and can use the previous layer data to get to its desired result

#rnx

diff --git a/Engine/Source/Editor/LandscapeEditor/Private/LandscapeEdMode.cpp b/Engine/Source/Editor/LandscapeEditor/Private/LandscapeEdMode.cpp
--- a/Engine/Source/Editor/LandscapeEditor/Private/LandscapeEdMode.cpp
+++ b/Engine/Source/Editor/LandscapeEditor/Private/LandscapeEdMode.cpp
@@ -4643,15 +4643,6 @@ bool FEdModeLandscape::CanEditLayer(FText* Reason /*=nullptr*/, FLandscapeLayer*
 				}
 				return false;
 			}
-			else if ((TargetLayerIndex > 0) && (CurrentTool->GetToolName() == FName("Ramp")))
-			{
-				if (Reason)
-				{
-					FLandscapeLayer* Layer = GetLayer(0);
-					*Reason = FText::Format(NSLOCTEXT("UnrealEd", "LandscapeLayersToolAvailableOnlyOnFirstLayer", "{0} Tool is only available on the first layer {1}."), CurrentTool->GetDisplayName(), Layer ? FText::FromName(Layer->Name) : FText::GetEmpty());
-				}
-				return false;
-			}
 			else if (CurrentTool->GetToolName() == FName("Retopologize"))
 			{
 				if (Reason)
diff --git a/Engine/Source/Editor/LandscapeEditor/Private/LandscapeEdModeRampTool.cpp b/Engine/Source/Editor/LandscapeEditor/Private/LandscapeEdModeRampTool.cpp
--- a/Engine/Source/Editor/LandscapeEditor/Private/LandscapeEdModeRampTool.cpp
+++ b/Engine/Source/Editor/LandscapeEditor/Private/LandscapeEdModeRampTool.cpp
@@ -21,6 +21,7 @@
 #include "Raster.h"
 #include "Landscape.h"
 #include "Misc/MessageDialog.h"
+#include "LandscapeEdModeTools.h"
 
 #define LOCTEXT_NAMESPACE "Landscape"
 
@@ -535,6 +536,10 @@ public:
 		}
 
 		FLandscapeEditDataInterface LandscapeEdit(EdMode->CurrentToolTarget.LandscapeInfo.Get());
+		FLandscapeHeightCache HeightCache(EdMode->CurrentToolTarget);
+		FLandscapeLayerDataCache<FHeightmapToolTarget> LayerHeightDataCache(EdMode->CurrentToolTarget, HeightCache);
+		const bool bCombinedLayerOperation = EdMode->UISettings->bCombinedLayersOperation && Landscape && Landscape->HasLayersContent();
+		LayerHeightDataCache.Initialize(EdMode->CurrentToolTarget.LandscapeInfo.Get(), bCombinedLayerOperation);
 
 		// Heights raster
 		bool bRaiseTerrain = true; //EdMode->UISettings->Ramp_bRaiseTerrain;
@@ -548,7 +553,7 @@ public:
 			int32 ValidMinY = MinY;
 			int32 ValidMaxX = MaxX;
 			int32 ValidMaxY = MaxY;
-			LandscapeEdit.GetHeightData(ValidMinX, ValidMinY, ValidMaxX, ValidMaxY, Data.GetData(), 0);
+			LayerHeightDataCache.Read(ValidMinX, ValidMinY, ValidMaxX, ValidMaxY, Data);
 
 			if (ValidMinX > ValidMaxX || ValidMinY > ValidMaxY)
 			{
@@ -578,8 +583,7 @@ public:
 			Rasterizer.DrawTriangle(FVector2D(1, Heights[0]), FVector2D(0, Heights[0]), FVector2D(1, Heights[1]), InnerVerts[0][1], OuterVerts[0][1], InnerVerts[1][1], false);
 			Rasterizer.DrawTriangle(FVector2D(0, Heights[0]), FVector2D(1, Heights[1]), FVector2D(0, Heights[1]), OuterVerts[0][1], InnerVerts[1][1], OuterVerts[1][1], false);
 
-			LandscapeEdit.SetHeightData(MinX, MinY, MaxX, MaxY, Data.GetData(), 0, true);
-			LandscapeEdit.Flush();
+			LayerHeightDataCache.Write(MinX, MinY, MaxX, MaxY, Data);
 
 			if (!EdMode->HasLandscapeLayersContent())
 			{
diff --git a/Engine/Source/Editor/LandscapeEditor/Public/LandscapeEditorObject.h b/Engine/Source/Editor/LandscapeEditor/Public/LandscapeEditorObject.h
--- a/Engine/Source/Editor/LandscapeEditor/Public/LandscapeEditorObject.h
+++ b/Engine/Source/Editor/LandscapeEditor/Public/LandscapeEditorObject.h
@@ -266,7 +266,7 @@ class ULandscapeEditorObject : public UObject
 	UPROPERTY(NonTransactional)
 	float MaximumValueRadius;
 
-	UPROPERTY(Category="Tool Settings", EditAnywhere, NonTransactional, meta=(ShowForTools="Flatten,Smooth,Erosion,HydraErosion", ShowForTargetTypes="Heightmap", ShowForLandscapeLayerSystem))
+	UPROPERTY(Category="Tool Settings", EditAnywhere, NonTransactional, meta=(ShowForTools="Flatten,Smooth,Erosion,HydraErosion,Ramp", ShowForTargetTypes="Heightmap", ShowForLandscapeLayerSystem))
 	bool bCombinedLayersOperation;
 
 	// Flatten Tool:
