Fix some ed mode settings not being saved

#rnx

diff --git a/Engine/Source/Editor/FoliageEdit/Private/FoliageEdMode.cpp b/Engine/Source/Editor/FoliageEdit/Private/FoliageEdMode.cpp
--- a/Engine/Source/Editor/FoliageEdit/Private/FoliageEdMode.cpp
+++ b/Engine/Source/Editor/FoliageEdit/Private/FoliageEdMode.cpp
@@ -3777,6 +3777,9 @@ void FFoliageUISettings::Load()
 	ActivePaletteViewMode = EFoliagePaletteViewMode::Type(ActivePaletteViewModeAsInt);
 
 	GConfig->GetFloat(TEXT("FoliageEdit"), TEXT("PaletteThumbnailScale"), PaletteThumbnailScale, GEditorPerProjectIni);
+
+	GConfig->GetBool(TEXT("FoliageEdit"), TEXT("IsInSingleInstantiationMode"), IsInSingleInstantiationMode, GEditorPerProjectIni);
+	GConfig->GetBool(TEXT("FoliageEdit"), TEXT("IsInSpawnInCurrentLevelMode"), IsInSpawnInCurrentLevelMode, GEditorPerProjectIni);
 }
 
 /** Save UI settings to ini file */
@@ -3798,6 +3801,10 @@ void FFoliageUISettings::Save()
 	GConfig->SetBool(TEXT("FoliageEdit"), TEXT("bShowPaletteItemTooltips"), bShowPaletteItemTooltips, GEditorPerProjectIni);
 	GConfig->SetInt(TEXT("FoliageEdit"), TEXT("ActivePaletteViewMode"), ActivePaletteViewMode, GEditorPerProjectIni);
 	GConfig->SetFloat(TEXT("FoliageEdit"), TEXT("PaletteThumbnailScale"), PaletteThumbnailScale, GEditorPerProjectIni);
+
+	GConfig->SetBool(TEXT("FoliageEdit"), TEXT("IsInSingleInstantiationMode"), IsInSingleInstantiationMode, GEditorPerProjectIni);
+	GConfig->SetBool(TEXT("FoliageEdit"), TEXT("IsInSpawnInCurrentLevelMode"), IsInSpawnInCurrentLevelMode, GEditorPerProjectIni);
+
 }
 
 #undef LOCTEXT_NAMESPACE
diff --git a/Engine/Source/Editor/LandscapeEditor/Private/LandscapeEditorObject.cpp b/Engine/Source/Editor/LandscapeEditor/Private/LandscapeEditorObject.cpp
--- a/Engine/Source/Editor/LandscapeEditor/Private/LandscapeEditorObject.cpp
+++ b/Engine/Source/Editor/LandscapeEditor/Private/LandscapeEditorObject.cpp
@@ -228,9 +228,14 @@ void ULandscapeEditorObject::Load()
 	bUseFlattenTarget = InbUseFlattenTarget;
 	GConfig->GetFloat(TEXT("LandscapeEdit"), TEXT("FlattenTarget"), FlattenTarget, GEditorPerProjectIni);
 
+	GConfig->GetFloat(TEXT("LandscapeEdit"), TEXT("TerraceSmooth"), TerraceSmooth, GEditorPerProjectIni);
+	GConfig->GetFloat(TEXT("LandscapeEdit"), TEXT("TerraceInterval"), TerraceInterval, GEditorPerProjectIni);
+
 	GConfig->GetFloat(TEXT("LandscapeEdit"), TEXT("RampWidth"), RampWidth, GEditorPerProjectIni);
 	GConfig->GetFloat(TEXT("LandscapeEdit"), TEXT("RampSideFalloff"), RampSideFalloff, GEditorPerProjectIni);
 
+	GConfig->GetBool(TEXT("LandscapeEdit"), TEXT("bCombinedLayersOperation"), bCombinedLayersOperation, GEditorPerProjectIni);
+
 	GConfig->GetInt(TEXT("LandscapeEdit"), TEXT("ErodeThresh"), ErodeThresh, GEditorPerProjectIni);
 	GConfig->GetInt(TEXT("LandscapeEdit"), TEXT("ErodeIterationNum"), ErodeIterationNum, GEditorPerProjectIni);
 	GConfig->GetInt(TEXT("LandscapeEdit"), TEXT("ErodeSurfaceThickness"), ErodeSurfaceThickness, GEditorPerProjectIni);
@@ -348,9 +353,14 @@ void ULandscapeEditorObject::Save()
 	GConfig->SetBool(TEXT("LandscapeEdit"), TEXT("bUseFlattenTarget"), bUseFlattenTarget, GEditorPerProjectIni);
 	GConfig->SetFloat(TEXT("LandscapeEdit"), TEXT("FlattenTarget"), FlattenTarget, GEditorPerProjectIni);
 
+	GConfig->SetFloat(TEXT("LandscapeEdit"), TEXT("TerraceSmooth"), TerraceSmooth, GEditorPerProjectIni);
+	GConfig->SetFloat(TEXT("LandscapeEdit"), TEXT("TerraceInterval"), TerraceInterval, GEditorPerProjectIni);
+
 	GConfig->SetFloat(TEXT("LandscapeEdit"), TEXT("RampWidth"), RampWidth, GEditorPerProjectIni);
 	GConfig->SetFloat(TEXT("LandscapeEdit"), TEXT("RampSideFalloff"), RampSideFalloff, GEditorPerProjectIni);
 
+	GConfig->SetBool(TEXT("LandscapeEdit"), TEXT("bCombinedLayersOperation"), bCombinedLayersOperation, GEditorPerProjectIni);
+
 	GConfig->SetInt(TEXT("LandscapeEdit"), TEXT("ErodeThresh"), ErodeThresh, GEditorPerProjectIni);
 	GConfig->SetInt(TEXT("LandscapeEdit"), TEXT("ErodeIterationNum"), ErodeIterationNum, GEditorPerProjectIni);
 	GConfig->SetInt(TEXT("LandscapeEdit"), TEXT("ErodeSurfaceThickness"), ErodeSurfaceThickness, GEditorPerProjectIni);
