Landscape Layers: Fix Mouse jumping when using Flatten tool

#rnx

diff --git a/Engine/Source/Editor/LandscapeEditor/Private/LandscapeEdModeTools.h b/Engine/Source/Editor/LandscapeEditor/Private/LandscapeEdModeTools.h
--- a/Engine/Source/Editor/LandscapeEditor/Private/LandscapeEdModeTools.h
+++ b/Engine/Source/Editor/LandscapeEditor/Private/LandscapeEdModeTools.h
@@ -970,7 +970,8 @@ public:
 			{
 				TSet<ULandscapeComponent*> AffectedComponents;
 				LandscapeInfo->GetComponentsInRegion(X1, Y1, X2, Y2, AffectedComponents);
-				SynchronousUpdateHeightmapForComponents(AffectedComponents);
+                const bool bUpdateCollision = true;
+				SynchronousUpdateHeightmapForComponents(AffectedComponents, bUpdateCollision);
 				bVisibilityChanged = false;
 			}
 		}
@@ -983,11 +984,12 @@ public:
 
 private:
 
-	void SynchronousUpdateHeightmapForComponents(const TSet<ULandscapeComponent*>& InComponents)
+	void SynchronousUpdateHeightmapForComponents(const TSet<ULandscapeComponent*>& InComponents, bool bUpdateCollision)
 	{
 		for (ULandscapeComponent* Component : InComponents)
 		{
-			Component->RequestHeightmapUpdate();
+			const bool bUpdateAll = false; // default value
+			Component->RequestHeightmapUpdate(bUpdateAll, bUpdateCollision);
 		}
 		Landscape->ForceUpdateLayersContent();
 	};
@@ -1011,7 +1013,9 @@ private:
 	void SynchronousUpdateComponentVisibilityForHeight(const TSet<ULandscapeComponent*>& InComponents, const TArray<bool>& InLayerVisibility)
 	{
 		SetLayersVisibility(InLayerVisibility);
-		SynchronousUpdateHeightmapForComponents(InComponents);
+        // No need to update collision here as we are only doing a intermediate render to gather heightdata
+        const bool bUpdateCollision = false;
+		SynchronousUpdateHeightmapForComponents(InComponents, bUpdateCollision);
 	};
 
 	ULandscapeInfo* LandscapeInfo;
