Fix from 4.23 CL 7660764

#rnx

diff --git a/Engine/Source/Editor/LandscapeEditor/Private/LandscapeEdModePaintTools.cpp b/Engine/Source/Editor/LandscapeEditor/Private/LandscapeEdModePaintTools.cpp
--- a/Engine/Source/Editor/LandscapeEditor/Private/LandscapeEdModePaintTools.cpp
+++ b/Engine/Source/Editor/LandscapeEditor/Private/LandscapeEdModePaintTools.cpp
@@ -799,6 +799,8 @@ public:
 				this->Cache.DataAccess.GetDataFast(FlattenHeightX + 1, FlattenHeightY,     FlattenHeightX + 1, FlattenHeightY,     &P10);
 				this->Cache.DataAccess.GetDataFast(FlattenHeightX,     FlattenHeightY + 1, FlattenHeightX,     FlattenHeightY + 1, &P01);
 				this->Cache.DataAccess.GetDataFast(FlattenHeightX + 1, FlattenHeightY + 1, FlattenHeightX + 1, FlattenHeightY + 1, &P11);
+				// Release Texture Mips that will be Locked by the next SynchronousUpdateComponentVisibilityForHeight (inside the LayerDataCache.Read call)
+				this->Cache.DataAccess.Flush();
 				V00 = P00;
 				V10 = P10;
 				V01 = P01;
diff --git a/Engine/Source/Editor/LandscapeEditor/Private/LandscapeEdModeTools.h b/Engine/Source/Editor/LandscapeEditor/Private/LandscapeEdModeTools.h
--- a/Engine/Source/Editor/LandscapeEditor/Private/LandscapeEdModeTools.h
+++ b/Engine/Source/Editor/LandscapeEditor/Private/LandscapeEdModeTools.h
@@ -924,6 +924,8 @@ public:
 				SynchronousUpdateComponentVisibilityForHeight(AffectedComponents, NewLayerVisibility);
 				return Bounds;
 			});
+			// Release Texture Mips that will be Locked by the next SynchronousUpdateComponentVisibilityForHeight
+			CacheUpToEditingLayer.DataAccess.Flush();
 
 			CacheBottomLayers.GetDataAndCache(X1, Y1, X2, Y2, BottomLayersData, [&]() -> FIntRect
 			{
@@ -933,6 +935,8 @@ public:
 				SynchronousUpdateComponentVisibilityForHeight(AffectedComponents, NewLayerVisibility);
 				return Bounds;
 			});
+			// Do the same here for consistency
+			CacheBottomLayers.DataAccess.Flush();
 		}
 		else
 		{
