Fix Smooth/Flatten when using Landscape.Optim = 1. (mismatch between Cached Region and Dirtied Region)

diff --git a/Engine/Source/Editor/LandscapeEditor/Private/LandscapeEdModeTools.h b/Engine/Source/Editor/LandscapeEditor/Private/LandscapeEdModeTools.h
--- a/Engine/Source/Editor/LandscapeEditor/Private/LandscapeEdModeTools.h
+++ b/Engine/Source/Editor/LandscapeEditor/Private/LandscapeEdModeTools.h
@@ -920,7 +920,7 @@ public:
 			CacheUpToEditingLayer.GetDataAndCache(X1, Y1, X2, Y2, Data, [&]() -> FIntRect
 			{
 				TSet<ULandscapeComponent*> AffectedComponents;
-				LandscapeInfo->GetComponentsInRegion(X1, Y1, X2, Y2, AffectedComponents);
+				LandscapeInfo->GetComponentsInRegion(Bounds.Min.X, Bounds.Min.Y, Bounds.Max.X, Bounds.Max.Y, AffectedComponents);
 				SynchronousUpdateComponentVisibilityForHeight(AffectedComponents, NewLayerVisibility);
 				return Bounds;
 			});
@@ -931,7 +931,7 @@ public:
 			{
 				NewLayerVisibility[EditingLayerIndex] = false;
 				TSet<ULandscapeComponent*> AffectedComponents;
-				LandscapeInfo->GetComponentsInRegion(X1, Y1, X2, Y2, AffectedComponents);
+				LandscapeInfo->GetComponentsInRegion(Bounds.Min.X, Bounds.Min.Y, Bounds.Max.X, Bounds.Max.Y, AffectedComponents);
 				SynchronousUpdateComponentVisibilityForHeight(AffectedComponents, NewLayerVisibility);
 				return Bounds;
 			});
