LandscapeLayers: Make sure Intermediate Render used by Flatten/Smooth doesn't dirty Landscape on Read

#rnx

diff --git a/Engine/Source/Editor/LandscapeEditor/Private/LandscapeEdModeTools.h b/Engine/Source/Editor/LandscapeEditor/Private/LandscapeEdModeTools.h
--- a/Engine/Source/Editor/LandscapeEditor/Private/LandscapeEdModeTools.h
+++ b/Engine/Source/Editor/LandscapeEditor/Private/LandscapeEdModeTools.h
@@ -975,7 +975,8 @@ public:
 				TSet<ULandscapeComponent*> AffectedComponents;
 				LandscapeInfo->GetComponentsInRegion(X1, Y1, X2, Y2, AffectedComponents);
                 const bool bUpdateCollision = true;
-				SynchronousUpdateHeightmapForComponents(AffectedComponents, bUpdateCollision);
+				const bool bIntermediateRender = false;
+				SynchronousUpdateHeightmapForComponents(AffectedComponents, bUpdateCollision, bIntermediateRender);
 				bVisibilityChanged = false;
 			}
 		}
@@ -988,14 +989,14 @@ public:
 
 private:
 
-	void SynchronousUpdateHeightmapForComponents(const TSet<ULandscapeComponent*>& InComponents, bool bUpdateCollision)
+	void SynchronousUpdateHeightmapForComponents(const TSet<ULandscapeComponent*>& InComponents, bool bUpdateCollision, bool bIntermediateRender)
 	{
 		for (ULandscapeComponent* Component : InComponents)
 		{
 			const bool bUpdateAll = false; // default value
 			Component->RequestHeightmapUpdate(bUpdateAll, bUpdateCollision);
 		}
-		Landscape->ForceUpdateLayersContent();
+		Landscape->ForceUpdateLayersContent(bIntermediateRender);
 	};
 
 	void SetLayersVisibility(const TArray<bool>& InLayerVisibility)
@@ -1019,7 +1020,8 @@ private:
 		SetLayersVisibility(InLayerVisibility);
         // No need to update collision here as we are only doing a intermediate render to gather heightdata
         const bool bUpdateCollision = false;
-		SynchronousUpdateHeightmapForComponents(InComponents, bUpdateCollision);
+		const bool bIntermediateRender = true;
+		SynchronousUpdateHeightmapForComponents(InComponents, bUpdateCollision, bIntermediateRender);
 	};
 
 	ULandscapeInfo* LandscapeInfo;
diff --git a/Engine/Source/Runtime/Landscape/Classes/Landscape.h b/Engine/Source/Runtime/Landscape/Classes/Landscape.h
--- a/Engine/Source/Runtime/Landscape/Classes/Landscape.h
+++ b/Engine/Source/Runtime/Landscape/Classes/Landscape.h
@@ -326,7 +326,7 @@ public:
 	void ReleaseLayersRenderingResource();
 	
 	LANDSCAPE_API void ToggleCanHaveLayersContent();
-	LANDSCAPE_API void ForceUpdateLayersContent();
+	LANDSCAPE_API void ForceUpdateLayersContent(bool bIntermediateRender = false);
 
 private:
 	void TickLayers(float DeltaTime, ELevelTick TickType, FActorTickFunction& ThisTickFunction);
@@ -451,6 +451,9 @@ private:
 
 	FLandscapeEdModeInfo LandscapeEdModeInfo;
 
+	/** Some tools need to do an intermediate render with hidden layers. Do not dirty the landscape for those renders. */
+	bool bIntermediateRender;
+
 	UPROPERTY(Transient)
 	bool bLandscapeLayersAreInitialized;
 	
diff --git a/Engine/Source/Runtime/Landscape/Private/Landscape.cpp b/Engine/Source/Runtime/Landscape/Private/Landscape.cpp
--- a/Engine/Source/Runtime/Landscape/Private/Landscape.cpp
+++ b/Engine/Source/Runtime/Landscape/Private/Landscape.cpp
@@ -1103,6 +1103,7 @@ ALandscape::ALandscape(const FObjectInitializer& ObjectInitializer)
 	WeightmapScratchPackLayerTextureResource = nullptr;
 	bLandscapeLayersAreInitialized = false;
 	LandscapeEdMode = nullptr;	
+	bIntermediateRender = false;
 	bGrassUpdateEnabled = true;
 #endif // WITH_EDITORONLY_DATA
 }
diff --git a/Engine/Source/Runtime/Landscape/Private/LandscapeEditLayers.cpp b/Engine/Source/Runtime/Landscape/Private/LandscapeEditLayers.cpp
--- a/Engine/Source/Runtime/Landscape/Private/LandscapeEditLayers.cpp
+++ b/Engine/Source/Runtime/Landscape/Private/LandscapeEditLayers.cpp
@@ -3044,12 +3044,13 @@ bool ALandscape::ResolveLayersTexture(FLandscapeLayersTexture2DCPUReadBackResour
 	FlushRenderingCommands();
 
 	bool bChanged = false;
+    const bool bUpdateHash = !bIntermediateRender;
 	for (int8 MipIndex = 0; MipIndex < OutMipsData.Num(); ++MipIndex)
 	{
 		if (OutMipsData[MipIndex].Num() > 0)
 		{
 			uint8* TextureData = InOutputTexture->Source.LockMip(MipIndex);
-			if (MipIndex == 0)
+			if (MipIndex == 0 && bUpdateHash)
 			{
 				bChanged = InCPUReadBackTexture->UpdateHashFromTextureSource((uint8*)OutMipsData[MipIndex].GetData());
 			}
@@ -4823,11 +4824,14 @@ void ALandscape::OnPreSave()
 	ForceUpdateLayersContent();
 }
 
-void ALandscape::ForceUpdateLayersContent()
+void ALandscape::ForceUpdateLayersContent(bool bInIntermediateRender)
 {
 	const bool bWaitForStreaming = true;
 	const bool bInSkipMonitorLandscapeEdModeChanges = true;
+
+	bIntermediateRender = bInIntermediateRender;
 	UpdateLayersContent(bWaitForStreaming, bInSkipMonitorLandscapeEdModeChanges);
+	bInIntermediateRender = false;
 }
 
 void ALandscape::TickLayers(float DeltaTime, ELevelTick TickType, FActorTickFunction& ThisTickFunction)
