Fix duplicate Segment Mesh after Move to Current Level operation

#rnx

diff --git a/Engine/Source/Editor/LandscapeEditor/Private/LandscapeEdModeSplineTools.cpp b/Engine/Source/Editor/LandscapeEditor/Private/LandscapeEdModeSplineTools.cpp
--- a/Engine/Source/Editor/LandscapeEditor/Private/LandscapeEdModeSplineTools.cpp
+++ b/Engine/Source/Editor/LandscapeEditor/Private/LandscapeEdModeSplineTools.cpp
@@ -896,6 +896,26 @@ public:
 						FromProxy->SplineComponent->MarkRenderStateDirty();
 					}
 
+					// Delete all Mesh Components associated with the ControlPoint. (Will get recreated in UpdateSplinePoints)
+					if (ControlPoint->LocalMeshComponent)
+					{
+						ControlPoint->LocalMeshComponent->Modify();
+						ControlPoint->LocalMeshComponent->UnregisterComponent();
+						ControlPoint->LocalMeshComponent->DestroyComponent();
+						FromProxy->SplineComponent->Modify();
+						FromProxy->SplineComponent->MeshComponentLocalOwnersMap.Remove(ControlPoint->LocalMeshComponent);
+						ControlPoint->LocalMeshComponent = nullptr;
+					}
+
+					TMap<ULandscapeSplinesComponent*, UControlPointMeshComponent*> ForeignMeshComponents = ControlPoint->GetForeignMeshComponents();
+					for (auto Pair : ForeignMeshComponents)
+					{
+						Pair.Key->RemoveForeignMeshComponent(ControlPoint, Pair.Value);
+						Pair.Value->Modify();
+						Pair.Value->UnregisterComponent();
+						Pair.Value->DestroyComponent();
+					}
+					
 					// Move control point to new level
 					FromProxy->SplineComponent->ControlPoints.Remove(ControlPoint);
 					ControlPoint->Rename(nullptr, ToLandscape->SplineComponent);
@@ -905,7 +925,7 @@ public:
 
 					const bool bUpdateCollision = true; // default value
 					const bool bUpdateSegments = false; // done in next loop
-					const bool bUpdateMeshLevel = true;
+					const bool bUpdateMeshLevel = false; // no need because mesh have been deleted
 					ControlPoint->UpdateSplinePoints(bUpdateCollision, bUpdateSegments, bUpdateMeshLevel);
 				}
 			}
@@ -947,13 +967,36 @@ public:
 						FromProxy->SplineComponent->MarkRenderStateDirty();
 					}
 
+					// Delete all Mesh Components associated with the Segment. (Will get recreated in UpdateSplinePoints)
+					for (auto* MeshComponent : Segment->LocalMeshComponents)
+					{
+						MeshComponent->Modify();
+						MeshComponent->UnregisterComponent();
+						MeshComponent->DestroyComponent();
+						FromProxy->Modify();
+						FromProxy->SplineComponent->MeshComponentLocalOwnersMap.Remove(MeshComponent);
+					}
+					Segment->LocalMeshComponents.Empty();
+
+					TMap<ULandscapeSplinesComponent*, TArray<USplineMeshComponent*>> ForeignMeshComponents = Segment->GetForeignMeshComponents();
+					for (auto Pair : ForeignMeshComponents)
+					{
+						Pair.Key->RemoveAllForeignMeshComponents(Segment);
+						for (auto MeshComponent : Pair.Value)
+						{
+							MeshComponent->Modify();
+							MeshComponent->UnregisterComponent();
+							MeshComponent->DestroyComponent();
+						}
+					}
+										
 					// Move segment to new level
 					FromProxy->SplineComponent->Segments.Remove(Segment);
 					Segment->Rename(nullptr, ToLandscape->SplineComponent);
 					ToLandscape->SplineComponent->Segments.Add(Segment);
 
 					const bool bUpdateCollision = true; // default value
-					const bool bUpdateMeshLevel = true;
+					const bool bUpdateMeshLevel = false; // no need because mesh have been deleted 
 					Segment->UpdateSplinePoints(bUpdateCollision, bUpdateMeshLevel);
 				}
 			}
diff --git a/Engine/Source/Runtime/Landscape/Classes/LandscapeSplineControlPoint.h b/Engine/Source/Runtime/Landscape/Classes/LandscapeSplineControlPoint.h
--- a/Engine/Source/Runtime/Landscape/Classes/LandscapeSplineControlPoint.h
+++ b/Engine/Source/Runtime/Landscape/Classes/LandscapeSplineControlPoint.h
@@ -215,7 +215,7 @@ public:
 
 	virtual void AutoSetConnections(bool bIncludingValid);
 
-	TMap<ULandscapeSplinesComponent*, UControlPointMeshComponent*> GetForeignMeshComponents();
+	LANDSCAPE_API TMap<ULandscapeSplinesComponent*, UControlPointMeshComponent*> GetForeignMeshComponents();
 
 	/** Update spline points */
 	virtual void UpdateSplinePoints(bool bUpdateCollision = true, bool bUpdateAttachedSegments = true, bool bUpdateMeshLevel = false);
diff --git a/Engine/Source/Runtime/Landscape/Classes/LandscapeSplineSegment.h b/Engine/Source/Runtime/Landscape/Classes/LandscapeSplineSegment.h
--- a/Engine/Source/Runtime/Landscape/Classes/LandscapeSplineSegment.h
+++ b/Engine/Source/Runtime/Landscape/Classes/LandscapeSplineSegment.h
@@ -317,7 +317,7 @@ public:
 
 	virtual void AutoFlipTangents();
 
-	TMap<ULandscapeSplinesComponent*, TArray<USplineMeshComponent*>> GetForeignMeshComponents();
+	LANDSCAPE_API TMap<ULandscapeSplinesComponent*, TArray<USplineMeshComponent*>> GetForeignMeshComponents();
 	TArray<USplineMeshComponent*> GetLocalMeshComponents() const;
 	
 	virtual void UpdateSplinePoints(bool bUpdateCollision = true, bool bUpdateMeshLevel = false);
diff --git a/Engine/Source/Runtime/Landscape/Private/LandscapeSplines.cpp b/Engine/Source/Runtime/Landscape/Private/LandscapeSplines.cpp
--- a/Engine/Source/Runtime/Landscape/Private/LandscapeSplines.cpp
+++ b/Engine/Source/Runtime/Landscape/Private/LandscapeSplines.cpp
@@ -2367,7 +2367,8 @@ void ULandscapeSplineSegment::UpdateSplinePoints(bool bUpdateCollision, bool bUp
 	// Unregister components, Remove Foreign/Local Associations
 	for (auto* LocalMeshComponent : OldLocalMeshComponents)
 	{
-		OuterSplines->MeshComponentLocalOwnersMap.Remove(LocalMeshComponent);
+		checkSlow(OuterSplines->MeshComponentLocalOwnersMap.FindRef(LocalMeshComponent) == this);
+		verifySlow(OuterSplines->MeshComponentLocalOwnersMap.Remove(LocalMeshComponent) == 1);
 		LocalMeshComponent->Modify();
 		LocalMeshComponent->UnregisterComponent();
 	}
@@ -2730,8 +2731,6 @@ void ULandscapeSplineSegment::UpdateSplinePoints(bool bUpdateCollision, bool bUp
 	// Clean up unused components
 	for (auto* LocalMeshComponent : OldLocalMeshComponents)
 	{
-		checkSlow(OuterSplines->MeshComponentLocalOwnersMap.FindRef(LocalMeshComponent) == this);
-		verifySlow(OuterSplines->MeshComponentLocalOwnersMap.Remove(LocalMeshComponent) == 1);
 		LocalMeshComponent->DestroyComponent();
 	}
 	OldLocalMeshComponents.Empty();
