Landscape Spline Editing: Fix Split Segment sometimes placing new control point under one of the existing Segments control points.

diff --git a/Engine/Source/Editor/LandscapeEditor/Private/LandscapeEdModeSplineTools.cpp b/Engine/Source/Editor/LandscapeEditor/Private/LandscapeEdModeSplineTools.cpp
--- a/Engine/Source/Editor/LandscapeEditor/Private/LandscapeEdModeSplineTools.cpp
+++ b/Engine/Source/Editor/LandscapeEditor/Private/LandscapeEdModeSplineTools.cpp
@@ -1443,6 +1443,8 @@ public:
 									if (ULandscapeSplineSegment* SplineSegment = Cast<ULandscapeSplineSegment>(ComponentOwner))
 									{
 										ClickedSplineSegment = SplineSegment;
+										// Find actual SplineComponent owner of the Segment (not the SplineComponent owner of the mesh)
+										SplineComponent = SplineSegment->GetTypedOuter<ULandscapeSplinesComponent>();
 										ALandscapeProxy* LandscapeProxy = CastChecked<ALandscapeProxy>(SplineComponent->GetOwner());
 										LandscapeToSpline = LandscapeProxy->LandscapeActorToWorld().GetRelativeTransform(SplineComponent->GetComponentTransform());
 									}
