Landscape Layers: Fix Crash when passing empty component list (when Landscape.Optim == 1) to the layer render

diff --git a/Engine/Source/Editor/LandscapeEditor/Private/LandscapeEdModeComponentTools.cpp b/Engine/Source/Editor/LandscapeEditor/Private/LandscapeEdModeComponentTools.cpp
--- a/Engine/Source/Editor/LandscapeEditor/Private/LandscapeEdModeComponentTools.cpp
+++ b/Engine/Source/Editor/LandscapeEditor/Private/LandscapeEdModeComponentTools.cpp
@@ -130,6 +130,21 @@ public:
 	{
 	}
 
+	virtual ELandscapeLayerUpdateMode GetBeginToolContentUpdateFlag() const override
+	{
+		return ELandscapeLayerUpdateMode::None;
+	}
+
+	virtual ELandscapeLayerUpdateMode GetTickToolContentUpdateFlag() const override
+	{
+		return ELandscapeLayerUpdateMode::None;
+	}
+
+	virtual ELandscapeLayerUpdateMode GetEndToolContentUpdateFlag() const override
+	{
+		return ELandscapeLayerUpdateMode::None;
+	}
+
 	virtual const TCHAR* GetToolName() override { return TEXT("Select"); }
 	virtual FText GetDisplayName() override { return NSLOCTEXT("UnrealEd", "LandscapeMode_Selection", "Component Selection"); };
 	virtual void SetEditRenderType() override { GLandscapeEditRenderMode = ELandscapeEditRenderMode::SelectComponent | (GLandscapeEditRenderMode & ELandscapeEditRenderMode::BitMaskForMask); }
diff --git a/Engine/Source/Runtime/Landscape/Classes/LandscapeComponent.h b/Engine/Source/Runtime/Landscape/Classes/LandscapeComponent.h
--- a/Engine/Source/Runtime/Landscape/Classes/LandscapeComponent.h
+++ b/Engine/Source/Runtime/Landscape/Classes/LandscapeComponent.h
@@ -315,12 +315,16 @@ enum ELandscapeComponentUpdateFlag : uint32
 
 enum ELandscapeLayerUpdateMode : uint32
 { 
+	// No Update
+	None = 0,
+	// Update types
 	Update_Heightmap_All = 1 << 0,
 	Update_Heightmap_Editing = 1 << 1,
 	Update_Heightmap_Editing_NoCollision = 1 << 2,
 	Update_Weightmap_All = 1 << 3,
 	Update_Weightmap_Editing = 1 << 4,
 	Update_Weightmap_Editing_NoCollision = 1 << 5,
+	// Combinations
 	Update_All = Update_Weightmap_All | Update_Heightmap_All,
 	Update_All_Editing = Update_Weightmap_Editing | Update_Heightmap_Editing,
 	Update_All_Editing_NoCollision = Update_Weightmap_Editing_NoCollision | Update_Heightmap_Editing_NoCollision,
diff --git a/Engine/Source/Runtime/Landscape/Private/LandscapeEditLayers.cpp b/Engine/Source/Runtime/Landscape/Private/LandscapeEditLayers.cpp
--- a/Engine/Source/Runtime/Landscape/Private/LandscapeEditLayers.cpp
+++ b/Engine/Source/Runtime/Landscape/Private/LandscapeEditLayers.cpp
@@ -2721,6 +2721,12 @@ int32 ALandscape::RegenerateLayersHeightmaps(const TArray<ULandscapeComponent*>&
 		return 0;
 	}
 
+	// Nothing to do (return that we did the processing)
+	if (InLandscapeComponentsToResolve.Num() == 0)
+	{
+		return HeightmapUpdateModes;
+	}
+
 	// Init CPU Readbacks
 	if (HeightmapUpdateModes)
 	{
@@ -3556,6 +3562,11 @@ int32 ALandscape::RegenerateLayersWeightmaps(const TArray<ULandscapeComponent*>&
 	{
 		return 0;
 	}
+
+	if (InLandscapeComponentsToResolve.Num() == 0)
+	{
+		return WeightmapUpdateModes;
+	}
 			
 	TArray<ULandscapeComponent*> ComponentThatNeedMaterialRebuild;
 	TArray<ULandscapeLayerInfoObject*> BrushRequiredAllocations;
