Landscape Layers: Quality of life : Auto select the Spline tool when selecting the Spline Layer

diff --git a/Engine/Source/Editor/LandscapeEditor/Private/LandscapeEdMode.cpp b/Engine/Source/Editor/LandscapeEditor/Private/LandscapeEdMode.cpp
--- a/Engine/Source/Editor/LandscapeEditor/Private/LandscapeEdMode.cpp
+++ b/Engine/Source/Editor/LandscapeEditor/Private/LandscapeEdMode.cpp
@@ -4375,6 +4375,17 @@ void FEdModeLandscape::SetCurrentLayer(int32 InLayerIndex)
 {
 	UISettings->Modify();
 	UISettings->CurrentLayerIndex = InLayerIndex;
+
+	if (ALandscape* Landscape = GetLandscape())
+	{
+		const FLandscapeLayer* SplineLayer = Landscape->GetLandscapeSplinesReservedLayer();
+		if (SplineLayer != nullptr && SplineLayer == Landscape->GetLayer(InLayerIndex))
+		{
+			SetCurrentToolMode("ToolMode_Manage", false);
+			SetCurrentTool(FName("Splines"));
+		}
+	}
+
 	RefreshDetailPanel();
 	RequestLayersContentUpdateForceAll(ELandscapeLayerUpdateMode::Update_Client_Editing);
 }
