Added option to bind hotkeys for increment/decrement of landscape alpha brush rotation

diff --git a/Engine/Source/Editor/LandscapeEditor/Private/LandscapeEdMode.cpp b/Engine/Source/Editor/LandscapeEditor/Private/LandscapeEdMode.cpp
--- a/Engine/Source/Editor/LandscapeEditor/Private/LandscapeEdMode.cpp
+++ b/Engine/Source/Editor/LandscapeEditor/Private/LandscapeEdMode.cpp
@@ -366,6 +366,8 @@ FEdModeLandscape::FEdModeLandscape()
 	CommandList->MapAction(LandscapeActions.DecreaseBrushFalloff, FExecuteAction::CreateRaw(this, &FEdModeLandscape::ChangeBrushFalloff, false), FCanExecuteAction(), FIsActionChecked());
 	CommandList->MapAction(LandscapeActions.IncreaseBrushStrength, FExecuteAction::CreateRaw(this, &FEdModeLandscape::ChangeBrushStrength, true), FCanExecuteAction(), FIsActionChecked());
 	CommandList->MapAction(LandscapeActions.DecreaseBrushStrength, FExecuteAction::CreateRaw(this, &FEdModeLandscape::ChangeBrushStrength, false), FCanExecuteAction(), FIsActionChecked());
+	CommandList->MapAction(LandscapeActions.IncreaseAlphaBrushRotation, FExecuteAction::CreateRaw(this, &FEdModeLandscape::ChangeAlphaBrushRotation, true), FCanExecuteAction(), FIsActionChecked());
+	CommandList->MapAction(LandscapeActions.DecreaseAlphaBrushRotation, FExecuteAction::CreateRaw(this, &FEdModeLandscape::ChangeAlphaBrushRotation, false), FCanExecuteAction(), FIsActionChecked());
 }
 
 /** Destructor */
@@ -1904,6 +1906,16 @@ void FEdModeLandscape::ChangeBrushStrength(bool bIncrease)
 	UISettings->ToolStrength = NewValue;
 }
 
+void FEdModeLandscape::ChangeAlphaBrushRotation(bool bIncrease)
+{
+	const float SliderMin = -180.f;
+	const float SliderMax = 180.f;
+	const float DefaultDiffValue = 10.f;
+	const float Diff = bIncrease ? DefaultDiffValue : -DefaultDiffValue;
+	UISettings->Modify();
+	UISettings->AlphaBrushRotation = FMath::Clamp(UISettings->AlphaBrushRotation + Diff, SliderMin, SliderMax);
+}
+
 
 /** FEdMode: Called when a key is pressed */
 bool FEdModeLandscape::InputKey(FEditorViewportClient* ViewportClient, FViewport* Viewport, FKey Key, EInputEvent Event)
diff --git a/Engine/Source/Editor/LandscapeEditor/Private/LandscapeEdMode.h b/Engine/Source/Editor/LandscapeEditor/Private/LandscapeEdMode.h
--- a/Engine/Source/Editor/LandscapeEditor/Private/LandscapeEdMode.h
+++ b/Engine/Source/Editor/LandscapeEditor/Private/LandscapeEdMode.h
@@ -418,6 +418,7 @@ public:
 	void ChangeBrushSize(bool bIncrease);
 	void ChangeBrushFalloff(bool bIncrease);
 	void ChangeBrushStrength(bool bIncrease);
+	void ChangeAlphaBrushRotation(bool bIncrease);
 
 	/** FEdMode: Called when a key is pressed */
 	virtual bool InputKey(FEditorViewportClient* InViewportClient, FViewport* InViewport, FKey InKey, EInputEvent InEvent) override;
diff --git a/Engine/Source/Editor/LandscapeEditor/Private/LandscapeEditorCommands.cpp b/Engine/Source/Editor/LandscapeEditor/Private/LandscapeEditorCommands.cpp
--- a/Engine/Source/Editor/LandscapeEditor/Private/LandscapeEditorCommands.cpp
+++ b/Engine/Source/Editor/LandscapeEditor/Private/LandscapeEditorCommands.cpp
@@ -114,6 +114,8 @@ void FLandscapeEditorCommands::RegisterCommands()
 	UI_COMMAND(DecreaseBrushFalloff, "Decrease Brush Falloff", "Press this key to decrease brush falloff by a fixed increment.", EUserInterfaceActionType::RadioButton, FInputChord());
 	UI_COMMAND(IncreaseBrushStrength, "Increase Brush Strength", "Press this key to increase brush strength by a fixed increment.", EUserInterfaceActionType::RadioButton, FInputChord());
 	UI_COMMAND(DecreaseBrushStrength, "Decrease Brush Strength", "Press this key to decrease brush strength by a fixed increment.", EUserInterfaceActionType::RadioButton, FInputChord());
+	UI_COMMAND(IncreaseAlphaBrushRotation, "Increase Alpha Brush Rotation", "Press this key to increase alpha brush rotation by a fixed increment.", EUserInterfaceActionType::Button, FInputChord());
+	UI_COMMAND(DecreaseAlphaBrushRotation, "Decrease Alpha Brush Rotation", "Press this key to decrease alpha brush rotation by a fixed increment.", EUserInterfaceActionType::Button, FInputChord());
 
 	UI_COMMAND(DragBrushSize, "Change Brush Size", "Hold this key and then click and drag to increase or decrease brush size.", EUserInterfaceActionType::RadioButton, FInputChord());
 	UI_COMMAND(DragBrushFalloff, "Change Brush Falloff", "Hold this key and then click and drag to increase or decrease brush falloff.", EUserInterfaceActionType::RadioButton, FInputChord());
diff --git a/Engine/Source/Editor/LandscapeEditor/Private/LandscapeEditorCommands.h b/Engine/Source/Editor/LandscapeEditor/Private/LandscapeEditorCommands.h
--- a/Engine/Source/Editor/LandscapeEditor/Private/LandscapeEditorCommands.h
+++ b/Engine/Source/Editor/LandscapeEditor/Private/LandscapeEditorCommands.h
@@ -84,6 +84,8 @@ public:
 	TSharedPtr<FUICommandInfo> DecreaseBrushFalloff;
 	TSharedPtr<FUICommandInfo> IncreaseBrushStrength;
 	TSharedPtr<FUICommandInfo> DecreaseBrushStrength;
+	TSharedPtr<FUICommandInfo> IncreaseAlphaBrushRotation;
+	TSharedPtr<FUICommandInfo> DecreaseAlphaBrushRotation;
 
 	TSharedPtr<FUICommandInfo> DragBrushSize;
 	TSharedPtr<FUICommandInfo> DragBrushFalloff;
