Fix LandscapeComponent.h Linux Compile warning

#rnx

diff --git a/Engine/Source/Editor/LandscapeEditor/Private/LandscapeEdModeComponentTools.cpp b/Engine/Source/Editor/LandscapeEditor/Private/LandscapeEdModeComponentTools.cpp
--- a/Engine/Source/Editor/LandscapeEditor/Private/LandscapeEdModeComponentTools.cpp
+++ b/Engine/Source/Editor/LandscapeEditor/Private/LandscapeEdModeComponentTools.cpp
@@ -132,17 +132,17 @@ public:
 
 	virtual ELandscapeLayerUpdateMode GetBeginToolContentUpdateFlag() const override
 	{
-		return ELandscapeLayerUpdateMode::None;
+		return ELandscapeLayerUpdateMode::Update_None;
 	}
 
 	virtual ELandscapeLayerUpdateMode GetTickToolContentUpdateFlag() const override
 	{
-		return ELandscapeLayerUpdateMode::None;
+		return ELandscapeLayerUpdateMode::Update_None;
 	}
 
 	virtual ELandscapeLayerUpdateMode GetEndToolContentUpdateFlag() const override
 	{
-		return ELandscapeLayerUpdateMode::None;
+		return ELandscapeLayerUpdateMode::Update_None;
 	}
 
 	virtual const TCHAR* GetToolName() override { return TEXT("Select"); }
diff --git a/Engine/Source/Runtime/Landscape/Classes/LandscapeComponent.h b/Engine/Source/Runtime/Landscape/Classes/LandscapeComponent.h
--- a/Engine/Source/Runtime/Landscape/Classes/LandscapeComponent.h
+++ b/Engine/Source/Runtime/Landscape/Classes/LandscapeComponent.h
@@ -316,7 +316,7 @@ enum ELandscapeComponentUpdateFlag : uint32
 enum ELandscapeLayerUpdateMode : uint32
 { 
 	// No Update
-	None = 0,
+	Update_None = 0,
 	// Update types
 	Update_Heightmap_All = 1 << 0,
 	Update_Heightmap_Editing = 1 << 1,
