Fix check when switching selection between two landscapes in scene outliner

diff --git a/Engine/Source/Editor/LandscapeEditor/Private/LandscapeEdMode.cpp b/Engine/Source/Editor/LandscapeEditor/Private/LandscapeEdMode.cpp
--- a/Engine/Source/Editor/LandscapeEditor/Private/LandscapeEdMode.cpp
+++ b/Engine/Source/Editor/LandscapeEditor/Private/LandscapeEdMode.cpp
@@ -308,8 +308,9 @@ FEdModeLandscape::FEdModeLandscape()
 	, CachedLandscapeMaterial(nullptr)
 	, ToolActiveViewport(nullptr)
 	, bIsPaintingInVR(false)
-	, InteractorPainting( nullptr )
+	, InteractorPainting(nullptr)
 	, bNeedsUpdateShownLayerList(false)
+	, bUpdatingLandscapeInfo(false)
 {
 	GLayerDebugColorMaterial = LandscapeTool::CreateMaterialInstance(LoadObject<UMaterial>(nullptr, TEXT("/Engine/EditorLandscapeResources/LayerVisMaterial.LayerVisMaterial")));
 	GSelectionColorMaterial  = LandscapeTool::CreateMaterialInstance(LoadObject<UMaterialInstanceConstant>(nullptr, TEXT("/Engine/EditorLandscapeResources/SelectBrushMaterial_Selected.SelectBrushMaterial_Selected")));
@@ -550,8 +551,14 @@ void FEdModeLandscape::SetLandscapeInfo(ULandscapeInfo* InLandscapeInfo)
 {
 	if (CurrentToolTarget.LandscapeInfo != InLandscapeInfo)
 	{
-		CurrentToolTarget.LandscapeInfo = InLandscapeInfo;
-		UpdateToolModes();
+		{
+			TGuardValue<bool> GuardFlag(bUpdatingLandscapeInfo, true);
+			CurrentToolTarget.LandscapeInfo = InLandscapeInfo;
+			UpdateTargetList();
+			UpdateShownLayerList();
+			UpdateToolModes();
+		}
+		RefreshDetailPanel();
 	}
 }
 
@@ -609,7 +616,6 @@ void FEdModeLandscape::Enter()
 	// Update list of landscapes and layers
 	// For now depends on the SpawnActor() above in order to get the current editor world as edmodes don't get told
 	UpdateLandscapeList();
-	UpdateTargetList();
 	UpdateBrushList();
 
 	OnWorldChangeDelegateHandle                 = FEditorSupportDelegates::WorldChange.AddRaw(this, &FEdModeLandscape::HandleLevelsChanged, true);
@@ -2393,7 +2399,7 @@ void FEdModeLandscape::SetCurrentTool(int32 ToolIndex, FName TargetLayerName)
 
 void FEdModeLandscape::RefreshDetailPanel()
 {
-	if (Toolkit.IsValid())
+	if (Toolkit.IsValid() && !bUpdatingLandscapeInfo)
 	{
 		StaticCastSharedPtr<FLandscapeToolKit>(Toolkit)->RefreshDetailPanel();
 	}
@@ -2551,15 +2557,6 @@ int32 FEdModeLandscape::UpdateLandscapeList()
 
 			SetCurrentLayer(0);
 
-			// Init UI to saved value
-			ALandscapeProxy* LandscapeProxy = CurrentToolTarget.LandscapeInfo->GetLandscapeProxy();
-
-			if (LandscapeProxy != nullptr)
-			{
-				UISettings->TargetDisplayOrder = LandscapeProxy->TargetDisplayOrder;
-			}
-
-			UpdateTargetList();
 			UpdateShownLayerList();
 						
 			if (!CurrentToolName.IsNone())
@@ -2571,7 +2568,6 @@ int32 FEdModeLandscape::UpdateLandscapeList()
 		{
 			// no landscape, switch to "new landscape" tool
 			SetLandscapeInfo(nullptr);
-			UpdateTargetList();
 			SetCurrentToolMode("ToolMode_Manage", false);
 			SetCurrentTool("NewLandscape");
 		}
@@ -2605,7 +2601,6 @@ void FEdModeLandscape::SetTargetLandscape(const TWeakObjectPtr<ULandscapeInfo>&
 	}
 
 	SetLandscapeInfo(InLandscapeInfo.Get());
-	UpdateTargetList();
 	// force a Leave and Enter the current tool, in case it has something about the current landscape cached
 	SetCurrentTool(CurrentToolIndex);
 	if (CurrentGizmoActor.IsValid())
@@ -2624,7 +2619,6 @@ void FEdModeLandscape::SetTargetLandscape(const TWeakObjectPtr<ULandscapeInfo>&
 		}
 	}
 
-	UpdateTargetList();
 	UpdateShownLayerList();
 }
 
@@ -2743,6 +2737,8 @@ void FEdModeLandscape::UpdateTargetList()
 				SetCurrentTargetLayer(NAME_None, nullptr);
 			}
 
+			UISettings->TargetDisplayOrder = LandscapeProxy->TargetDisplayOrder;
+
 			UpdateTargetLayerDisplayOrder(UISettings->TargetDisplayOrder);
 		}
 	}
@@ -3431,8 +3427,7 @@ bool FEdModeLandscape::Select(AActor* InActor, bool bInSelected)
 		if (CurrentToolTarget.LandscapeInfo != Landscape->GetLandscapeInfo())
 		{
 			SetLandscapeInfo(Landscape->GetLandscapeInfo());
-			UpdateTargetList();
-
+		
 			// If we were in "New Landscape" mode and we select a landscape then switch to editing mode
 			if (NewLandscapePreviewMode != ENewLandscapePreviewMode::None)
 			{
diff --git a/Engine/Source/Editor/LandscapeEditor/Private/LandscapeEdMode.h b/Engine/Source/Editor/LandscapeEditor/Private/LandscapeEdMode.h
--- a/Engine/Source/Editor/LandscapeEditor/Private/LandscapeEdMode.h
+++ b/Engine/Source/Editor/LandscapeEditor/Private/LandscapeEdMode.h
@@ -648,4 +648,5 @@ private:
 
 	/** Delayed refresh */
 	bool bNeedsUpdateShownLayerList;
+	bool bUpdatingLandscapeInfo;
 };
\ No newline at end of file
diff --git a/Engine/Source/Editor/LandscapeEditor/Private/LandscapeEditorDetailCustomization_NewLandscape.cpp b/Engine/Source/Editor/LandscapeEditor/Private/LandscapeEditorDetailCustomization_NewLandscape.cpp
--- a/Engine/Source/Editor/LandscapeEditor/Private/LandscapeEditorDetailCustomization_NewLandscape.cpp
+++ b/Engine/Source/Editor/LandscapeEditor/Private/LandscapeEditorDetailCustomization_NewLandscape.cpp
@@ -873,8 +873,6 @@ FReply FLandscapeEditorDetailCustomization_NewLandscape::OnCreateButtonClicked()
 		LandscapeEdMode->SetLandscapeInfo(LandscapeInfo);
 		LandscapeEdMode->CurrentToolTarget.TargetType = ELandscapeToolTargetType::Heightmap;
 		LandscapeEdMode->SetCurrentTargetLayer(NAME_None, nullptr);
-		LandscapeEdMode->UpdateTargetList();
-
 		LandscapeEdMode->SetCurrentTool("Select"); // change tool so switching back to the manage mode doesn't give "New Landscape" again
 		LandscapeEdMode->SetCurrentTool("Sculpt"); // change to sculpting mode and tool
 		LandscapeEdMode->SetCurrentLayer(0);
