Landscape Layers:
- Fix bug in Heightmap Render causing visual glitch while painting
- Fix bug in Heightmap Flatten/Smooth tool causing visual glitch

diff --git a/Engine/Source/Editor/LandscapeEditor/Private/LandscapeEdModeTools.h b/Engine/Source/Editor/LandscapeEditor/Private/LandscapeEdModeTools.h
--- a/Engine/Source/Editor/LandscapeEditor/Private/LandscapeEdModeTools.h
+++ b/Engine/Source/Editor/LandscapeEditor/Private/LandscapeEdModeTools.h
@@ -922,6 +922,10 @@ public:
 				TSet<ULandscapeComponent*> AffectedComponents;
 				LandscapeInfo->GetComponentsInRegion(Bounds.Min.X, Bounds.Min.Y, Bounds.Max.X, Bounds.Max.Y, AffectedComponents);
 				SynchronousUpdateComponentVisibilityForHeight(AffectedComponents, NewLayerVisibility);
+				if (bVisibilityChanged)
+				{
+					VisibilityChangedComponents.Append(AffectedComponents);
+				}
 				return Bounds;
 			});
 			// Release Texture Mips that will be Locked by the next SynchronousUpdateComponentVisibilityForHeight
@@ -933,6 +937,10 @@ public:
 				TSet<ULandscapeComponent*> AffectedComponents;
 				LandscapeInfo->GetComponentsInRegion(Bounds.Min.X, Bounds.Min.Y, Bounds.Max.X, Bounds.Max.Y, AffectedComponents);
 				SynchronousUpdateComponentVisibilityForHeight(AffectedComponents, NewLayerVisibility);
+				if (bVisibilityChanged)
+				{
+					VisibilityChangedComponents.Append(AffectedComponents);
+				}
 				return Bounds;
 			});
 			// Do the same here for consistency
@@ -972,11 +980,10 @@ public:
 			CacheUpToEditingLayer.DataAccess.Flush();
 			if (bVisibilityChanged)
 			{
-				TSet<ULandscapeComponent*> AffectedComponents;
-				LandscapeInfo->GetComponentsInRegion(X1, Y1, X2, Y2, AffectedComponents);
-                const bool bUpdateCollision = true;
+				const bool bUpdateCollision = true;
 				const bool bIntermediateRender = false;
-				SynchronousUpdateHeightmapForComponents(AffectedComponents, bUpdateCollision, bIntermediateRender);
+				SynchronousUpdateHeightmapForComponents(VisibilityChangedComponents, bUpdateCollision, bIntermediateRender);
+				VisibilityChangedComponents.Empty();
 				bVisibilityChanged = false;
 			}
 		}
@@ -1034,6 +1041,7 @@ private:
 	bool bCombinedLayerOperation;
 	bool bVisibilityChanged;
 	bool bTargetIsHeightmap;
+	TSet<ULandscapeComponent*> VisibilityChangedComponents;
 
 	typename ToolTarget::CacheClass& CacheUpToEditingLayer;
 	typename ToolTarget::CacheClass CacheBottomLayers;
diff --git a/Engine/Source/Runtime/Landscape/Private/LandscapeEditLayers.cpp b/Engine/Source/Runtime/Landscape/Private/LandscapeEditLayers.cpp
--- a/Engine/Source/Runtime/Landscape/Private/LandscapeEditLayers.cpp
+++ b/Engine/Source/Runtime/Landscape/Private/LandscapeEditLayers.cpp
@@ -2964,7 +2964,7 @@ int32 ALandscape::RegenerateLayersHeightmaps(const TArray<ULandscapeComponent*>&
 		// Copy back all Mips to original heightmap data
 		{
 			TArray<FHeightmapTopLeft> Heightmaps;
-			GetUniqueHeightmaps(InLandscapeComponentsToRender, Heightmaps, LandscapeExtent.Min);
+			GetUniqueHeightmaps(InLandscapeComponentsToResolve, Heightmaps, LandscapeExtent.Min);
 			for (const FHeightmapTopLeft& Heightmap : Heightmaps)
 			{
 				FIntPoint TextureTopLeftSectionBase = Heightmap.TopLeftSectionBase;
