Landscape: Non-initialized member could be causing rare crash

#rnx

diff --git a/Engine/Source/Editor/LandscapeEditor/Private/LandscapeEdMode.cpp b/Engine/Source/Editor/LandscapeEditor/Private/LandscapeEdMode.cpp
--- a/Engine/Source/Editor/LandscapeEditor/Private/LandscapeEdMode.cpp
+++ b/Engine/Source/Editor/LandscapeEditor/Private/LandscapeEdMode.cpp
@@ -305,6 +305,7 @@ FEdModeLandscape::FEdModeLandscape()
 	, CopyPasteTool(nullptr)
 	, SplinesTool(nullptr)
 	, LandscapeRenderAddCollision(nullptr)
+	, TargetLayerStartingIndex(0)
 	, CachedLandscapeMaterial(nullptr)
 	, ToolActiveViewport(nullptr)
 	, bIsPaintingInVR(false)
