Fix Shown Layers not being updated (move call inside of UpdateTargetList)

#rnx

diff --git a/Engine/Source/Editor/LandscapeEditor/Private/LandscapeEdMode.cpp b/Engine/Source/Editor/LandscapeEditor/Private/LandscapeEdMode.cpp
--- a/Engine/Source/Editor/LandscapeEditor/Private/LandscapeEdMode.cpp
+++ b/Engine/Source/Editor/LandscapeEditor/Private/LandscapeEdMode.cpp
@@ -559,7 +559,6 @@ void FEdModeLandscape::SetLandscapeInfo(ULandscapeInfo* InLandscapeInfo)
 			TGuardValue<bool> GuardFlag(bUpdatingLandscapeInfo, true);
 			CurrentToolTarget.LandscapeInfo = InLandscapeInfo;
 			UpdateTargetList();
-			UpdateShownLayerList();
 			UpdateToolModes();
 		}
 		RefreshDetailPanel();
@@ -2757,6 +2756,8 @@ void FEdModeLandscape::UpdateTargetList()
 	}
 
 	TargetsListUpdated.Broadcast();
+
+	UpdateShownLayerList();
 }
 
 void FEdModeLandscape::UpdateTargetLayerDisplayOrder(ELandscapeLayerDisplayMode InTargetDisplayOrder)
@@ -2874,7 +2875,6 @@ void FEdModeLandscape::UpdateTargetLayerDisplayOrder(ELandscapeLayerDisplayMode
 void FEdModeLandscape::OnLandscapeMaterialChangedDelegate()
 {
 	UpdateTargetList();
-	UpdateShownLayerList();
 }
 
 void FEdModeLandscape::RequestUpdateShownLayerList()
@@ -3066,7 +3066,6 @@ void FEdModeLandscape::HandleLevelsChanged(bool ShouldExitMode)
 
 	UpdateLandscapeList();
 	UpdateTargetList();
-	UpdateShownLayerList();
 	UpdateBrushList();
 
 	// if the Landscape is deleted then close the landscape editor
@@ -3092,7 +3091,6 @@ void FEdModeLandscape::OnMaterialCompilationFinished(UMaterialInterface* Materia
 	{
 		CurrentToolTarget.LandscapeInfo->UpdateLayerInfoMap();
 		UpdateTargetList();
-		UpdateShownLayerList();
 	}
 }
 
diff --git a/Engine/Source/Editor/LandscapeEditor/Private/LandscapeEditorDetailCustomization_TargetLayers.cpp b/Engine/Source/Editor/LandscapeEditor/Private/LandscapeEditorDetailCustomization_TargetLayers.cpp
--- a/Engine/Source/Editor/LandscapeEditor/Private/LandscapeEditorDetailCustomization_TargetLayers.cpp
+++ b/Engine/Source/Editor/LandscapeEditor/Private/LandscapeEditorDetailCustomization_TargetLayers.cpp
@@ -1384,7 +1384,6 @@ FReply FLandscapeEditorCustomNodeBuilder_TargetLayers::OnTargetLayerDeleteClicke
 		if (LandscapeEdMode)
 		{
 			LandscapeEdMode->UpdateTargetList();
-			LandscapeEdMode->UpdateShownLayerList();
 		}
 	}
 
