Landscape Blueprint Brushes: Small changes for Ryan
- Duplicate brush
- Show only selected
- Show/Hide selected
- Selecting Brush changes tool to BlueprintBrushes

#rnx

diff --git a/Engine/Source/Editor/LandscapeEditor/Private/LandscapeEdMode.cpp b/Engine/Source/Editor/LandscapeEditor/Private/LandscapeEdMode.cpp
--- a/Engine/Source/Editor/LandscapeEditor/Private/LandscapeEdMode.cpp
+++ b/Engine/Source/Editor/LandscapeEditor/Private/LandscapeEdMode.cpp
@@ -4599,6 +4599,27 @@ TArray<ALandscapeBlueprintBrushBase*> FEdModeLandscape::GetBrushesForCurrentLaye
 	return Brushes;
 }
 
+void FEdModeLandscape::ShowOnlySelectedBrush(class ALandscapeBlueprintBrushBase* InBrush)
+{
+	if (ALandscape * Landscape = GetLandscape())
+	{
+		int32 BrushLayer = Landscape->GetBrushLayer(InBrush);
+		TArray<ALandscapeBlueprintBrushBase*> Brushes = Landscape->GetBrushesForLayer(BrushLayer);
+		for (ALandscapeBlueprintBrushBase* Brush : Brushes)
+		{
+			Brush->SetIsVisible(Brush == InBrush);
+		}
+	}
+}
+
+void FEdModeLandscape::DuplicateBrush(class ALandscapeBlueprintBrushBase* InBrush)
+{
+	GEditor->SelectNone(false, true);
+	GEditor->SelectActor(InBrush, true, false, false);
+
+	GEditor->edactDuplicateSelected(InBrush->GetLevel(), false);
+}
+
 bool FEdModeLandscape::IsCurrentLayerBlendSubstractive(const TWeakObjectPtr<ULandscapeLayerInfoObject>& InLayerInfoObj) const
 {
 	ALandscape* Landscape = GetLandscape();
diff --git a/Engine/Source/Editor/LandscapeEditor/Private/LandscapeEdMode.h b/Engine/Source/Editor/LandscapeEditor/Private/LandscapeEdMode.h
--- a/Engine/Source/Editor/LandscapeEditor/Private/LandscapeEdMode.h
+++ b/Engine/Source/Editor/LandscapeEditor/Private/LandscapeEdMode.h
@@ -573,6 +573,10 @@ public:
 	class ALandscapeBlueprintBrushBase* GetBrushForCurrentLayer(int8 BrushIndex) const;
 	TArray<class ALandscapeBlueprintBrushBase*> GetBrushesForCurrentLayer();
 	
+	void ShowOnlySelectedBrush(class ALandscapeBlueprintBrushBase* InBrush);
+	
+	void DuplicateBrush(class ALandscapeBlueprintBrushBase* InBrush);
+
 	bool NeedToFillEmptyMaterialLayers() const;
 	void RequestLayersContentUpdate(ELandscapeLayerUpdateMode InUpdateMode);
 	void RequestLayersContentUpdateForceAll(ELandscapeLayerUpdateMode InUpdateMode = ELandscapeLayerUpdateMode::Update_All);
diff --git a/Engine/Source/Editor/LandscapeEditor/Private/LandscapeEditorDetailCustomization_LayersBrushStack.cpp b/Engine/Source/Editor/LandscapeEditor/Private/LandscapeEditorDetailCustomization_LayersBrushStack.cpp
--- a/Engine/Source/Editor/LandscapeEditor/Private/LandscapeEditorDetailCustomization_LayersBrushStack.cpp
+++ b/Engine/Source/Editor/LandscapeEditor/Private/LandscapeEditorDetailCustomization_LayersBrushStack.cpp
@@ -214,6 +214,33 @@ TSharedPtr<SWidget> FLandscapeEditorCustomNodeBuilder_LayersBrushStack::OnBrushC
 		{
 			TSharedRef<FLandscapeEditorCustomNodeBuilder_LayersBrushStack> SharedThis = AsShared();
 
+			// Show only selected
+			FUIAction ShowOnlySelectedAction = FUIAction(FExecuteAction::CreateLambda([SharedThis, CurrentBrush]
+			{
+				const FScopedTransaction Transaction(LOCTEXT("LandscapeBrushShowOnlySelectedTransaction", "Show Only Selected Brush"));
+				GetEditorMode()->ShowOnlySelectedBrush(CurrentBrush);
+			}));
+			MenuBuilder.AddMenuEntry(LOCTEXT("LandscapeBrushShowOnlySelected", "Show Only Selected"), LOCTEXT("LandscapeBrushShowOnlySelectedToolTip", "Hides all other brushes from the same layer"), FSlateIcon(), ShowOnlySelectedAction);
+
+			// Hide selected
+			FUIAction ShowHideSelectedAction = FUIAction(FExecuteAction::CreateLambda([SharedThis, CurrentBrush]
+			{
+				SharedThis->OnToggleVisibility(CurrentBrush);
+			}));
+
+			FText MenuText = CurrentBrush->IsVisible() ? LOCTEXT("LandscapeBrushHideSelected", "Hide Selected") : LOCTEXT("LandscapeBrushShowSelected", "Show Selected");
+			MenuBuilder.AddMenuEntry(MenuText, LOCTEXT("LandscapeBrushToggleVisiblityToolTip", "Toggle brush visiblity"), FSlateIcon(), ShowHideSelectedAction);
+
+			MenuBuilder.AddMenuSeparator();
+
+			// Duplicate Brush
+			FUIAction DuplicateAction = FUIAction(FExecuteAction::CreateLambda([SharedThis, CurrentBrush]
+			{
+				const FScopedTransaction Transaction(LOCTEXT("LandscapeBrushDuplicateTransaction", "Duplicate Brush"));
+				GetEditorMode()->DuplicateBrush(CurrentBrush);
+			}));
+			MenuBuilder.AddMenuEntry(LOCTEXT("LandscapeBrushDuplicate", "Duplicate"), LOCTEXT("LandscapeBrushDuplicateToolTip", "Duplicate brush"), FSlateIcon(), DuplicateAction);
+
 			// Add Brush
 			const TArray<ALandscapeBlueprintBrushBase*>& Brushes = GetEditorMode()->GetBrushList();
 			TArray<ALandscapeBlueprintBrushBase*> FilteredBrushes = Brushes.FilterByPredicate([](ALandscapeBlueprintBrushBase* Brush) { return Brush->GetOwningLandscape() == nullptr; });
@@ -283,13 +310,18 @@ FReply FLandscapeEditorCustomNodeBuilder_LayersBrushStack::OnToggleVisibility(in
 {
 	if (ALandscapeBlueprintBrushBase* Brush = GetBrush(InBrushIndex))
 	{
-		const FScopedTransaction Transaction(LOCTEXT("Landscape_Brush_SetVisibility", "Set Brush Visibility"));
-		bool bVisible = Brush->IsVisible();
-		Brush->SetIsVisible(!bVisible);
+		OnToggleVisibility(Brush);
 	}
 	return FReply::Handled();
 }
 
+void FLandscapeEditorCustomNodeBuilder_LayersBrushStack::OnToggleVisibility(ALandscapeBlueprintBrushBase* Brush)
+{
+	const FScopedTransaction Transaction(LOCTEXT("Landscape_Brush_SetVisibility", "Set Brush Visibility"));
+	bool bVisible = Brush->IsVisible();
+	Brush->SetIsVisible(!bVisible);
+}
+
 const FSlateBrush* FLandscapeEditorCustomNodeBuilder_LayersBrushStack::GetAffectsWeightmapBrush(int32 InBrushIndex) const
 {
 	ALandscapeBlueprintBrushBase* Brush = GetBrush(InBrushIndex);
@@ -331,15 +363,19 @@ void FLandscapeEditorCustomNodeBuilder_LayersBrushStack::OnBrushSelectionChanged
 	}
 
 	const FName CurrentToolName = LandscapeEdMode->CurrentTool->GetToolName();
-	if (CurrentToolName == TEXT("BlueprintBrush"))
+	const FName BlueprintBrushTool = TEXT("BlueprintBrush");
+	bool bSwitchTool = CurrentToolName != BlueprintBrushTool;
+	
+	ALandscapeBlueprintBrushBase* Brush = GetBrush(InBrushIndex);
+	if (Brush != nullptr)
 	{
-		ALandscapeBlueprintBrushBase* Brush = GetBrush(InBrushIndex);
-		if (Brush != nullptr)
+		FScopedTransaction Transaction(LOCTEXT("LandscapeBrushSelect", "Brush selection"));
+		if (bSwitchTool)
 		{
-			FScopedTransaction Transaction(LOCTEXT("LandscapeBrushSelect", "Brush selection"));
-			GEditor->SelectNone(true, true);
-			GEditor->SelectActor(Brush, true, true);
+			LandscapeEdMode->SetCurrentTool(BlueprintBrushTool);
 		}
+		GEditor->SelectNone(true, true);
+		GEditor->SelectActor(Brush, true, true);
 	}
 }
 
diff --git a/Engine/Source/Runtime/Landscape/Classes/Landscape.h b/Engine/Source/Runtime/Landscape/Classes/Landscape.h
--- a/Engine/Source/Runtime/Landscape/Classes/Landscape.h
+++ b/Engine/Source/Runtime/Landscape/Classes/Landscape.h
@@ -318,6 +318,7 @@ public:
 	LANDSCAPE_API bool IsLayerBlendSubstractive(int32 InLayerIndex, const TWeakObjectPtr<ULandscapeLayerInfoObject>& InLayerInfoObj) const;
 	LANDSCAPE_API void SetLayerSubstractiveBlendStatus(int32 InLayerIndex, bool InStatus, const TWeakObjectPtr<ULandscapeLayerInfoObject>& InLayerInfoObj);
 
+	LANDSCAPE_API int32 GetBrushLayer(class ALandscapeBlueprintBrushBase* InBrush) const;
 	LANDSCAPE_API void AddBrushToLayer(int32 InLayerIndex, class ALandscapeBlueprintBrushBase* InBrush);
 	LANDSCAPE_API void RemoveBrush(class ALandscapeBlueprintBrushBase* InBrush);
 	LANDSCAPE_API void RemoveBrushFromLayer(int32 InLayerIndex, class ALandscapeBlueprintBrushBase* InBrush);
diff --git a/Engine/Source/Runtime/Landscape/Private/LandscapeEditLayers.cpp b/Engine/Source/Runtime/Landscape/Private/LandscapeEditLayers.cpp
--- a/Engine/Source/Runtime/Landscape/Private/LandscapeEditLayers.cpp
+++ b/Engine/Source/Runtime/Landscape/Private/LandscapeEditLayers.cpp
@@ -6178,8 +6178,25 @@ bool ALandscape::ReorderLayerBrush(int32 InLayerIndex, int32 InStartingLayerBrus
 	return false;
 }
 
+int32 ALandscape::GetBrushLayer(class ALandscapeBlueprintBrushBase* InBrush) const
+{
+	for (int32 LayerIndex = 0; LayerIndex < LandscapeLayers.Num(); ++LayerIndex)
+	{
+		for (const FLandscapeLayerBrush& Brush : LandscapeLayers[LayerIndex].Brushes)
+		{
+			if (Brush.GetBrush() == InBrush)
+			{
+				return LayerIndex;
+			}
+		}
+	}
+
+	return INDEX_NONE;
+}
+
 void ALandscape::AddBrushToLayer(int32 InLayerIndex, ALandscapeBlueprintBrushBase* InBrush)
 {
+	check(GetBrushLayer(InBrush) == INDEX_NONE);
 	if (FLandscapeLayer* Layer = GetLayer(InLayerIndex))
 	{
 		Modify();
@@ -6191,7 +6208,8 @@ void ALandscape::AddBrushToLayer(int32 InLayerIndex, ALandscapeBlueprintBrushBas
 
 void ALandscape::RemoveBrush(ALandscapeBlueprintBrushBase* InBrush)
 {
-	for (int32 LayerIndex = 0; LayerIndex < LandscapeLayers.Num(); ++LayerIndex)
+	int32 LayerIndex = GetBrushLayer(InBrush);
+	if (LayerIndex != INDEX_NONE)
 	{
 		RemoveBrushFromLayer(LayerIndex, InBrush);
 	}
