Forgot file

#rnx

diff --git a/Engine/Source/Editor/LandscapeEditor/Private/LandscapeEditorDetailCustomization_LayersBrushStack.h b/Engine/Source/Editor/LandscapeEditor/Private/LandscapeEditorDetailCustomization_LayersBrushStack.h
--- a/Engine/Source/Editor/LandscapeEditor/Private/LandscapeEditorDetailCustomization_LayersBrushStack.h
+++ b/Engine/Source/Editor/LandscapeEditor/Private/LandscapeEditorDetailCustomization_LayersBrushStack.h
@@ -65,6 +65,8 @@ protected:
 	FReply OnToggleAffectsHeightmap(int32 InBrushIndex);
 	FReply OnToggleAffectsWeightmap(int32 InBrushIndex);
 
+	void OnToggleVisibility(ALandscapeBlueprintBrushBase* Brush);
+
 	const FSlateBrush* GetAffectsHeightmapBrush(int32 InBrushIndex) const;
 	const FSlateBrush* GetAffectsWeightmapBrush(int32 InBrushIndex) const;
 	const FSlateBrush* GetVisibilityBrush(int32 InBrushIndex) const;
