Landscape: Fix TargetLayerStartingIndex reset to 0 when LandscapeTargetList is emptied.

#rnx

diff --git a/Engine/Source/Editor/LandscapeEditor/Private/LandscapeEdMode.cpp b/Engine/Source/Editor/LandscapeEditor/Private/LandscapeEdMode.cpp
--- a/Engine/Source/Editor/LandscapeEditor/Private/LandscapeEdMode.cpp
+++ b/Engine/Source/Editor/LandscapeEditor/Private/LandscapeEdMode.cpp
@@ -851,6 +851,7 @@ void FEdModeLandscape::Exit()
 
 	LandscapeList.Empty();
 	LandscapeTargetList.Empty();
+	TargetLayerStartingIndex = 0;
 
 	// Save UI settings to config file
 	UISettings->Save();
@@ -2668,6 +2669,7 @@ bool FEdModeLandscape::CanEditCurrentTarget(FText* Reason) const
 void FEdModeLandscape::UpdateTargetList()
 {
 	LandscapeTargetList.Empty();
+	TargetLayerStartingIndex = 0;
 
 	if (CurrentToolTarget.LandscapeInfo.IsValid())
 	{
