Landscape Layers: Remove 'Experimental' banner

#rnx

diff --git a/Engine/Source/Editor/LandscapeEditor/Private/SLandscapeEditor.cpp b/Engine/Source/Editor/LandscapeEditor/Private/SLandscapeEditor.cpp
--- a/Engine/Source/Editor/LandscapeEditor/Private/SLandscapeEditor.cpp
+++ b/Engine/Source/Editor/LandscapeEditor/Private/SLandscapeEditor.cpp
@@ -341,18 +341,6 @@ void SLandscapeEditor::Construct(const FArguments& InArgs, TSharedRef<FLandscape
 			.IsEnabled(this, &SLandscapeEditor::GetLandscapeEditorIsEnabled)
 			+ SVerticalBox::Slot()
 			.AutoHeight()
-			.Padding(0, 0, 0, 5)
-			[
-				SNew(SMultiLineEditableTextBox)
-				.Visibility_Lambda([=]() { return LandscapeEdMode->CanHaveLandscapeLayersContent() ? EVisibility::Visible : EVisibility::Collapsed; })
-				.IsReadOnly(true)
-				.Font(FCoreStyle::GetDefaultFontStyle("Bold", 10))
-				.Justification(ETextJustify::Center)
-				.BackgroundColor(FCoreStyle::Get().GetColor("InfoReporting.BackgroundColor"))
-				.Text(FText::FromString("Experimental Layer System"))
-			]
-			+ SVerticalBox::Slot()
-			.AutoHeight()
 			.Padding(4, 0, 4, 5)
 			[
 				SNew(SOverlay)
