Undo support for runtime virtual texture landscape painting.
Quick implementation. Still need to capture bounds.

diff --git a/Engine/Source/Editor/LandscapeEditor/Private/LandscapeEdModePaintTools.cpp b/Engine/Source/Editor/LandscapeEditor/Private/LandscapeEdModePaintTools.cpp
--- a/Engine/Source/Editor/LandscapeEditor/Private/LandscapeEdModePaintTools.cpp
+++ b/Engine/Source/Editor/LandscapeEditor/Private/LandscapeEdModePaintTools.cpp
@@ -268,16 +268,11 @@ public:
 		this->Cache.Flush();
 
 		// If we render to a runtime virtual texture then we need to flush here
-		if (Landscape != nullptr && Landscape->GetLandscapeMaterial() != nullptr)
+		//todo[vt]: This logic doesn't cover material overrides etc.
+		if (Landscape != nullptr && Landscape->RuntimeVirtualTextures.Num() > 0)
 		{
-			//todo[vt]: This logic doesn't cover material overrides etc.
-			FMaterialResource* MaterialResource = Landscape->GetLandscapeMaterial()->GetMaterialResource(GMaxRHIFeatureLevel);
-			//todo[vt]: This test is slow, we could cache the result in the landscape or the material
-			if (MaterialResource != nullptr && (MaterialResource->GetMaterialDomain() == MD_RuntimeVirtualTexture || MaterialResource->HasRuntimeVirtualTextureOutput()))
-			{
-				//todo[vt]: Only flush Bounds 
-				GetRendererModule().FlushVirtualTextureCache();
-			}
+			//todo[vt]: Only flush Bounds 
+			GetRendererModule().FlushVirtualTextureCache();
 		}
 	}
 };
diff --git a/Engine/Source/Runtime/Landscape/Private/LandscapeEdit.cpp b/Engine/Source/Runtime/Landscape/Private/LandscapeEdit.cpp
--- a/Engine/Source/Runtime/Landscape/Private/LandscapeEdit.cpp
+++ b/Engine/Source/Runtime/Landscape/Private/LandscapeEdit.cpp
@@ -45,6 +45,7 @@ LandscapeEdit.cpp: Landscape editing
 #include "MeshAttributeArray.h"
 #include "MeshUtilitiesCommon.h"
 
+#include "EngineModule.h"
 #include "EngineUtils.h"
 #include "Framework/Notifications/NotificationManager.h"
 #include "Widgets/Notifications/SNotificationList.h"
@@ -594,6 +595,12 @@ void ULandscapeComponent::PostEditUndo()
 		Components.Add(this);
 		GetLandscapeProxy()->FlushGrassComponents(&Components);
 	}
+
+	if (GetLandscapeProxy()->RuntimeVirtualTextures.Num() > 0)
+	{
+		//todo[vt]: Only flush Bounds 
+		GetRendererModule().FlushVirtualTextureCache();
+	}
 }
 
 void ALandscapeProxy::FixupWeightmaps()
