Fixes compilation failure in Landscape paint tool.

diff --git a/Engine/Source/Editor/LandscapeEditor/Private/LandscapeEdModePaintTools.cpp b/Engine/Source/Editor/LandscapeEditor/Private/LandscapeEdModePaintTools.cpp
--- a/Engine/Source/Editor/LandscapeEditor/Private/LandscapeEdModePaintTools.cpp
+++ b/Engine/Source/Editor/LandscapeEditor/Private/LandscapeEdModePaintTools.cpp
@@ -257,6 +257,8 @@ public:
 
 		this->Cache.SetCachedData(X1, Y1, X2, Y2, Data, UISettings->PaintingRestriction);
 		this->Cache.Flush();
+		
+		ALandscape* Landscape = LandscapeInfo->LandscapeActor.Get();
 
 		// If we render to a runtime virtual texture then we need to flush here
 		//todo[vt]: This logic doesn't cover material overrides etc.
