Merging //UE4/Dev-Main@6618752 to Dev-Rendering (//UE4/Dev-Rendering)

#rnx

diff --git a/Engine/Source/Editor/DetailCustomizations/Private/LandscapeUIDetails.cpp b/Engine/Source/Editor/DetailCustomizations/Private/LandscapeUIDetails.cpp
--- a/Engine/Source/Editor/DetailCustomizations/Private/LandscapeUIDetails.cpp
+++ b/Engine/Source/Editor/DetailCustomizations/Private/LandscapeUIDetails.cpp
@@ -61,22 +61,25 @@ void FLandscapeUIDetails::CustomizeDetails( IDetailLayoutBuilder& DetailBuilder
 				Category.AddCustomRow(FText(), true)
 				.NameContent()
 				[
-					PropertyHandle->CreatePropertyNameWidget(FText::FromString("Layer System"))
+					PropertyHandle->CreatePropertyNameWidget(NSLOCTEXT("UnrealEd", "LandscapeToggleLayerName", "Enable Layer System"))
 				]
 				.ValueContent()
-				[
-					SNew(SButton)
-					.OnClicked_Lambda([this, Landscape]() -> FReply
-					{
-						return ToggleCanHaveLayersContent(Landscape);
-					})
 					[
-						SNew(STextBlock)
-						.Text_Lambda([Landscape]()
-						{
-							return Landscape->bCanHaveLayersContent ? NSLOCTEXT("UnrealEd", "LandscapeDisableLayers", "Disable") : NSLOCTEXT("UnrealEd", "LandscapeEnableLayers", "Enable");	
+						SNew(SCheckBox)
+						.ToolTipText(NSLOCTEXT("UnrealEd", "LandscapeToggleLayerToolTip", "Toggle whether or not to support layers on this Landscape and its streaming proxies. Toggling this will clear the undo stack."))
+						.Type(ESlateCheckBoxType::CheckBox)
+						.IsChecked_Lambda([=]() 
+						{ 
+							return Landscape->CanHaveLayersContent() ? ECheckBoxState::Checked : ECheckBoxState::Unchecked;
+						})
+						.OnCheckStateChanged_Lambda([=](ECheckBoxState NewState) 
+						{ 
+							bool bChecked = (NewState == ECheckBoxState::Checked);
+							if (Landscape->CanHaveLayersContent() != bChecked)
+							{
+								ToggleCanHaveLayersContent(Landscape);
+							}
 						})
-					]
 				];
 			}
 			
@@ -92,7 +95,7 @@ void FLandscapeUIDetails::OnEditorExperimentalSettingsChanged(FName PropertyName
 	}
 }
 
-FReply FLandscapeUIDetails::ToggleCanHaveLayersContent(ALandscape* Landscape)
+void FLandscapeUIDetails::ToggleCanHaveLayersContent(ALandscape* Landscape)
 {
 	bool bToggled = false;
 
@@ -110,8 +113,7 @@ FReply FLandscapeUIDetails::ToggleCanHaveLayersContent(ALandscape* Landscape)
 				break;
 			}
 		}
-
-		
+				
 		FText Reason;
 
 		if (bHasHiddenLayers)
@@ -131,12 +133,12 @@ FReply FLandscapeUIDetails::ToggleCanHaveLayersContent(ALandscape* Landscape)
 		bToggled = FMessageDialog::Open(EAppMsgType::YesNo, NSLOCTEXT("UnrealEd", "LandscapeEnableLayers", "Are you sure you want to enable the layer system on this landscape and streaming proxies? Doing so will clear the undo/redo buffer.")) == EAppReturnType::Yes;
 	}
 
-	if (DetailLayoutBuilder && bToggled)
+	if (bToggled)
 	{
 		Landscape->ToggleCanHaveLayersContent();
-		GEditor->ResetTransaction(NSLOCTEXT("UnrealEd", "ToggleLanscapeLayers", "Toggling Landscape Layers"));
-		DetailLayoutBuilder->ForceRefreshDetails();
+		if (GEditor)
+		{
+			GEditor->ResetTransaction(NSLOCTEXT("UnrealEd", "ToggleLanscapeLayers", "Toggling Landscape Layers"));
+		}
 	}
-
-	return FReply::Handled();
 }
\ No newline at end of file
diff --git a/Engine/Source/Editor/DetailCustomizations/Private/LandscapeUIDetails.h b/Engine/Source/Editor/DetailCustomizations/Private/LandscapeUIDetails.h
--- a/Engine/Source/Editor/DetailCustomizations/Private/LandscapeUIDetails.h
+++ b/Engine/Source/Editor/DetailCustomizations/Private/LandscapeUIDetails.h
@@ -25,7 +25,7 @@ public:
 
 private:
 	void OnEditorExperimentalSettingsChanged(FName PropertyName);
-	FReply ToggleCanHaveLayersContent(ALandscape* Landscape);
+	void ToggleCanHaveLayersContent(ALandscape* Landscape);
 
 	/** Use MakeInstance to create an instance of this class */
 	FLandscapeUIDetails();
diff --git a/Engine/Source/Editor/LandscapeEditor/Private/LandscapeEditorDetailCustomization_Layers.cpp b/Engine/Source/Editor/LandscapeEditor/Private/LandscapeEditorDetailCustomization_Layers.cpp
--- a/Engine/Source/Editor/LandscapeEditor/Private/LandscapeEditorDetailCustomization_Layers.cpp
+++ b/Engine/Source/Editor/LandscapeEditor/Private/LandscapeEditorDetailCustomization_Layers.cpp
@@ -355,6 +355,7 @@ void FLandscapeEditorCustomNodeBuilder_Layers::SetLayerName(const FText& InText,
 	{
 		const FScopedTransaction Transaction(LOCTEXT("Landscape_Layers_Rename", "Rename Layer"));
 		LandscapeEdMode->SetLayerName(InLayerIndex, *InText.ToString());
+		OnLayerSelectionChanged(InLayerIndex);
 	}
 }
 
@@ -466,6 +467,7 @@ void FLandscapeEditorCustomNodeBuilder_Layers::SetLandscapeSplinesReservedLayer(
 				Landscape->SetLandscapeSplinesReservedLayer(InLayerIndex);
 				LandscapeEdMode->RefreshDetailPanel();
 				LandscapeEdMode->AutoUpdateDirtyLandscapeSplines();
+				OnLayerSelectionChanged(InLayerIndex);
 			}
 		}
 	}
@@ -493,6 +495,7 @@ void FLandscapeEditorCustomNodeBuilder_Layers::ClearLayer(int32 InLayerIndex)
 			{
 				const FScopedTransaction Transaction(LOCTEXT("Landscape_Layers_Clean", "Clear Layer"));
 				Landscape->ClearLayer(InLayerIndex);
+				OnLayerSelectionChanged(InLayerIndex);
 			}
 		}
 	}
@@ -528,6 +531,7 @@ void FLandscapeEditorCustomNodeBuilder_Layers::ShowOnlySelectedLayer(int32 InLay
 	{
 		const FScopedTransaction Transaction(LOCTEXT("ShowOnlySelectedLayer", "Show Only Selected Layer"));
 		Landscape->ShowOnlySelectedLayer(InLayerIndex);
+		OnLayerSelectionChanged(InLayerIndex);
 	}
 }
 
@@ -544,22 +548,23 @@ void FLandscapeEditorCustomNodeBuilder_Layers::ShowAllLayers()
 
 void FLandscapeEditorCustomNodeBuilder_Layers::CreateLayer()
 {
-	FEdModeLandscape* LandscapeEdMode = GetEditorMode();
-	ALandscape* Landscape = LandscapeEdMode ? LandscapeEdMode->GetLandscape() : nullptr;
-	if (Landscape)
+FEdModeLandscape* LandscapeEdMode = GetEditorMode();
+ALandscape* Landscape = LandscapeEdMode ? LandscapeEdMode->GetLandscape() : nullptr;
+if (Landscape)
+{
 	{
-		{
-			const FScopedTransaction Transaction(LOCTEXT("Landscape_Layers_Create", "Create Layer"));
-			Landscape->CreateLayer();
-		}
-		LandscapeEdMode->RefreshDetailPanel();
+		const FScopedTransaction Transaction(LOCTEXT("Landscape_Layers_Create", "Create Layer"));
+		Landscape->CreateLayer();
+		OnLayerSelectionChanged(Landscape->GetLayerCount() - 1);
 	}
+	LandscapeEdMode->RefreshDetailPanel();
+}
 }
 
 void FLandscapeEditorCustomNodeBuilder_Layers::OnLayerSelectionChanged(int32 InLayerIndex)
 {
 	FEdModeLandscape* LandscapeEdMode = GetEditorMode();
-	if (LandscapeEdMode)
+	if (LandscapeEdMode && LandscapeEdMode->GetCurrentLayerIndex() != InLayerIndex)
 	{
 		FScopedTransaction Transaction(LOCTEXT("Landscape_Layers_SetCurrentLayer", "Set Current Layer"));
 		LandscapeEdMode->SetCurrentLayer(InLayerIndex);
@@ -603,6 +608,10 @@ FReply FLandscapeEditorCustomNodeBuilder_Layers::OnToggleVisibility(int32 InLaye
 	{
 		const FScopedTransaction Transaction(LOCTEXT("Landscape_Layers_SetVisibility", "Set Layer Visibility"));
 		LandscapeEdMode->SetLayerVisibility(!LandscapeEdMode->IsLayerVisible(InLayerIndex), InLayerIndex);
+		if (LandscapeEdMode->IsLayerVisible(InLayerIndex))
+		{
+			OnLayerSelectionChanged(InLayerIndex);
+		}
 	}
 	return FReply::Handled();
 }
diff --git a/Engine/Source/Editor/LandscapeEditor/Private/LandscapeEditorDetailCustomization_TargetLayers.cpp b/Engine/Source/Editor/LandscapeEditor/Private/LandscapeEditorDetailCustomization_TargetLayers.cpp
--- a/Engine/Source/Editor/LandscapeEditor/Private/LandscapeEditorDetailCustomization_TargetLayers.cpp
+++ b/Engine/Source/Editor/LandscapeEditor/Private/LandscapeEditorDetailCustomization_TargetLayers.cpp
@@ -1201,15 +1201,13 @@ void FLandscapeEditorCustomNodeBuilder_TargetLayers::OnTargetLayerSetObject(cons
 					Target->LandscapeInfo->CreateLayerEditorSettingsFor(SelectedLayerInfo);
 				}
 			}
-
+						
 			FEdModeLandscape* LandscapeEdMode = GetEditorMode();
 			if (LandscapeEdMode)
 			{
-				if (LandscapeEdMode->CurrentToolTarget.LayerName == Target->LayerName
-					&& LandscapeEdMode->CurrentToolTarget.LayerInfo == Target->LayerInfoObj)
-				{
-					LandscapeEdMode->CurrentToolTarget.LayerInfo = SelectedLayerInfo;
-				}
+				LandscapeEdMode->CurrentToolTarget.LayerName = Target->LayerName;
+				LandscapeEdMode->CurrentToolTarget.TargetType = Target->TargetType;
+				LandscapeEdMode->CurrentToolTarget.LayerInfo = SelectedLayerInfo;
 				LandscapeEdMode->UpdateTargetList();
 			}
 
@@ -1545,15 +1543,15 @@ FReply SLandscapeEditorSelectableBorder::OnMouseButtonUp(const FGeometry& MyGeom
 		}
 		else if (MouseEvent.GetEffectingButton() == EKeys::RightMouseButton &&
 			OnContextMenuOpening.IsBound())
-		{
-			TSharedPtr<SWidget> Content = OnContextMenuOpening.Execute();
-			if (Content.IsValid())
 			{
-				FWidgetPath WidgetPath = MouseEvent.GetEventPath() != nullptr ? *MouseEvent.GetEventPath() : FWidgetPath();
-
-				FSlateApplication::Get().PushMenu(SharedThis(this), WidgetPath, Content.ToSharedRef(), MouseEvent.GetScreenSpacePosition(), FPopupTransitionEffect(FPopupTransitionEffect::ContextMenu));
-			}
+				TSharedPtr<SWidget> Content = OnContextMenuOpening.Execute();
+				if (Content.IsValid())
+				{
+					FWidgetPath WidgetPath = MouseEvent.GetEventPath() != nullptr ? *MouseEvent.GetEventPath() : FWidgetPath();
 
+					FSlateApplication::Get().PushMenu(SharedThis(this), WidgetPath, Content.ToSharedRef(), MouseEvent.GetScreenSpacePosition(), FPopupTransitionEffect(FPopupTransitionEffect::ContextMenu));
+				}
+			
 			return FReply::Handled().ReleaseMouseCapture();
 		}
 	}
diff --git a/Engine/Source/Runtime/D3D12RHI/Private/D3D12RayTracing.cpp b/Engine/Source/Runtime/D3D12RHI/Private/D3D12RayTracing.cpp
--- a/Engine/Source/Runtime/D3D12RHI/Private/D3D12RayTracing.cpp
+++ b/Engine/Source/Runtime/D3D12RHI/Private/D3D12RayTracing.cpp
@@ -2379,9 +2379,6 @@ void FD3D12CommandContext::RHIUpdateAccelerationStructures(const TArrayView<cons
 	{
 		FD3D12RayTracingGeometry* Geometry = FD3D12DynamicRHI::ResourceCast(P.Geometry.GetReference());
 		Geometry->SetDirty(GetGPUMask(), true);
-
-		const bool bIsUpdate = !GRayTracingDebugForceFullBuild;
-		Geometry->BuildAccelerationStructure(*this, bIsUpdate);
 	}
 }
 
diff --git a/Engine/Source/Runtime/Engine/Private/Animation/AnimNode_Layer.cpp b/Engine/Source/Runtime/Engine/Private/Animation/AnimNode_Layer.cpp
--- a/Engine/Source/Runtime/Engine/Private/Animation/AnimNode_Layer.cpp
+++ b/Engine/Source/Runtime/Engine/Private/Animation/AnimNode_Layer.cpp
@@ -30,11 +30,14 @@ void FAnimNode_Layer::OnInitializeAnimInstance(const FAnimInstanceProxy* InProxy
 		USkeletalMeshComponent* MeshComp = InAnimInstance->GetSkelMeshComponent();
 		check(MeshComp);
 
-		// Switch from dynamic external to internal, kill old instance
-		if(CurrentTarget && CurrentTarget != InAnimInstance)
+		if(LinkedRoot)
 		{
 			DynamicUnlink(const_cast<UAnimInstance*>(InAnimInstance));
+		}
 
+		// Switch from dynamic external to internal, kill old instance
+		if(CurrentTarget && CurrentTarget != InAnimInstance)
+		{
 			MeshComp->SubInstances.Remove(CurrentTarget);
 			CurrentTarget->MarkPendingKill();
 			CurrentTarget = nullptr;
