Fix Empty Target List when going from Paint -> Placement -> Paint

#rnx

diff --git a/Engine/Source/Editor/LandscapeEditor/Private/LandscapeEdMode.cpp b/Engine/Source/Editor/LandscapeEditor/Private/LandscapeEdMode.cpp
--- a/Engine/Source/Editor/LandscapeEditor/Private/LandscapeEdMode.cpp
+++ b/Engine/Source/Editor/LandscapeEditor/Private/LandscapeEdMode.cpp
@@ -728,6 +728,8 @@ void FEdModeLandscape::Enter()
 	// Load UI settings from config file
 	UISettings->Load();
 
+	// It is cleared on exit so update here even if LandscapeInfo hasn't changed
+	UpdateTargetList();
 	UpdateShownLayerList();
 
 	// Initialize current tool prior to creating the landscape toolkit in case it has a dependency on it
