Landscape Layers: Collapse Feature

diff --git a/Engine/Source/Editor/LandscapeEditor/Private/LandscapeEditorDetailCustomization_Layers.cpp b/Engine/Source/Editor/LandscapeEditor/Private/LandscapeEditorDetailCustomization_Layers.cpp
--- a/Engine/Source/Editor/LandscapeEditor/Private/LandscapeEditorDetailCustomization_Layers.cpp
+++ b/Engine/Source/Editor/LandscapeEditor/Private/LandscapeEditorDetailCustomization_Layers.cpp
@@ -456,6 +456,13 @@ TSharedPtr<SWidget> FLandscapeEditorCustomNodeBuilder_Layers::OnLayerContextMenu
 					FUIAction DeleteLayerAction = FUIAction(FExecuteAction::CreateLambda([SharedThis, InLayerIndex] { SharedThis->DeleteLayer(InLayerIndex); }), FCanExecuteAction::CreateLambda([Landscape] { return Landscape->LandscapeLayers.Num() > 1; }));
 					TAttribute<FText> DeleteLayerText(Landscape->LandscapeLayers.Num() == 1 ? LOCTEXT("CantDeleteLastLayer", "Delete (Last layer)") : LOCTEXT("DeleteLayer", "Delete..."));
 					MenuBuilder.AddMenuEntry(DeleteLayerText, LOCTEXT("DeleteLayerTooltip", "Delete Layer"), FSlateIcon(), DeleteLayerAction);
+
+					// Collapse Layer
+					if (LandscapeEdMode->CurrentToolMode->ToolModeName == "ToolMode_Manage")
+					{
+						FUIAction CollapseLayerAction = FUIAction(FExecuteAction::CreateLambda([SharedThis, InLayerIndex] { SharedThis->CollapseLayer(InLayerIndex); }));
+						MenuBuilder.AddMenuEntry(LOCTEXT("CollapseLayer", "Collapse..."), LOCTEXT("CollapseLayerTooltip", "Collapse layer into top neighbor layer"), FSlateIcon(), CollapseLayerAction);
+					}
 				}
 
 				if (Landscape->GetLandscapeSplinesReservedLayer() != Landscape->GetLayer(InLayerIndex))
@@ -644,6 +651,74 @@ void FLandscapeEditorCustomNodeBuilder_Layers::DeleteLayer(int32 InLayerIndex)
 	}
 }
 
+bool FLandscapeEditorCustomNodeBuilder_Layers::CanCollapseLayer(int32 InLayerIndex, FText& OutReason) const
+{
+	FEdModeLandscape* LandscapeEdMode = GetEditorMode();
+	ALandscape* Landscape = LandscapeEdMode ? LandscapeEdMode->GetLandscape() : nullptr;
+	if (!Landscape || Landscape->LandscapeLayers.Num() <= 1)
+	{
+		OutReason = LOCTEXT("Landscape_CollapseLayer_Reason_NotEnoughLayersToCollapse", "Not enough layers to do Collapse.");
+		return false;
+	}
+
+	if (InLayerIndex < 1)
+	{
+		OutReason = LOCTEXT("Landscape_CollapseLayer_Reason_CantCollapseBaseLayer", "Can't Collapse first layer.");
+		return false;
+	}
+	
+	FLandscapeLayer* SplineLayer = Landscape->GetLandscapeSplinesReservedLayer();
+	if (SplineLayer == LandscapeEdMode->GetLayer(InLayerIndex))
+	{
+		OutReason = LOCTEXT("Landscape_CollapseLayer_Reason_CantCollapseSplineLayer", "Can't Collapse reserved spline layer.");
+		return false;
+	}
+
+	// Can't collapse on layer that has a Brush because result will change...
+	FLandscapeLayer* BaseLayer = LandscapeEdMode->GetLayer(InLayerIndex - 1);
+	if (SplineLayer == BaseLayer)
+	{
+		OutReason = LOCTEXT("Landscape_CollapseLayer_Reason_CantCollapseOnSplineLayer", "Can't Collapse on reserved spline layer.");
+		return false;
+	}
+	
+	if (!BaseLayer || BaseLayer->Brushes.Num() > 0)
+	{
+		OutReason = FText::Format(LOCTEXT("Landscape_CollapseLayer_Reason_CantCollapseOnLayerWithBrush", "Can't Collapse because base layer '{0}' contains Brush(es)."), FText::FromName(BaseLayer->Name));
+		return false;
+	}
+	
+
+	return true;
+}
+
+void FLandscapeEditorCustomNodeBuilder_Layers::CollapseLayer(int32 InLayerIndex)
+{
+	FText Reason;
+	if (CanCollapseLayer(InLayerIndex, Reason))
+	{
+		FEdModeLandscape* LandscapeEdMode = GetEditorMode();
+		ALandscape* Landscape = LandscapeEdMode ? LandscapeEdMode->GetLandscape() : nullptr;
+		FLandscapeLayer* Layer = LandscapeEdMode->GetLayer(InLayerIndex);
+		FLandscapeLayer* BaseLayer = LandscapeEdMode->GetLayer(InLayerIndex - 1);
+		if (Layer && BaseLayer)
+		{
+			EAppReturnType::Type Result = FMessageDialog::Open(EAppMsgType::YesNo, FText::Format(LOCTEXT("Landscape_CollapseLayer_Message", "The layer {0} will be collapsed into layer {1}.  Continue?"), FText::FromName(Layer->Name), FText::FromName(BaseLayer->Name)));
+			if (Result == EAppReturnType::Yes)
+			{
+				const FScopedTransaction Transaction(LOCTEXT("Landscape_Layers_Collapse", "Collapse Layer"));
+				Landscape->CollapseLayer(InLayerIndex);
+				OnLayerSelectionChanged(InLayerIndex-1);
+				LandscapeEdMode->RefreshDetailPanel();
+			}
+		}
+	}
+	else
+	{
+		FMessageDialog::Open(EAppMsgType::Ok, Reason);
+	}
+}
+
 void FLandscapeEditorCustomNodeBuilder_Layers::ShowOnlySelectedLayer(int32 InLayerIndex)
 {
 	FEdModeLandscape* LandscapeEdMode = GetEditorMode();
diff --git a/Engine/Source/Editor/LandscapeEditor/Private/LandscapeEditorDetailCustomization_Layers.h b/Engine/Source/Editor/LandscapeEditor/Private/LandscapeEditorDetailCustomization_Layers.h
--- a/Engine/Source/Editor/LandscapeEditor/Private/LandscapeEditorDetailCustomization_Layers.h
+++ b/Engine/Source/Editor/LandscapeEditor/Private/LandscapeEditorDetailCustomization_Layers.h
@@ -78,6 +78,8 @@ protected:
 	void ClearLayer(int32 InLayerIndex, ELandscapeClearMode InClearMode);
 	void RenameLayer(int32 InLayerIndex);
 	void DeleteLayer(int32 InLayerIndex);
+	bool CanCollapseLayer(int32 InLayerIndex, FText& OutReason) const;
+	void CollapseLayer(int32 InLayerIndex);
 	void ShowOnlySelectedLayer(int32 InLayerIndex);
 	void ShowAllLayers();
 	bool CanRenameLayerTo(const FText& NewText, FText& OutErrorMessage, int32 InLayerIndex);
diff --git a/Engine/Source/Runtime/Landscape/Classes/Landscape.h b/Engine/Source/Runtime/Landscape/Classes/Landscape.h
--- a/Engine/Source/Runtime/Landscape/Classes/Landscape.h
+++ b/Engine/Source/Runtime/Landscape/Classes/Landscape.h
@@ -94,9 +94,7 @@ struct FLandscapeLayerBrush
 
 	FLandscapeLayerBrush()
 #if WITH_EDITORONLY_DATA
-		: BlueprintBrush(nullptr)
-		, LandscapeSize(MAX_int32, MAX_int32)
-		, LandscapeRenderTargetSize(MAX_int32, MAX_int32)
+		: FLandscapeLayerBrush(nullptr)
 #endif
 	{}
 
@@ -302,6 +300,7 @@ public:
 	LANDSCAPE_API void ClearLayer(int32 InLayerIndex, TSet<ULandscapeComponent*>* InComponents = nullptr, ELandscapeClearMode InClearMode = ELandscapeClearMode::Clear_All);
 	LANDSCAPE_API void ClearLayer(const FGuid& InLayerGuid, TSet<ULandscapeComponent*>* InComponents = nullptr, ELandscapeClearMode InClearMode = ELandscapeClearMode::Clear_All, bool bMarkPackageDirty = true);
 	LANDSCAPE_API void DeleteLayer(int32 InLayerIndex);
+	LANDSCAPE_API void CollapseLayer(int32 InLayerIndex);
 	LANDSCAPE_API void DeleteLayers();
 	LANDSCAPE_API void SetEditingLayer(const FGuid& InLayerGuid = FGuid());
 	LANDSCAPE_API void SetGrassUpdateEnabled(bool bInGrassUpdateEnabled);
diff --git a/Engine/Source/Runtime/Landscape/Classes/LandscapeComponent.h b/Engine/Source/Runtime/Landscape/Classes/LandscapeComponent.h
--- a/Engine/Source/Runtime/Landscape/Classes/LandscapeComponent.h
+++ b/Engine/Source/Runtime/Landscape/Classes/LandscapeComponent.h
@@ -627,6 +627,7 @@ public:
 
 	LANDSCAPE_API TArray<FWeightmapLayerAllocationInfo>& GetWeightmapLayerAllocations(bool InReturnEditingWeightmap = false);
 	LANDSCAPE_API const TArray<FWeightmapLayerAllocationInfo>& GetWeightmapLayerAllocations(bool InReturnEditingWeightmap = false) const;
+	LANDSCAPE_API TArray<FWeightmapLayerAllocationInfo>& GetWeightmapLayerAllocations(const FGuid& InLayerGuid);
 	LANDSCAPE_API const TArray<FWeightmapLayerAllocationInfo>& GetWeightmapLayerAllocations(const FGuid& InLayerGuid) const;
 
 #if WITH_EDITOR
@@ -653,6 +654,8 @@ public:
 	FLandscapeLayerComponentData* GetEditingLayer();
 	const FLandscapeLayerComponentData* GetEditingLayer() const;
 	FGuid GetEditingLayerGUID() const;
+
+	void CopyFinalLayerIntoEditingLayer(FLandscapeEditDataInterface& DataInterface, TSet<UTexture2D*>& ProcessedHeightmaps);
 #endif 
 
 #if WITH_EDITOR
diff --git a/Engine/Source/Runtime/Landscape/Private/Landscape.cpp b/Engine/Source/Runtime/Landscape/Private/Landscape.cpp
--- a/Engine/Source/Runtime/Landscape/Private/Landscape.cpp
+++ b/Engine/Source/Runtime/Landscape/Private/Landscape.cpp
@@ -56,6 +56,7 @@ Landscape.cpp: Terrain rendering
 #include "MaterialUtilities.h"
 #include "Editor.h"
 #include "Algo/Transform.h"
+#include "Engine/Texture2D.h"
 #endif
 #include "LandscapeVersion.h"
 #include "UObject/FortniteMainBranchObjectVersion.h"
@@ -1566,6 +1567,21 @@ const TArray<FWeightmapLayerAllocationInfo>& ULandscapeComponent::GetWeightmapLa
 	return WeightmapLayerAllocations;
 }
 
+TArray<FWeightmapLayerAllocationInfo>& ULandscapeComponent::GetWeightmapLayerAllocations(const FGuid& InLayerGuid)
+{
+#if WITH_EDITORONLY_DATA
+	if (InLayerGuid.IsValid())
+	{
+		if (FLandscapeLayerComponentData* LayerData = GetLayerData(InLayerGuid))
+		{
+			return LayerData->WeightmapData.LayerAllocations;
+		}
+	}
+#endif
+
+	return WeightmapLayerAllocations;
+}
+
 const TArray<FWeightmapLayerAllocationInfo>& ULandscapeComponent::GetWeightmapLayerAllocations(const FGuid& InLayerGuid) const
 {
 #if WITH_EDITORONLY_DATA
@@ -1623,6 +1639,44 @@ const FLandscapeLayerComponentData* ULandscapeComponent::GetEditingLayer() const
 	return CachedEditingLayerData;
 }
 
+void ULandscapeComponent::CopyFinalLayerIntoEditingLayer(FLandscapeEditDataInterface& DataInterface, TSet<UTexture2D*>& ProcessedHeightmaps)
+{
+	Modify();
+	GetLandscapeProxy()->Modify();
+	
+	// Heightmap	
+	UTexture2D* EditingTexture = GetHeightmap(true);
+	if (!ProcessedHeightmaps.Contains(EditingTexture))
+	{
+		DataInterface.CopyTextureFromHeightmap(EditingTexture, this, 0);	
+		ProcessedHeightmaps.Add(EditingTexture);
+	}
+
+	// Weightmap
+	TArray<FWeightmapLayerAllocationInfo>& FinalWeightmapLayerAllocations = GetWeightmapLayerAllocations();
+	TArray<FWeightmapLayerAllocationInfo>& EditingLayerWeightmapLayerAllocations = GetWeightmapLayerAllocations(GetEditingLayerGUID());
+
+	// Add missing Alloc Infos
+	for (const FWeightmapLayerAllocationInfo& FinalAllocInfo : FinalWeightmapLayerAllocations)
+	{
+		int32 Index = EditingLayerWeightmapLayerAllocations.IndexOfByPredicate([&FinalAllocInfo](const FWeightmapLayerAllocationInfo& EditingAllocInfo) { return EditingAllocInfo.LayerInfo == FinalAllocInfo.LayerInfo; });
+		if (Index == INDEX_NONE)
+		{
+			new (EditingLayerWeightmapLayerAllocations) FWeightmapLayerAllocationInfo(FinalAllocInfo.LayerInfo);
+		}
+	}
+
+	const bool bEditingWeighmaps = true;
+	const bool bSaveToTransactionBuffer = true;
+	ReallocateWeightmaps(&DataInterface, bEditingWeighmaps, bSaveToTransactionBuffer);
+
+	TArray<UTexture2D*>& EditingWeightmapTextures = GetWeightmapTextures(true);
+	for (const FWeightmapLayerAllocationInfo& AllocInfo : EditingLayerWeightmapLayerAllocations)
+	{
+		DataInterface.CopyTextureFromWeightmap(EditingWeightmapTextures[AllocInfo.WeightmapTextureIndex], AllocInfo.WeightmapTextureChannel, this, AllocInfo.LayerInfo, 0);
+	}
+}
+
 FGuid ULandscapeComponent::GetEditingLayerGUID() const
 {
 	ALandscape* Landscape = GetLandscapeActor();
diff --git a/Engine/Source/Runtime/Landscape/Private/LandscapeDataAccess.cpp b/Engine/Source/Runtime/Landscape/Private/LandscapeDataAccess.cpp
--- a/Engine/Source/Runtime/Landscape/Private/LandscapeDataAccess.cpp
+++ b/Engine/Source/Runtime/Landscape/Private/LandscapeDataAccess.cpp
@@ -186,4 +186,14 @@ LANDSCAPE_API void FLandscapeComponentDataInterface::GetWorldPositionTangents(in
 	WorldTangentZ = Component->GetComponentTransform().TransformVectorNoScale(WorldTangentZ);
 }
 
+int32 FLandscapeComponentDataInterface::GetHeightmapSizeX(int32 MipIndex) const
+{
+	return Component->GetHeightmap(bWorkOnEditingLayer)->Source.GetSizeX() >> MipIndex;
+}
+
+int32 FLandscapeComponentDataInterface::GetHeightmapSizeY(int32 MipIndex) const
+{
+	return Component->GetHeightmap(bWorkOnEditingLayer)->Source.GetSizeY() >> MipIndex;
+}
+
 #endif // WITH_EDITOR
diff --git a/Engine/Source/Runtime/Landscape/Private/LandscapeEditInterface.cpp b/Engine/Source/Runtime/Landscape/Private/LandscapeEditInterface.cpp
--- a/Engine/Source/Runtime/Landscape/Private/LandscapeEditInterface.cpp
+++ b/Engine/Source/Runtime/Landscape/Private/LandscapeEditInterface.cpp
@@ -4389,28 +4389,59 @@ void FLandscapeTextureDataInterface::CopyTextureFromHeightmap(UTexture2D* Dest,
 	}
 }
 
-void FLandscapeTextureDataInterface::CopyTextureFromWeightmap(UTexture2D* Dest, int32 DestChannel, ULandscapeComponent* Comp, ULandscapeLayerInfoObject* LayerInfo)
+void FLandscapeTextureDataInterface::CopyTextureFromHeightmap(UTexture2D* Dest, ULandscapeComponent* Comp, int32 MipIndex)
 {
 	FLandscapeTextureDataInfo* DestDataInfo = GetTextureDataInfo(Dest);
-	int32 MipSize = Dest->Source.GetSizeX();
-	check(Dest->Source.GetSizeX() == Dest->Source.GetSizeY());
+	const int32 MipSize = DestDataInfo->GetMipSizeX(MipIndex);
+	check(MipSize == DestDataInfo->GetMipSizeY(MipIndex));
+	
+	const bool bEditingLayer = false;
+	FLandscapeComponentDataInterface DataInterface(Comp, MipIndex, bEditingLayer);
+	
+	FColor* DestTextureData = (FColor*)DestDataInfo->GetMipData(MipIndex);
+	FColor* SrcTextureData = DataInterface.GetRawHeightData();
+
+	int32 MipSizeSquare = FMath::Square(MipSize);
+	check(MipSizeSquare == DataInterface.GetHeightmapSizeX(MipIndex)*DataInterface.GetHeightmapSizeY(MipIndex));
+
+	FMemory::Memcpy(DestTextureData, SrcTextureData, MipSizeSquare*sizeof(FColor));
+	DestDataInfo->AddMipUpdateRegion(MipIndex, 0, 0, MipSize - 1, MipSize - 1);
+}
 
+void FLandscapeTextureDataInterface::CopyTextureFromWeightmap(UTexture2D* Dest, int32 DestChannel, ULandscapeComponent* Comp, ULandscapeLayerInfoObject* LayerInfo)
+{
+	FLandscapeTextureDataInfo* DestDataInfo = GetTextureDataInfo(Dest);
 	for (int32 MipIdx = 0; MipIdx < DestDataInfo->NumMips(); MipIdx++)
 	{
-		FLandscapeComponentDataInterface DataInterface(Comp, MipIdx);
-		TArray<uint8> WeightData;
-		DataInterface.GetWeightmapTextureData(LayerInfo, WeightData);
-
-		uint8* DestTextureData = (uint8*)DestDataInfo->GetMipData(MipIdx) + ChannelOffsets[DestChannel];
+		CopyTextureFromWeightmap(DestDataInfo, DestChannel, Comp, LayerInfo, MipIdx);
+	}
+}
 
-		for (int32 i = 0; i < FMath::Square(MipSize); i++)
-		{
-			DestTextureData[i * 4] = WeightData[i];
-		}
+void FLandscapeTextureDataInterface::CopyTextureFromWeightmap(UTexture2D* Dest, int32 DestChannel, ULandscapeComponent* Comp, ULandscapeLayerInfoObject* LayerInfo, int32 MipIndex)
+{
+	FLandscapeTextureDataInfo* DestDataInfo = GetTextureDataInfo(Dest);
+	CopyTextureFromWeightmap(DestDataInfo, DestChannel, Comp, LayerInfo, MipIndex);
+}
 
-		DestDataInfo->AddMipUpdateRegion(MipIdx, 0, 0, MipSize - 1, MipSize - 1);
-		MipSize >>= 1;
+void FLandscapeTextureDataInterface::CopyTextureFromWeightmap(FLandscapeTextureDataInfo* DestDataInfo, int32 DestChannel, ULandscapeComponent* Comp, ULandscapeLayerInfoObject* LayerInfo, int32 MipIndex)
+{
+	const int32 MipSize = DestDataInfo->GetMipSizeX(MipIndex);
+	check(MipSize == DestDataInfo->GetMipSizeY(MipIndex));
+
+	const bool bEditingLayer = false;
+	FLandscapeComponentDataInterface DataInterface(Comp, MipIndex, bEditingLayer);
+	TArray<uint8> WeightData;
+	DataInterface.GetWeightmapTextureData(LayerInfo, WeightData, bEditingLayer);
+	int32 MipSizeSquare = FMath::Square(MipSize);
+	check(WeightData.Num() == MipSizeSquare);
+	
+	uint8* DestTextureData = (uint8*)DestDataInfo->GetMipData(MipIndex) + ChannelOffsets[DestChannel];
+	for (int32 i = 0; i < MipSizeSquare; i++)
+	{
+		DestTextureData[i * 4] = WeightData[i];
 	}
+
+	DestDataInfo->AddMipUpdateRegion(MipIndex, 0, 0, MipSize - 1, MipSize - 1);
 }
 
 void FLandscapeTextureDataInterface::ZeroTextureChannel(UTexture2D* Dest, int32 DestChannel)
diff --git a/Engine/Source/Runtime/Landscape/Private/LandscapeEditLayers.cpp b/Engine/Source/Runtime/Landscape/Private/LandscapeEditLayers.cpp
--- a/Engine/Source/Runtime/Landscape/Private/LandscapeEditLayers.cpp
+++ b/Engine/Source/Runtime/Landscape/Private/LandscapeEditLayers.cpp
@@ -39,6 +39,7 @@ LandscapeEditLayers.cpp: Landscape editing layers mode
 #include "Misc/FileHelper.h"
 #include "Misc/MapErrors.h"
 #include "Misc/UObjectToken.h"
+#include "Misc/ScopedSlowTask.h"
 #endif
 
 #define LOCTEXT_NAMESPACE "Landscape"
@@ -5557,6 +5558,72 @@ void ALandscape::DeleteLayer(int32 InLayerIndex)
 	RequestLayersContentUpdateForceAll();
 }
 
+void ALandscape::CollapseLayer(int32 InLayerIndex)
+{
+	FScopedSlowTask SlowTask(GetLandscapeInfo()->XYtoComponentMap.Num(), LOCTEXT("Landscape_CollapseLayer_SlowWork", "Collapsing Layer..."));
+	SlowTask.MakeDialog();
+	TArray<bool> BackupVisibility;
+	TArray<bool> BackupBrushVisibility;
+	for (int32 i = 0; i < LandscapeLayers.Num(); ++i)
+	{
+		BackupVisibility.Add(LandscapeLayers[i].bVisible);
+		LandscapeLayers[i].bVisible = i == InLayerIndex || i == InLayerIndex - 1;
+	}
+
+	for (int32 i = 0; i < LandscapeLayers[InLayerIndex].Brushes.Num(); ++i)
+	{
+		BackupBrushVisibility.Add(LandscapeLayers[InLayerIndex].Brushes[i].GetBrush()->IsVisible());
+		LandscapeLayers[InLayerIndex].Brushes[i].GetBrush()->SetIsVisible(false);
+	}
+
+	// Call Request Update on all components...
+	GetLandscapeInfo()->ForAllLandscapeComponents([](ULandscapeComponent* LandscapeComponent)
+	{
+		LandscapeComponent->RequestWeightmapUpdate(false, false);
+		LandscapeComponent->RequestHeightmapUpdate(false, false);
+	});
+
+	const bool bLocalIntermediateRender = true;
+	ForceUpdateLayersContent(bLocalIntermediateRender);
+
+	// Do copy
+	{
+		FLandscapeEditDataInterface DataInterface(GetLandscapeInfo());
+		DataInterface.SetShouldDirtyPackage(true);
+
+		TSet<UTexture2D*> ProcessedHeightmaps;
+		FScopedSetLandscapeEditingLayer ScopeEditingLayer(this, LandscapeLayers[InLayerIndex - 1].Guid);
+		GetLandscapeInfo()->ForAllLandscapeComponents([&](ULandscapeComponent* LandscapeComponent)
+		{
+			SlowTask.EnterProgressFrame(1.f);
+			LandscapeComponent->CopyFinalLayerIntoEditingLayer(DataInterface, ProcessedHeightmaps);
+		});
+	}
+	
+	TArray<ALandscapeBlueprintBrushBase*> BrushesToMove;
+	for (int32 i = 0; i < LandscapeLayers[InLayerIndex].Brushes.Num(); ++i)
+	{
+		ALandscapeBlueprintBrushBase* CurrentBrush = LandscapeLayers[InLayerIndex].Brushes[i].GetBrush();
+		CurrentBrush->SetIsVisible(BackupBrushVisibility[i]);
+		BrushesToMove.Add(CurrentBrush);
+	}
+
+	for (ALandscapeBlueprintBrushBase* Brush : BrushesToMove)
+	{
+		RemoveBrushFromLayer(InLayerIndex, Brush);
+		AddBrushToLayer(InLayerIndex - 1, Brush);
+	}
+
+	for (int32 i = 0; i < LandscapeLayers.Num(); ++i)
+	{
+		LandscapeLayers[i].bVisible = BackupVisibility[i];
+	}
+
+	DeleteLayer(InLayerIndex);
+	
+	RequestLayersContentUpdateForceAll();
+}
+
 void ALandscape::GetUsedPaintLayers(int32 InLayerIndex, TArray<ULandscapeLayerInfoObject*>& OutUsedLayerInfos) const
 {
 	const FLandscapeLayer* Layer = GetLayer(InLayerIndex);
diff --git a/Engine/Source/Runtime/Landscape/Public/LandscapeDataAccess.h b/Engine/Source/Runtime/Landscape/Public/LandscapeDataAccess.h
--- a/Engine/Source/Runtime/Landscape/Public/LandscapeDataAccess.h
+++ b/Engine/Source/Runtime/Landscape/Public/LandscapeDataAccess.h
@@ -210,6 +210,9 @@ struct FLandscapeComponentDataInterface
 		XYOffsetMipData = NewXYOffsetData;
 	}
 
+	LANDSCAPE_API int32 GetHeightmapSizeX(int32 MipIndex) const;
+	LANDSCAPE_API int32 GetHeightmapSizeY(int32 MipIndex) const;
+
 	/* Return the raw heightmap data exactly same size for Heightmap texture which belong to only this component */
 	LANDSCAPE_API void GetHeightmapTextureData(TArray<FColor>& OutData, bool bOkToFail = false);
 
diff --git a/Engine/Source/Runtime/Landscape/Public/LandscapeEdit.h b/Engine/Source/Runtime/Landscape/Public/LandscapeEdit.h
--- a/Engine/Source/Runtime/Landscape/Public/LandscapeEdit.h
+++ b/Engine/Source/Runtime/Landscape/Public/LandscapeEdit.h
@@ -86,12 +86,12 @@ struct FLandscapeTextureDataInfo
 		return MipInfo[MipNum].MipData;
 	}
 
-	int32 GetMipSizeX(int32 MipNum)
+	int32 GetMipSizeX(int32 MipNum) const
 	{
 		return FMath::Max(Texture->Source.GetSizeX() >> MipNum, 1);
 	}
 
-	int32 GetMipSizeY(int32 MipNum)
+	int32 GetMipSizeY(int32 MipNum) const
 	{
 		return FMath::Max(Texture->Source.GetSizeY() >> MipNum, 1);
 	}
@@ -118,8 +118,10 @@ struct LANDSCAPE_API FLandscapeTextureDataInterface
 	// Texture bulk operations for weightmap reallocation
 	void CopyTextureChannel(UTexture2D* Dest, int32 DestChannel, UTexture2D* Src, int32 SrcChannel);
 	void ZeroTextureChannel(UTexture2D* Dest, int32 DestChannel);
+	void CopyTextureFromHeightmap(UTexture2D* Dest, ULandscapeComponent* Comp, int32 MipIndex);
 	void CopyTextureFromHeightmap(UTexture2D* Dest, int32 DestChannel, ULandscapeComponent* Comp, int32 SrcChannel);
 	void CopyTextureFromWeightmap(UTexture2D* Dest, int32 DestChannel, ULandscapeComponent* Comp, ULandscapeLayerInfoObject* LayerInfo);
+	void CopyTextureFromWeightmap(UTexture2D* Dest, int32 DestChannel, ULandscapeComponent* Comp, ULandscapeLayerInfoObject* LayerInfo, int32 MipIndex);
 
 	template<typename TData>
 	void SetTextureValueTempl(UTexture2D* Dest, TData Value);
@@ -131,6 +133,8 @@ struct LANDSCAPE_API FLandscapeTextureDataInterface
 	bool EqualTextureValue(UTexture2D* Src, FColor Value);
 
 private:
+	void CopyTextureFromWeightmap(FLandscapeTextureDataInfo* DestDataInfo, int32 DestChannel, ULandscapeComponent* Comp, ULandscapeLayerInfoObject* LayerInfo, int32 MipIndex);
+
 	TMap<UTexture2D*, FLandscapeTextureDataInfo*> TextureDataMap;
 	bool bUploadTextureChangesToGPU;
 	bool bShouldDirtyPackage;
