Landscape Layers: UI Reverse order of Layers

diff --git a/Engine/Source/Editor/LandscapeEditor/Private/LandscapeEditorDetailCustomization_Layers.cpp b/Engine/Source/Editor/LandscapeEditor/Private/LandscapeEditorDetailCustomization_Layers.cpp
--- a/Engine/Source/Editor/LandscapeEditor/Private/LandscapeEditorDetailCustomization_Layers.cpp
+++ b/Engine/Source/Editor/LandscapeEditor/Private/LandscapeEditorDetailCustomization_Layers.cpp
@@ -152,7 +152,8 @@ void FLandscapeEditorCustomNodeBuilder_Layers::GenerateChildContent(IDetailChild
 
 		InlineTextBlocks.Empty();
 		InlineTextBlocks.AddDefaulted(LandscapeEdMode->GetLayerCount());
-		for (int32 i = 0; i < LandscapeEdMode->GetLayerCount(); ++i)
+		// Slots are displayed in the opposite order of LandscapeLayers
+		for (int32 i = LandscapeEdMode->GetLayerCount()-1; i >= 0 ; --i)
 		{
 			TSharedPtr<SWidget> GeneratedRowWidget = GenerateRow(i);
 
@@ -853,15 +854,29 @@ const FSlateBrush* FLandscapeEditorCustomNodeBuilder_Layers::GetLockBrushForLaye
 	return bIsLocked ? FEditorStyle::GetBrush(TEXT("PropertyWindow.Locked")) : FEditorStyle::GetBrush(TEXT("PropertyWindow.Unlocked"));
 }
 
+int32 FLandscapeEditorCustomNodeBuilder_Layers::SlotIndexToLayerIndex(int32 SlotIndex)
+{
+	FEdModeLandscape* LandscapeEdMode = GetEditorMode();
+	ALandscape* Landscape = LandscapeEdMode ? LandscapeEdMode->GetLandscape() : nullptr;
+	if (!Landscape)
+	{
+		return INDEX_NONE;
+	}
+	
+	check(Landscape->LandscapeLayers.IsValidIndex(SlotIndex));
+	return Landscape->LandscapeLayers.Num() - SlotIndex - 1;
+}
+
 FReply FLandscapeEditorCustomNodeBuilder_Layers::HandleDragDetected(const FGeometry& MyGeometry, const FPointerEvent& MouseEvent, int32 SlotIndex, SVerticalBox::FSlot* Slot)
 {
 	FEdModeLandscape* LandscapeEdMode = GetEditorMode();
 	if (LandscapeEdMode)
 	{
-		FLandscapeLayer* Layer = LandscapeEdMode->GetLayer(SlotIndex);
+		int32 LayerIndex = SlotIndexToLayerIndex(SlotIndex);
+		FLandscapeLayer* Layer = LandscapeEdMode->GetLayer(LayerIndex);
 		if (Layer && !Layer->bLocked)
 		{
-			TSharedPtr<SWidget> Row = GenerateRow(SlotIndex);
+			TSharedPtr<SWidget> Row = GenerateRow(LayerIndex);
 			if (Row.IsValid())
 			{
 				return FReply::Handled().BeginDragDrop(FLandscapeListElementDragDropOp::New(SlotIndex, Slot, Row));
@@ -891,8 +906,8 @@ FReply FLandscapeEditorCustomNodeBuilder_Layers::HandleAcceptDrop(FDragDropEvent
 		ALandscape* Landscape = LandscapeEdMode ? LandscapeEdMode->GetLandscape() : nullptr;
 		if (Landscape)
 		{
-			int32 StartingLayerIndex = DragDropOperation->SlotIndexBeingDragged;
-			int32 DestinationLayerIndex = SlotIndex;
+			int32 StartingLayerIndex = SlotIndexToLayerIndex(DragDropOperation->SlotIndexBeingDragged);
+			int32 DestinationLayerIndex = SlotIndexToLayerIndex(SlotIndex);
 			const FScopedTransaction Transaction(LOCTEXT("Landscape_Layers_Reorder", "Reorder Layer"));
 			if (Landscape->ReorderLayer(StartingLayerIndex, DestinationLayerIndex))
 			{
diff --git a/Engine/Source/Editor/LandscapeEditor/Private/LandscapeEditorDetailCustomization_Layers.h b/Engine/Source/Editor/LandscapeEditor/Private/LandscapeEditorDetailCustomization_Layers.h
--- a/Engine/Source/Editor/LandscapeEditor/Private/LandscapeEditorDetailCustomization_Layers.h
+++ b/Engine/Source/Editor/LandscapeEditor/Private/LandscapeEditorDetailCustomization_Layers.h
@@ -66,6 +66,7 @@ protected:
 	TSharedPtr<SWidget> GenerateRow(int32 InLayerIndex);
 
 	// Drag/Drop handling
+	int32 SlotIndexToLayerIndex(int32 SlotIndex);
 	FReply HandleDragDetected(const FGeometry& MyGeometry, const FPointerEvent& MouseEvent, int32 SlotIndex, SVerticalBox::FSlot* Slot);
 	TOptional<SDragAndDropVerticalBox::EItemDropZone> HandleCanAcceptDrop(const FDragDropEvent& DragDropEvent, SDragAndDropVerticalBox::EItemDropZone DropZone, SVerticalBox::FSlot* Slot);
 	FReply HandleAcceptDrop(FDragDropEvent const& DragDropEvent, SDragAndDropVerticalBox::EItemDropZone DropZone, int32 SlotIndex, SVerticalBox::FSlot* Slot);
