Fix Static Analysis

#rnx

diff --git a/Engine/Source/Editor/LandscapeEditor/Private/LandscapeEditorDetailCustomization_Layers.cpp b/Engine/Source/Editor/LandscapeEditor/Private/LandscapeEditorDetailCustomization_Layers.cpp
--- a/Engine/Source/Editor/LandscapeEditor/Private/LandscapeEditorDetailCustomization_Layers.cpp
+++ b/Engine/Source/Editor/LandscapeEditor/Private/LandscapeEditorDetailCustomization_Layers.cpp
@@ -675,15 +675,21 @@ bool FLandscapeEditorCustomNodeBuilder_Layers::CanCollapseLayer(int32 InLayerInd
 		return false;
 	}
 
-	// Can't collapse on layer that has a Brush because result will change...
 	FLandscapeLayer* BaseLayer = LandscapeEdMode->GetLayer(InLayerIndex - 1);
+	if (!BaseLayer)
+	{
+		OutReason = LOCTEXT("Landscape_CollapseLayer_Reason_InvalidBaseLayer", "Invalid base layer.");
+		return false;
+	}
+
 	if (SplineLayer == BaseLayer)
 	{
 		OutReason = LOCTEXT("Landscape_CollapseLayer_Reason_CantCollapseOnSplineLayer", "Can't Collapse on reserved spline layer.");
 		return false;
 	}
 	
-	if (!BaseLayer || BaseLayer->Brushes.Num() > 0)
+	// Can't collapse on layer that has a Brush because result will change...
+	if (BaseLayer->Brushes.Num() > 0)
 	{
 		OutReason = FText::Format(LOCTEXT("Landscape_CollapseLayer_Reason_CantCollapseOnLayerWithBrush", "Can't Collapse because base layer '{0}' contains Brush(es)."), FText::FromName(BaseLayer->Name));
 		return false;
@@ -698,19 +704,21 @@ void FLandscapeEditorCustomNodeBuilder_Layers::CollapseLayer(int32 InLayerIndex)
 	FText Reason;
 	if (CanCollapseLayer(InLayerIndex, Reason))
 	{
-		FEdModeLandscape* LandscapeEdMode = GetEditorMode();
-		ALandscape* Landscape = LandscapeEdMode ? LandscapeEdMode->GetLandscape() : nullptr;
-		FLandscapeLayer* Layer = LandscapeEdMode->GetLayer(InLayerIndex);
-		FLandscapeLayer* BaseLayer = LandscapeEdMode->GetLayer(InLayerIndex - 1);
-		if (Layer && BaseLayer)
+		if (FEdModeLandscape* LandscapeEdMode = GetEditorMode())
 		{
-			EAppReturnType::Type Result = FMessageDialog::Open(EAppMsgType::YesNo, FText::Format(LOCTEXT("Landscape_CollapseLayer_Message", "The layer {0} will be collapsed into layer {1}.  Continue?"), FText::FromName(Layer->Name), FText::FromName(BaseLayer->Name)));
-			if (Result == EAppReturnType::Yes)
+			ALandscape* Landscape = LandscapeEdMode->GetLandscape();
+			FLandscapeLayer* Layer = LandscapeEdMode->GetLayer(InLayerIndex);
+			FLandscapeLayer* BaseLayer = LandscapeEdMode->GetLayer(InLayerIndex - 1);
+			if (Landscape && Layer && BaseLayer)
 			{
-				const FScopedTransaction Transaction(LOCTEXT("Landscape_Layers_Collapse", "Collapse Layer"));
-				Landscape->CollapseLayer(InLayerIndex);
-				OnLayerSelectionChanged(InLayerIndex-1);
-				LandscapeEdMode->RefreshDetailPanel();
+				EAppReturnType::Type Result = FMessageDialog::Open(EAppMsgType::YesNo, FText::Format(LOCTEXT("Landscape_CollapseLayer_Message", "The layer {0} will be collapsed into layer {1}.  Continue?"), FText::FromName(Layer->Name), FText::FromName(BaseLayer->Name)));
+				if (Result == EAppReturnType::Yes)
+				{
+					const FScopedTransaction Transaction(LOCTEXT("Landscape_Layers_Collapse", "Collapse Layer"));
+					Landscape->CollapseLayer(InLayerIndex);
+					OnLayerSelectionChanged(InLayerIndex - 1);
+					LandscapeEdMode->RefreshDetailPanel();
+				}
 			}
 		}
 	}
