Add FlipSegment Button (for existing 'F' shortcut)

#rnx

diff --git a/Engine/Source/Editor/LandscapeEditor/Private/LandscapeEdMode.h b/Engine/Source/Editor/LandscapeEditor/Private/LandscapeEdMode.h
--- a/Engine/Source/Editor/LandscapeEditor/Private/LandscapeEdMode.h
+++ b/Engine/Source/Editor/LandscapeEditor/Private/LandscapeEdMode.h
@@ -317,6 +317,8 @@ public:
 	// UI callbacks for splines tool
 	FLandscapeToolSplines* SplinesTool;
 	void ShowSplineProperties();
+	bool HasSelectedSplineSegments() const;
+	void FlipSelectedSplineSegments();
 	virtual void SelectAllConnectedSplineControlPoints();
 	virtual void SelectAllConnectedSplineSegments();
 	virtual void SplineMoveToCurrentLevel();
diff --git a/Engine/Source/Editor/LandscapeEditor/Private/LandscapeEdModeSplineTools.cpp b/Engine/Source/Editor/LandscapeEditor/Private/LandscapeEdModeSplineTools.cpp
--- a/Engine/Source/Editor/LandscapeEditor/Private/LandscapeEdModeSplineTools.cpp
+++ b/Engine/Source/Editor/LandscapeEditor/Private/LandscapeEdModeSplineTools.cpp
@@ -364,6 +364,16 @@ public:
 		}
 	}
 
+	void FlipSelectedSplineSegments()
+	{
+		FScopedTransaction Transaction(LOCTEXT("LandscapeSpline_FlipSegments", "Flip Selected Landscape Spline Segments"));
+		for (ULandscapeSplineSegment* Segment : SelectedSplineSegments)
+		{
+			FlipSegment(Segment);
+		}
+		EdMode->AutoUpdateDirtyLandscapeSplines();
+	}
+
 	// called when alt-dragging a newly added end segment
 	bool UpdateAddSegment(ULandscapeSplineControlPoint* ControlPoint, FVector Location)
 	{
@@ -1151,14 +1161,7 @@ public:
 		{
 			if (SelectedSplineSegments.Num() > 0)
 			{
-				FScopedTransaction Transaction(LOCTEXT("LandscapeSpline_FlipSegments", "Flip Landscape Spline Segments"));
-
-				for (ULandscapeSplineSegment* Segment : SelectedSplineSegments)
-				{
-					FlipSegment(Segment);
-				}
-
-				EdMode->AutoUpdateDirtyLandscapeSplines();
+				FlipSelectedSplineSegments();
 				return true;
 			}
 		}
@@ -2223,6 +2226,22 @@ protected:
 };
 
 
+bool FEdModeLandscape::HasSelectedSplineSegments() const
+{
+	return SplinesTool && (SplinesTool->SelectedSplineSegments.Num() > 0);
+}
+
+void FEdModeLandscape::FlipSelectedSplineSegments()
+{
+	if (!SplinesTool)
+	{
+		return;
+	}
+
+	// Do Flip
+	SplinesTool->FlipSelectedSplineSegments();
+}
+
 void FEdModeLandscape::ShowSplineProperties()
 {
 	if (SplinesTool /*&& SplinesTool == CurrentTool*/)
diff --git a/Engine/Source/Editor/LandscapeEditor/Private/LandscapeSplineDetails.cpp b/Engine/Source/Editor/LandscapeEditor/Private/LandscapeSplineDetails.cpp
--- a/Engine/Source/Editor/LandscapeEditor/Private/LandscapeSplineDetails.cpp
+++ b/Engine/Source/Editor/LandscapeEditor/Private/LandscapeSplineDetails.cpp
@@ -83,6 +83,23 @@ void FLandscapeSplineDetails::CustomizeDetails(IDetailLayoutBuilder& DetailBuild
 			.IsEnabled(this, &FLandscapeSplineDetails::IsUpdateSplineMeshLevelsButtonEnabled)
 		]
 	];
+
+	IDetailCategoryBuilder& LandscapeSplineSegmentCategory = DetailBuilder.EditCategory("LandscapeSplineSegment", FText::GetEmpty(), ECategoryPriority::Default);
+	LandscapeSplineSegmentCategory.AddCustomRow(FText::GetEmpty())
+	[
+		SNew(SHorizontalBox)
+		+ SHorizontalBox::Slot()
+		.Padding(0, 0, 2, 0)
+		.VAlign(VAlign_Center)
+		.FillWidth(1)
+		[
+			SNew(SButton)
+			.Text(LOCTEXT("FlipSegment", "Flip Selected Segment(s)"))
+			.HAlign(HAlign_Center)
+			.OnClicked(this, &FLandscapeSplineDetails::OnFlipSegmentButtonClicked)
+			.IsEnabled(this, &FLandscapeSplineDetails::IsFlipSegmentButtonEnabled)
+		]
+	];
 }
 
 class FEdModeLandscape* FLandscapeSplineDetails::GetEditorMode() const
@@ -90,6 +107,22 @@ class FEdModeLandscape* FLandscapeSplineDetails::GetEditorMode() const
 	return (FEdModeLandscape*)GLevelEditorModeTools().GetActiveMode(FBuiltinEditorModes::EM_Landscape);
 }
 
+FReply FLandscapeSplineDetails::OnFlipSegmentButtonClicked()
+{
+	FEdModeLandscape* LandscapeEdMode = GetEditorMode();
+	if (LandscapeEdMode)
+	{
+		LandscapeEdMode->FlipSelectedSplineSegments();
+	}
+	return FReply::Handled();
+}
+
+bool FLandscapeSplineDetails::IsFlipSegmentButtonEnabled() const
+{
+	FEdModeLandscape* LandscapeEdMode = GetEditorMode();
+	return LandscapeEdMode && LandscapeEdMode->HasSelectedSplineSegments();
+}
+
 FReply FLandscapeSplineDetails::OnSelectConnectedControlPointsButtonClicked()
 {
 	FEdModeLandscape* LandscapeEdMode = GetEditorMode();
diff --git a/Engine/Source/Editor/LandscapeEditor/Private/LandscapeSplineDetails.h b/Engine/Source/Editor/LandscapeEditor/Private/LandscapeSplineDetails.h
--- a/Engine/Source/Editor/LandscapeEditor/Private/LandscapeSplineDetails.h
+++ b/Engine/Source/Editor/LandscapeEditor/Private/LandscapeSplineDetails.h
@@ -27,4 +27,7 @@ private:
 
 	FReply OnUpdateSplineMeshLevelsButtonClicked();
 	bool IsUpdateSplineMeshLevelsButtonEnabled() const;
+	
+	FReply OnFlipSegmentButtonClicked();
+	bool IsFlipSegmentButtonEnabled() const;
 };
