Fixed Add Ramp landscape tool debug display when landscape is rotated.

diff --git a/Engine/Source/Editor/LandscapeEditor/Private/LandscapeEdModeRampTool.cpp b/Engine/Source/Editor/LandscapeEditor/Private/LandscapeEdModeRampTool.cpp
--- a/Engine/Source/Editor/LandscapeEditor/Private/LandscapeEdModeRampTool.cpp
+++ b/Engine/Source/Editor/LandscapeEditor/Private/LandscapeEdModeRampTool.cpp
@@ -364,8 +364,11 @@ public:
 			if (NumPoints == 2)
 			{
 				const FVector Side = FVector::CrossProduct(Points[1] - Points[0], FVector(0, 0, 1)).GetSafeNormal2D();
-				const FVector InnerSide = Side * (EdMode->UISettings->RampWidth * 0.5f * (1 - EdMode->UISettings->RampSideFalloff));
-				const FVector OuterSide = Side * (EdMode->UISettings->RampWidth * 0.5f);
+				FVector InnerSide = Side * (EdMode->UISettings->RampWidth * 0.5f * (1 - EdMode->UISettings->RampSideFalloff));
+				FVector OuterSide = Side * (EdMode->UISettings->RampWidth * 0.5f);
+				InnerSide = LandscapeToWorld.TransformVectorNoScale(InnerSide);
+				OuterSide = LandscapeToWorld.TransformVectorNoScale(OuterSide);
+
 				FVector InnerVerts[2][2];
 				InnerVerts[0][0] = WorldPoints[0] - InnerSide;
 				InnerVerts[0][1] = WorldPoints[0] + InnerSide;
