Prevent possible crash if we have an invalid landscape, and simply return a null reimport path instead

diff --git a/Engine/Source/Editor/LandscapeEditor/Private/LandscapeEdMode.cpp b/Engine/Source/Editor/LandscapeEditor/Private/LandscapeEdMode.cpp
--- a/Engine/Source/Editor/LandscapeEditor/Private/LandscapeEdMode.cpp
+++ b/Engine/Source/Editor/LandscapeEditor/Private/LandscapeEdMode.cpp
@@ -3440,7 +3440,7 @@ void FEdModeLandscape::ForceRealTimeViewports(const bool bEnable, const bool bSt
 
 void FEdModeLandscape::ReimportData(const FLandscapeTargetListInfo& TargetInfo)
 {
-	const FString& SourceFilePath = TargetInfo.ReimportFilePath();
+	const FString& SourceFilePath = TargetInfo.GetReimportFilePath();
 	if (SourceFilePath.Len())
 	{
 		FScopedSetLandscapeEditingLayer Scope(GetLandscape(), GetCurrentLayerGuid(), [&] { RequestLayersContentUpdateForceAll(); });
diff --git a/Engine/Source/Editor/LandscapeEditor/Private/LandscapeEdMode.h b/Engine/Source/Editor/LandscapeEditor/Private/LandscapeEdMode.h
--- a/Engine/Source/Editor/LandscapeEditor/Private/LandscapeEdMode.h
+++ b/Engine/Source/Editor/LandscapeEditor/Private/LandscapeEdMode.h
@@ -152,7 +152,7 @@ struct FLandscapeTargetListInfo
 
 	FName GetLayerName() const;
 
-	FString& ReimportFilePath() const
+	FString GetReimportFilePath() const
 	{
 		if (TargetType == ELandscapeToolTargetType::Weightmap)
 		{
@@ -162,7 +162,39 @@ struct FLandscapeTargetListInfo
 		}
 		else //if (TargetType == ELandscapeToolTargetType::Heightmap)
 		{
-			return LandscapeInfo->GetLandscapeProxy()->ReimportHeightmapFilePath;
+			if (LandscapeInfo.IsValid())
+			{
+				ALandscapeProxy* LandscapeProxy = LandscapeInfo->GetLandscapeProxy();
+
+				if (LandscapeProxy)
+				{
+					return LandscapeProxy->ReimportHeightmapFilePath;
+				}
+			}
+
+			return FString(TEXT(""));
+		}
+	}
+
+	void SetReimportFilePath(const FString& InNewPath)
+	{
+		if (TargetType == ELandscapeToolTargetType::Weightmap)
+		{
+			FLandscapeEditorLayerSettings* EditorLayerSettings = GetEditorLayerSettings();
+			check(EditorLayerSettings);
+			EditorLayerSettings->ReimportLayerFilePath = InNewPath;
+		}
+		else //if (TargetType == ELandscapeToolTargetType::Heightmap)
+		{
+			if (LandscapeInfo.IsValid())
+			{
+				ALandscapeProxy* LandscapeProxy = LandscapeInfo->GetLandscapeProxy();
+
+				if (LandscapeProxy)
+				{
+					LandscapeProxy->ReimportHeightmapFilePath = InNewPath;
+				}
+			}
 		}
 	}
 };
diff --git a/Engine/Source/Editor/LandscapeEditor/Private/LandscapeEditorDetailCustomization_TargetLayers.cpp b/Engine/Source/Editor/LandscapeEditor/Private/LandscapeEditorDetailCustomization_TargetLayers.cpp
--- a/Engine/Source/Editor/LandscapeEditor/Private/LandscapeEditorDetailCustomization_TargetLayers.cpp
+++ b/Engine/Source/Editor/LandscapeEditor/Private/LandscapeEditorDetailCustomization_TargetLayers.cpp
@@ -917,7 +917,7 @@ TSharedPtr<SWidget> FLandscapeEditorCustomNodeBuilder_TargetLayers::OnTargetLaye
 			MenuBuilder.AddMenuEntry(LOCTEXT("LayerContextMenu.Import", "Import from file"), FText(), FSlateIcon(), ImportAction);
 
 			// Reimport
-			const FString& ReimportPath = Target->ReimportFilePath();
+			const FString& ReimportPath = Target->GetReimportFilePath();
 
 			if (!ReimportPath.IsEmpty())
 			{
@@ -1015,7 +1015,7 @@ void FLandscapeEditorCustomNodeBuilder_TargetLayers::OnExportLayer(const TShared
 				LandscapeInfo->ExportLayer(LayerInfoObj, SaveFilename);
 			}
 
-			Target->ReimportFilePath() = SaveFilename;
+			Target->SetReimportFilePath(SaveFilename);
 		}
 	}
 }
@@ -1070,7 +1070,7 @@ void FLandscapeEditorCustomNodeBuilder_TargetLayers::OnImportLayer(const TShared
 			// Actually do the Import
 			LandscapeEdMode->ImportData(*Target, OpenFilename);
 
-			Target->ReimportFilePath() = OpenFilename;
+			Target->SetReimportFilePath(OpenFilename);
 		}
 	}
 }
