Allow spinning again, and correctly update the property when dragging the sliders

diff --git a/Engine/Source/Editor/LandscapeEditor/Private/LandscapeEditorDetailCustomization_NewLandscape.cpp b/Engine/Source/Editor/LandscapeEditor/Private/LandscapeEditorDetailCustomization_NewLandscape.cpp
--- a/Engine/Source/Editor/LandscapeEditor/Private/LandscapeEditorDetailCustomization_NewLandscape.cpp
+++ b/Engine/Source/Editor/LandscapeEditor/Private/LandscapeEditorDetailCustomization_NewLandscape.cpp
@@ -216,6 +216,10 @@ void FLandscapeEditorDetailCustomization_NewLandscape::CustomizeDetails(IDetailL
 		.OnXCommitted_Static(&SetPropertyValue<float>, PropertyHandle_Location_X)
 		.OnYCommitted_Static(&SetPropertyValue<float>, PropertyHandle_Location_Y)
 		.OnZCommitted_Static(&SetPropertyValue<float>, PropertyHandle_Location_Z)
+		.OnXChanged_Lambda([=](float NewValue) { ensure(PropertyHandle_Location_X->SetValue(NewValue, EPropertyValueSetFlags::InteractiveChange) == FPropertyAccess::Success); })
+		.OnYChanged_Lambda([=](float NewValue) { ensure(PropertyHandle_Location_Y->SetValue(NewValue, EPropertyValueSetFlags::InteractiveChange) == FPropertyAccess::Success); })
+		.OnZChanged_Lambda([=](float NewValue) { ensure(PropertyHandle_Location_Z->SetValue(NewValue, EPropertyValueSetFlags::InteractiveChange) == FPropertyAccess::Success); })
+		.AllowSpin(true)
 	];
 
 	TSharedRef<IPropertyHandle> PropertyHandle_Rotation = DetailBuilder.GetProperty(GET_MEMBER_NAME_CHECKED(ULandscapeEditorObject, NewLandscape_Rotation));
@@ -241,6 +245,7 @@ void FLandscapeEditorDetailCustomization_NewLandscape::CustomizeDetails(IDetailL
 		.Yaw_Static(&GetOptionalPropertyValue<float>, PropertyHandle_Rotation_Yaw)
 		.OnYawCommitted_Static(&SetPropertyValue<float>, PropertyHandle_Rotation_Yaw) // not allowed to roll or pitch landscape
 		.OnYawChanged_Lambda([=](float NewValue){ ensure(PropertyHandle_Rotation_Yaw->SetValue(NewValue, EPropertyValueSetFlags::InteractiveChange) == FPropertyAccess::Success); })
+		.AllowSpin(true)
 	];
 
 	TSharedRef<IPropertyHandle> PropertyHandle_Scale = DetailBuilder.GetProperty(GET_MEMBER_NAME_CHECKED(ULandscapeEditorObject, NewLandscape_Scale));
@@ -266,6 +271,10 @@ void FLandscapeEditorDetailCustomization_NewLandscape::CustomizeDetails(IDetailL
 		.OnXCommitted_Static(&SetScale, PropertyHandle_Scale_X)
 		.OnYCommitted_Static(&SetScale, PropertyHandle_Scale_Y)
 		.OnZCommitted_Static(&SetScale, PropertyHandle_Scale_Z)
+		.OnXChanged_Lambda([=](float NewValue) { ensure(PropertyHandle_Scale_X->SetValue(NewValue, EPropertyValueSetFlags::InteractiveChange) == FPropertyAccess::Success); })
+		.OnYChanged_Lambda([=](float NewValue) { ensure(PropertyHandle_Scale_Y->SetValue(NewValue, EPropertyValueSetFlags::InteractiveChange) == FPropertyAccess::Success); })
+		.OnZChanged_Lambda([=](float NewValue) { ensure(PropertyHandle_Scale_Z->SetValue(NewValue, EPropertyValueSetFlags::InteractiveChange) == FPropertyAccess::Success); })
+		.AllowSpin(true)
 	];
 
 	TSharedRef<IPropertyHandle> PropertyHandle_QuadsPerSection = DetailBuilder.GetProperty(GET_MEMBER_NAME_CHECKED(ULandscapeEditorObject, NewLandscape_QuadsPerSection));
