Landscape Visibility Painting: Fix tool so Visibility can be painted even if some Landscape actors don't have a Hole Material set. (check per component)

diff --git a/Engine/Source/Editor/LandscapeEditor/Private/LandscapeEdModeBrushes.cpp b/Engine/Source/Editor/LandscapeEditor/Private/LandscapeEdModeBrushes.cpp
--- a/Engine/Source/Editor/LandscapeEditor/Private/LandscapeEdModeBrushes.cpp
+++ b/Engine/Source/Editor/LandscapeEditor/Private/LandscapeEdModeBrushes.cpp
@@ -206,7 +206,7 @@ public:
 			{
 				bCanPaint = false;
 			}
-			else if (EdMode->CurrentToolTarget.TargetType == ELandscapeToolTargetType::Weightmap &&
+			else if ((EdMode->CurrentToolTarget.TargetType == ELandscapeToolTargetType::Weightmap || EdMode->CurrentToolTarget.TargetType == ELandscapeToolTargetType::Visibility) &&
 				EdMode->UISettings->PaintingRestriction != ELandscapeLayerPaintingRestriction::None)
 			{
 				if (EdMode->UISettings->PaintingRestriction == ELandscapeLayerPaintingRestriction::UseComponentWhitelist &&
@@ -230,6 +230,11 @@ public:
 					}
 				}
 			}
+			
+			if (bCanPaint && EdMode->CurrentToolTarget.TargetType == ELandscapeToolTargetType::Visibility && !Component->IsLandscapeHoleMaterialValid())
+			{
+				bCanPaint = false;
+			}
 
 			MaterialInstance->SetScalarParameterValue("CanPaint", bCanPaint ? 1.0f : 0.0f);
 		}
diff --git a/Engine/Source/Editor/LandscapeEditor/Private/LandscapeEdModeComponentTools.cpp b/Engine/Source/Editor/LandscapeEditor/Private/LandscapeEdModeComponentTools.cpp
--- a/Engine/Source/Editor/LandscapeEditor/Private/LandscapeEdModeComponentTools.cpp
+++ b/Engine/Source/Editor/LandscapeEditor/Private/LandscapeEdModeComponentTools.cpp
@@ -327,19 +327,6 @@ public:
 
 	virtual bool BeginTool(FEditorViewportClient* ViewportClient, const FLandscapeToolTarget& InTarget, const FVector& InHitLocation) override
 	{
-		ALandscapeProxy* Proxy = InTarget.LandscapeInfo->GetLandscapeProxy();
-		UMaterialInterface* HoleMaterial = Proxy->GetLandscapeHoleMaterial();
-		if (!HoleMaterial)
-		{
-			HoleMaterial = Proxy->GetLandscapeMaterial();
-		}
-		if (!HoleMaterial->GetMaterial()->HasAnyExpressionsInMaterialAndFunctionsOfType<UMaterialExpressionLandscapeVisibilityMask>())
-		{
-			FMessageDialog::Open(EAppMsgType::Ok,
-				LOCTEXT("LandscapeVisibilityMaskMissing", "You must add a \"Landscape Visibility Mask\" node to your material before you can paint visibility."));
-			return false;
-		}
-
 		return FLandscapeToolBase<FLandscapeToolStrokeVisibility>::BeginTool(ViewportClient, InTarget, InHitLocation);
 	}
 
diff --git a/Engine/Source/Editor/LandscapeEditor/Private/LandscapeEditorDetailCustomization_TargetLayers.cpp b/Engine/Source/Editor/LandscapeEditor/Private/LandscapeEditorDetailCustomization_TargetLayers.cpp
--- a/Engine/Source/Editor/LandscapeEditor/Private/LandscapeEditorDetailCustomization_TargetLayers.cpp
+++ b/Engine/Source/Editor/LandscapeEditor/Private/LandscapeEditorDetailCustomization_TargetLayers.cpp
@@ -15,9 +15,9 @@
 #include "Widgets/Images/SImage.h"
 #include "Widgets/Layout/SBox.h"
 #include "Widgets/Input/SButton.h"
-#include "Widgets/Notifications/SErrorText.h"
 #include "Widgets/Input/SComboButton.h"
 #include "Widgets/Input/SCheckBox.h"
+#include "Widgets/Input/SMultiLineEditableTextBox.h"
 #include "EditorModeManager.h"
 #include "EditorModes.h"
 #include "LandscapeEditorModule.h"
@@ -71,14 +71,15 @@ void FLandscapeEditorDetailCustomization_TargetLayers::CustomizeDetails(IDetailL
 
 	TargetsCategory.AddProperty(PropertyHandle_PaintingRestriction)
 	.Visibility(TAttribute<EVisibility>::Create(TAttribute<EVisibility>::FGetter::CreateStatic(&FLandscapeEditorDetailCustomization_TargetLayers::GetVisibility_PaintingRestriction)));
-
+		
 	TargetsCategory.AddCustomRow(FText())
 	.Visibility(TAttribute<EVisibility>::Create(TAttribute<EVisibility>::FGetter::CreateStatic(&FLandscapeEditorDetailCustomization_TargetLayers::GetVisibility_VisibilityTip)))
 	[
-		SNew(SErrorText)
+		SNew(SMultiLineEditableTextBox)
 		.Font(DetailBuilder.GetDetailFontBold())
+		.BackgroundColor(TAttribute<FSlateColor>::Create(TAttribute<FSlateColor>::FGetter::CreateLambda([]() { return FEditorStyle::GetColor("ErrorReporting.WarningBackgroundColor"); })))
+		.Text(LOCTEXT("Visibility_Tip", "Note: There are some areas where visibility painting is disabled because Component/Proxy don't have a \"Landscape Visibility Mask\" node in their material."))
 		.AutoWrapText(true)
-		.ErrorText(LOCTEXT("Visibility_Tip","Note: You must add a \"Landscape Visibility Mask\" node to your material before you can paint visibility."))
 	];
 
 	TargetsCategory.AddCustomBuilder(MakeShareable(new FLandscapeEditorCustomNodeBuilder_TargetLayers(DetailBuilder.GetThumbnailPool().ToSharedRef(), PropertyHandle_TargetDisplayOrder, PropertyHandle_TargetShowUnusedLayers)));
@@ -137,16 +138,14 @@ bool FLandscapeEditorDetailCustomization_TargetLayers::ShouldShowVisibilityTip()
 	{
 		if (LandscapeEdMode->CurrentToolTarget.TargetType == ELandscapeToolTargetType::Visibility)
 		{
-			ALandscapeProxy* Proxy = LandscapeEdMode->CurrentToolTarget.LandscapeInfo->GetLandscapeProxy();
-			UMaterialInterface* HoleMaterial = Proxy->GetLandscapeHoleMaterial();
-			if (!HoleMaterial)
-			{
-				HoleMaterial = Proxy->GetLandscapeMaterial();
-			}
-			if (!HoleMaterial->GetMaterial()->HasAnyExpressionsInMaterialAndFunctionsOfType<UMaterialExpressionLandscapeVisibilityMask>())
+			ULandscapeInfo* LandscapeInfo = LandscapeEdMode->CurrentToolTarget.LandscapeInfo.Get();
+			bool bHasValidHoleMaterial = true;
+			LandscapeInfo->ForAllLandscapeComponents([&](ULandscapeComponent* LandscapeComponent)
 			{
-				return true;
-			}
+				bHasValidHoleMaterial &= LandscapeComponent->IsLandscapeHoleMaterialValid();
+			});
+
+			return !bHasValidHoleMaterial;
 		}
 	}
 
diff --git a/Engine/Source/Runtime/Landscape/Classes/LandscapeComponent.h b/Engine/Source/Runtime/Landscape/Classes/LandscapeComponent.h
--- a/Engine/Source/Runtime/Landscape/Classes/LandscapeComponent.h
+++ b/Engine/Source/Runtime/Landscape/Classes/LandscapeComponent.h
@@ -897,12 +897,15 @@ public:
 	 */
 	LANDSCAPE_API void ReallocateWeightmaps(FLandscapeEditDataInterface* DataInterface = nullptr, bool InCanUseEditingWeightmap = true, bool InSaveToTransactionBuffer = true, bool InInitPlatformDataAsync = false, bool InForceReallocate = false, ALandscapeProxy* InTargetProxy = nullptr, TArray<UTexture2D*>* OutNewCreatedTextures = nullptr);
 
-	/** Returns the actor's LandscapeMaterial, or the Component's OverrideLandscapeMaterial if set */
+	/** Returns the component's LandscapeMaterial, or the Component's OverrideLandscapeMaterial if set */
 	LANDSCAPE_API UMaterialInterface* GetLandscapeMaterial(int8 InLODIndex = INDEX_NONE) const;
 
-	/** Returns the actor's LandscapeHoleMaterial, or the Component's OverrideLandscapeHoleMaterial if set */
+	/** Returns the components's LandscapeHoleMaterial, or the Component's OverrideLandscapeHoleMaterial if set */
 	LANDSCAPE_API UMaterialInterface* GetLandscapeHoleMaterial() const;
 
+	/** Returns true if the component has a valid LandscapeHoleMaterial */
+	LANDSCAPE_API bool IsLandscapeHoleMaterialValid() const;
+
 	/** Returns true if this component has visibility painted */
 	LANDSCAPE_API bool ComponentHasVisibilityPainted() const;
 
diff --git a/Engine/Source/Runtime/Landscape/Private/Landscape.cpp b/Engine/Source/Runtime/Landscape/Private/Landscape.cpp
--- a/Engine/Source/Runtime/Landscape/Private/Landscape.cpp
+++ b/Engine/Source/Runtime/Landscape/Private/Landscape.cpp
@@ -593,6 +593,17 @@ UMaterialInterface* ULandscapeComponent::GetLandscapeHoleMaterial() const
 	return nullptr;
 }
 
+bool ULandscapeComponent::IsLandscapeHoleMaterialValid() const
+{
+	UMaterialInterface* HoleMaterial = GetLandscapeHoleMaterial();
+	if (!HoleMaterial)
+	{
+		return false;
+	}
+
+	return HoleMaterial->GetMaterial()->HasAnyExpressionsInMaterialAndFunctionsOfType<UMaterialExpressionLandscapeVisibilityMask>();
+}
+
 bool ULandscapeComponent::ComponentHasVisibilityPainted() const
 {
 	for (const FWeightmapLayerAllocationInfo& Allocation : WeightmapLayerAllocations)
diff --git a/Engine/Source/Runtime/Landscape/Private/LandscapeEditInterface.cpp b/Engine/Source/Runtime/Landscape/Private/LandscapeEditInterface.cpp
--- a/Engine/Source/Runtime/Landscape/Private/LandscapeEditInterface.cpp
+++ b/Engine/Source/Runtime/Landscape/Private/LandscapeEditInterface.cpp
@@ -2785,6 +2785,11 @@ void FLandscapeEditDataInterface::SetAlphaData(ULandscapeLayerInfoObject* const
 				continue;
 			}
 
+			if (LayerInfo == ALandscape::VisibilityLayer && !Component->IsLandscapeHoleMaterialValid())
+			{
+				continue;
+			}
+
 			if (PaintingRestriction == ELandscapeLayerPaintingRestriction::UseComponentWhitelist && !Component->LayerWhitelist.Contains(LayerInfo))
 			{
 				continue;
