Add Owner Textbox of Selected Spline Control Points/Segments in Details Panel

diff --git a/Engine/Source/Editor/LandscapeEditor/Private/LandscapeEdMode.h b/Engine/Source/Editor/LandscapeEditor/Private/LandscapeEdMode.h
--- a/Engine/Source/Editor/LandscapeEditor/Private/LandscapeEdMode.h
+++ b/Engine/Source/Editor/LandscapeEditor/Private/LandscapeEdMode.h
@@ -13,6 +13,7 @@
 #include "LandscapeLayerInfoObject.h"
 #include "LandscapeGizmoActiveActor.h"
 #include "LandscapeEdit.h"
+#include "Containers/Set.h"
 
 class ALandscape;
 class FCanvas;
@@ -319,6 +320,7 @@ public:
 	void ShowSplineProperties();
 	bool HasSelectedSplineSegments() const;
 	void FlipSelectedSplineSegments();
+	void GetSelectedSplineOwners(TSet<ALandscapeProxy*>& SelectedSplineOwners) const;
 	virtual void SelectAllConnectedSplineControlPoints();
 	virtual void SelectAllConnectedSplineSegments();
 	virtual void SplineMoveToCurrentLevel();
diff --git a/Engine/Source/Editor/LandscapeEditor/Private/LandscapeEdModeSplineTools.cpp b/Engine/Source/Editor/LandscapeEditor/Private/LandscapeEdModeSplineTools.cpp
--- a/Engine/Source/Editor/LandscapeEditor/Private/LandscapeEdModeSplineTools.cpp
+++ b/Engine/Source/Editor/LandscapeEditor/Private/LandscapeEdModeSplineTools.cpp
@@ -2250,6 +2250,20 @@ void FEdModeLandscape::ShowSplineProperties()
 	}
 }
 
+void FEdModeLandscape::GetSelectedSplineOwners(TSet<ALandscapeProxy*>& SelectedSplineOwners) const
+{
+	for (ULandscapeSplineSegment* Segment : SplinesTool->SelectedSplineSegments)
+	{
+		SelectedSplineOwners.Add(Segment->GetTypedOuter<ALandscapeProxy>());
+	}
+
+	for (ULandscapeSplineControlPoint* ControlPoint : SplinesTool->SelectedSplineControlPoints)
+	{
+		SelectedSplineOwners.Add(ControlPoint->GetTypedOuter<ALandscapeProxy>());
+	}
+}
+
+
 void FEdModeLandscape::SelectAllConnectedSplineControlPoints()
 {
 	if (SplinesTool /*&& SplinesTool == CurrentTool*/)
diff --git a/Engine/Source/Editor/LandscapeEditor/Private/LandscapeSplineDetails.cpp b/Engine/Source/Editor/LandscapeEditor/Private/LandscapeSplineDetails.cpp
--- a/Engine/Source/Editor/LandscapeEditor/Private/LandscapeSplineDetails.cpp
+++ b/Engine/Source/Editor/LandscapeEditor/Private/LandscapeSplineDetails.cpp
@@ -24,6 +24,19 @@ void FLandscapeSplineDetails::CustomizeDetails(IDetailLayoutBuilder& DetailBuild
 {
 	IDetailCategoryBuilder& LandscapeSplineCategory = DetailBuilder.EditCategory("LandscapeSpline", FText::GetEmpty(), ECategoryPriority::Transform);
 
+	LandscapeSplineCategory.AddCustomRow(FText::GetEmpty())
+	[
+		SNew(SHorizontalBox)
+		+ SHorizontalBox::Slot()
+		.Padding(0, 0, 2, 0)
+		.VAlign(VAlign_Center)
+		.FillWidth(1)
+		[
+			SNew(STextBlock)
+			.Text_Raw(this, &FLandscapeSplineDetails::OnGetSplineOwningLandscapeText)
+		]
+	];
+
 	LandscapeSplineCategory.AddCustomRow(FText::GetEmpty())
 	[
 		SNew(SHorizontalBox)
@@ -123,6 +136,32 @@ bool FLandscapeSplineDetails::IsFlipSegmentButtonEnabled() const
 	return LandscapeEdMode && LandscapeEdMode->HasSelectedSplineSegments();
 }
 
+FText FLandscapeSplineDetails::OnGetSplineOwningLandscapeText() const
+{
+	TSet<ALandscapeProxy*> SplineOwners;
+	FEdModeLandscape* LandscapeEdMode = GetEditorMode();
+	if (LandscapeEdMode && LandscapeEdMode->CurrentToolTarget.LandscapeInfo.IsValid())
+	{
+		LandscapeEdMode->GetSelectedSplineOwners(SplineOwners);
+	}
+
+	FString SplineOwnersStr;
+	for (ALandscapeProxy* Owner : SplineOwners)
+	{
+		if (Owner)
+		{
+			if (!SplineOwnersStr.IsEmpty())
+			{
+				SplineOwnersStr += ", ";
+			}
+			
+			SplineOwnersStr += Owner->GetActorLabel();
+		}
+	}
+	
+	return FText::FromString("Owner: " + SplineOwnersStr);
+}
+
 FReply FLandscapeSplineDetails::OnSelectConnectedControlPointsButtonClicked()
 {
 	FEdModeLandscape* LandscapeEdMode = GetEditorMode();
diff --git a/Engine/Source/Editor/LandscapeEditor/Private/LandscapeSplineDetails.h b/Engine/Source/Editor/LandscapeEditor/Private/LandscapeSplineDetails.h
--- a/Engine/Source/Editor/LandscapeEditor/Private/LandscapeSplineDetails.h
+++ b/Engine/Source/Editor/LandscapeEditor/Private/LandscapeSplineDetails.h
@@ -20,6 +20,7 @@ public:
 
 private:
 	FEdModeLandscape* GetEditorMode() const;
+	FText OnGetSplineOwningLandscapeText() const;
 	FReply OnSelectConnectedControlPointsButtonClicked();
 	FReply OnSelectConnectedSegmentsButtonClicked();
 	FReply OnMoveToCurrentLevelButtonClicked();
