Reimport tiled Landscape Layer Weightmaps now has same behavior as Import (no more forcing a normalization of layer weightmaps).

diff --git a/Engine/Source/Editor/LandscapeEditor/Private/LandscapeEdMode.cpp b/Engine/Source/Editor/LandscapeEditor/Private/LandscapeEdMode.cpp
--- a/Engine/Source/Editor/LandscapeEditor/Private/LandscapeEdMode.cpp
+++ b/Engine/Source/Editor/LandscapeEditor/Private/LandscapeEdMode.cpp
@@ -4703,7 +4703,7 @@ bool LandscapeEditorUtils::SetWeightmapData(ALandscapeProxy* Landscape, ULandsca
 
 	if (Data.Num() == (1 + ComponentsRect.Width())*(1 + ComponentsRect.Height()))
 	{
-		FAlphamapAccessor<false, true> AlphamapAccessor(Landscape->GetLandscapeInfo(), LayerObject);
+		FAlphamapAccessor<false, false> AlphamapAccessor(Landscape->GetLandscapeInfo(), LayerObject);
 		AlphamapAccessor.SetData(ComponentsRect.Min.X, ComponentsRect.Min.Y, ComponentsRect.Max.X, ComponentsRect.Max.Y, Data.GetData(), ELandscapeLayerPaintingRestriction::None);
 		return true;
 	}
