Fixed crash when entering Landscape Mode with Mirror Tool enabled and mirror point set to 0,0 but without any landscape

diff --git a/Engine/Source/Editor/LandscapeEditor/Private/LandscapeEdModeMirrorTool.cpp b/Engine/Source/Editor/LandscapeEditor/Private/LandscapeEdModeMirrorTool.cpp
--- a/Engine/Source/Editor/LandscapeEditor/Private/LandscapeEdModeMirrorTool.cpp
+++ b/Engine/Source/Editor/LandscapeEditor/Private/LandscapeEdModeMirrorTool.cpp
@@ -670,9 +670,8 @@ public:
 	void CenterMirrorPoint()
 	{
 		ULandscapeInfo* const LandscapeInfo = EdMode->CurrentToolTarget.LandscapeInfo.Get();
-		ALandscapeProxy* const LandscapeProxy = LandscapeInfo->GetLandscapeProxy();
 		int32 MinX, MinY, MaxX, MaxY;
-		if (LandscapeInfo->GetLandscapeExtent(MinX, MinY, MaxX, MaxY))
+		if (LandscapeInfo && LandscapeInfo->GetLandscapeExtent(MinX, MinY, MaxX, MaxY))
 		{
 			EdMode->UISettings->MirrorPoint = FVector2D((float)(MinX + MaxX) / 2.0f, (float)(MinY + MaxY) / 2.0f);
 		}
