Add UIMin and UIMax markup to Landscape brush falloff spinner
Fixes issue with new modes toolbar widget failing to get correct Min/Max range
Cherry pick integration from Dev-Editor

diff --git a/Engine/Source/Editor/LandscapeEditor/Public/LandscapeEditorObject.h b/Engine/Source/Editor/LandscapeEditor/Public/LandscapeEditorObject.h
--- a/Engine/Source/Editor/LandscapeEditor/Public/LandscapeEditorObject.h
+++ b/Engine/Source/Editor/LandscapeEditor/Public/LandscapeEditorObject.h
@@ -544,7 +544,7 @@ public:
 	float BrushRadius;
 
 	// The falloff at the edge of the brush, as a fraction of the brush's size. 0 = no falloff, 1 = all falloff
-	UPROPERTY(Category="Brush Settings", EditAnywhere, NonTransactional, meta=(ShowForBrushes="BrushSet_Circle,BrushSet_Gizmo,BrushSet_Pattern", ClampMin="0", ClampMax="1"))
+	UPROPERTY(Category="Brush Settings", EditAnywhere, NonTransactional, meta=(ShowForBrushes="BrushSet_Circle,BrushSet_Gizmo,BrushSet_Pattern", ClampMin="0", ClampMax="1", UIMin = "0", UIMax = "1"))
 	float BrushFalloff;
 
 	// Selects the Clay Brush painting mode
