Support use of greyscale textures for landscape image brushes
Cherry pick integration from Dev-Editor

diff --git a/Engine/Source/Editor/LandscapeEditor/Private/LandscapeEditorObject.cpp b/Engine/Source/Editor/LandscapeEditor/Private/LandscapeEditorObject.cpp
--- a/Engine/Source/Editor/LandscapeEditor/Private/LandscapeEditorObject.cpp
+++ b/Engine/Source/Editor/LandscapeEditor/Private/LandscapeEditorObject.cpp
@@ -478,8 +478,11 @@ bool ULandscapeEditorObject::SetAlphaTexture(UTexture2D* InTexture, EColorChanne
 		NewAlphaTexture->Source.GetMipData(NewTextureData, 0);
 	}
 
+	const bool bSourceDataIsG8 = NewAlphaTexture != NULL && NewAlphaTexture->Source.IsValid() && NewAlphaTexture->Source.GetFormat() == TSF_G8;
+	const int32 NumChannels = bSourceDataIsG8 ? 1 : 4;
+
 	// Load fallback if there's no texture or data
-	if (NewAlphaTexture == NULL || (NewTextureData.Num() != 4 * NewAlphaTexture->Source.GetSizeX() * NewAlphaTexture->Source.GetSizeY()))
+	if (NewAlphaTexture == NULL || (NewTextureData.Num() != NumChannels * NewAlphaTexture->Source.GetSizeX() * NewAlphaTexture->Source.GetSizeY()))
 	{
 		NewAlphaTexture = GetClass()->GetDefaultObject<ULandscapeEditorObject>()->AlphaTexture;
 		NewAlphaTexture->Source.GetMipData(NewTextureData, 0);
@@ -490,10 +493,10 @@ bool ULandscapeEditorObject::SetAlphaTexture(UTexture2D* InTexture, EColorChanne
 	AlphaTexture = NewAlphaTexture;
 	AlphaTextureSizeX = NewAlphaTexture->Source.GetSizeX();
 	AlphaTextureSizeY = NewAlphaTexture->Source.GetSizeY();
-	AlphaTextureChannel = InTextureChannel;
+	AlphaTextureChannel = NumChannels == 1 ? EColorChannel::Red : InTextureChannel;
 	AlphaTextureData.Empty(AlphaTextureSizeX*AlphaTextureSizeY);
 
-	if (NewTextureData.Num() != 4 *AlphaTextureSizeX*AlphaTextureSizeY)
+	if (NewTextureData.Num() != NumChannels * AlphaTextureSizeX * AlphaTextureSizeY)
 	{
 		// Don't crash if for some reason we couldn't load any source art
 		AlphaTextureData.AddZeroed(AlphaTextureSizeX*AlphaTextureSizeY);
@@ -520,7 +523,7 @@ bool ULandscapeEditorObject::SetAlphaTexture(UTexture2D* InTexture, EColorChanne
 		for (int32 i=0;i<AlphaTextureSizeX*AlphaTextureSizeY;i++)
 		{
 			AlphaTextureData.Add(*SrcPtr);
-			SrcPtr += 4;
+			SrcPtr += NumChannels;
 		}
 	}
 
