Add LandscapeToolsLog to debug Shift+Paint
Add LogLandscapeTools=VeryVerbose to DefaultEngine.ini to enable

#rnx

diff --git a/Engine/Source/Editor/LandscapeEditor/Private/LandscapeEdModePaintTools.cpp b/Engine/Source/Editor/LandscapeEditor/Private/LandscapeEdModePaintTools.cpp
--- a/Engine/Source/Editor/LandscapeEditor/Private/LandscapeEdModePaintTools.cpp
+++ b/Engine/Source/Editor/LandscapeEditor/Private/LandscapeEdModePaintTools.cpp
@@ -16,11 +16,14 @@
 #include "Landscape.h"
 #include "Logging/TokenizedMessage.h"
 #include "Logging/MessageLog.h"
+#include "Logging/LogMacros.h"
 #include "Misc/MapErrors.h"
 #include "EngineModule.h"
 
 #define LOCTEXT_NAMESPACE "LandscapeTools"
 
+DEFINE_LOG_CATEGORY(LogLandscapeTools);
+
 const int32 FNoiseParameter::Permutations[256] =
 {
 	151, 160, 137, 91, 90, 15,
@@ -103,8 +106,8 @@ public:
 	void Apply(FEditorViewportClient* ViewportClient, FLandscapeBrush* Brush, const ULandscapeEditorObject* UISettings, const TArray<FLandscapeToolInteractorPosition>& InteractorPositions)
 	{
 		// Invert when holding Shift
-		//UE_LOG(LogLandscape, Log, TEXT("bInvert = %d"), bInvert);
 		bool bInvert = InteractorPositions.Last().bModifierPressed;
+		UE_LOG(LogLandscapeTools, VeryVerbose, TEXT("bInvert = %d"), bInvert);
 
 		if (bIsWhitelistMode)
 		{
@@ -389,8 +392,8 @@ public:
 	void Apply(FEditorViewportClient* ViewportClient, FLandscapeBrush* Brush, const ULandscapeEditorObject* UISettings, const TArray<FLandscapeToolInteractorPosition>& InteractorPositions)
 	{
 		// Invert when holding Shift
-		//UE_LOG(LogLandscape, Log, TEXT("bInvert = %d"), bInvert);
 		bool bInvert = InteractorPositions.Last().bModifierPressed;
+		UE_LOG(LogLandscapeTools, VeryVerbose, TEXT("bInvert = %d"), bInvert);
 
 		// Get list of verts to update
 		FLandscapeBrushData BrushInfo = Brush->ApplyBrush(InteractorPositions);
diff --git a/Engine/Source/Editor/LandscapeEditor/Private/LandscapeEdModeTools.h b/Engine/Source/Editor/LandscapeEditor/Private/LandscapeEdModeTools.h
--- a/Engine/Source/Editor/LandscapeEditor/Private/LandscapeEdModeTools.h
+++ b/Engine/Source/Editor/LandscapeEditor/Private/LandscapeEdModeTools.h
@@ -16,7 +16,9 @@
 #include "VREditorInteractor.h"
 #include "AI/NavigationSystemBase.h"
 #include "Landscape.h"
+#include "Logging/LogMacros.h"
 
+DECLARE_LOG_CATEGORY_EXTERN(LogLandscapeTools, Log, All);
 // VR Editor
 
 //
@@ -1531,6 +1533,7 @@ protected:
 
 	bool IsModifierPressed(const class FEditorViewportClient* ViewportClient = nullptr)
 	{
+		UE_LOG(LogLandscapeTools, VeryVerbose, TEXT("bExternalModifierPressed = %d, ViewportClient = %d, IsShiftDown = %d"), bExternalModifierPressed, (ViewportClient != nullptr), (ViewportClient != nullptr && IsShiftDown(ViewportClient->Viewport)));
 		return bExternalModifierPressed || (ViewportClient != nullptr && IsShiftDown(ViewportClient->Viewport));
 	}
 };
