Fixed crash when importing Landscape layers in the New Landscape Tool

diff --git a/Engine/Source/Editor/LandscapeEditor/Private/LandscapeEditorDetailCustomization_NewLandscape.cpp b/Engine/Source/Editor/LandscapeEditor/Private/LandscapeEditorDetailCustomization_NewLandscape.cpp
--- a/Engine/Source/Editor/LandscapeEditor/Private/LandscapeEditorDetailCustomization_NewLandscape.cpp
+++ b/Engine/Source/Editor/LandscapeEditor/Private/LandscapeEditorDetailCustomization_NewLandscape.cpp
@@ -814,8 +814,9 @@ FReply FLandscapeEditorDetailCustomization_NewLandscape::OnCreateButtonClicked()
 		TMap<FGuid, TArray<uint16>> HeightDataPerLayers;
 		TMap<FGuid, TArray<FLandscapeImportLayerInfo>> MaterialLayerDataPerLayers;
 
-		MaterialLayerDataPerLayers.Add(FGuid(), MaterialImportLayers.GetValue());
 		HeightDataPerLayers.Add(FGuid(), FNewLandscapeUtils::ComputeHeightData(UISettings, MaterialImportLayers.GetValue(), LandscapeEdMode->NewLandscapePreviewMode));
+		// ComputeHeightData will also modify/expand material layers data, which is why we create MaterialLayerDataPerLayers after calling ComputeHeightData
+		MaterialLayerDataPerLayers.Add(FGuid(), MaterialImportLayers.GetValue());
 
 		FScopedTransaction Transaction(LOCTEXT("Undo", "Creating New Landscape"));
 
