Fixed assert with splines brush tool

diff --git a/Engine/Source/Editor/LandscapeEditor/Private/LandscapeEditorDetails.cpp b/Engine/Source/Editor/LandscapeEditor/Private/LandscapeEditorDetails.cpp
--- a/Engine/Source/Editor/LandscapeEditor/Private/LandscapeEditorDetails.cpp
+++ b/Engine/Source/Editor/LandscapeEditor/Private/LandscapeEditorDetails.cpp
@@ -340,6 +340,7 @@ FText FLandscapeEditorDetails::GetLocalizedName(FString Name)
 		LOCTEXT("BrushSet_Component", "Component");
 		LOCTEXT("BrushSet_Gizmo", "Gizmo");
 		LOCTEXT("BrushSet_Dummy", "NoBrush");
+		LOCTEXT("BrushSet_Splines", "Splines");
 
 		LOCTEXT("Circle_Smooth", "Smooth");
 		LOCTEXT("Circle_Linear", "Linear");
