Fix Ensure when going in and out of landscape (and in manage tab)

diff --git a/Engine/Source/Editor/LandscapeEditor/Private/LandscapeEditorDetails.cpp b/Engine/Source/Editor/LandscapeEditor/Private/LandscapeEditorDetails.cpp
--- a/Engine/Source/Editor/LandscapeEditor/Private/LandscapeEditorDetails.cpp
+++ b/Engine/Source/Editor/LandscapeEditor/Private/LandscapeEditorDetails.cpp
@@ -640,7 +640,7 @@ bool FLandscapeEditorDetails::GetBrushSelectorIsVisible() const
 FText FLandscapeEditorDetails::GetCurrentBrushFalloffName() const
 {
 	FEdModeLandscape* LandscapeEdMode = GetEditorMode();
-	if (LandscapeEdMode != NULL && LandscapeEdMode->CurrentBrush != NULL)
+	if (LandscapeEdMode != NULL && LandscapeEdMode->CurrentBrush != NULL && GetBrushFalloffSelectorIsVisible())
 	{
 		const TCHAR* CurrentBrushName = LandscapeEdMode->CurrentBrush->GetBrushName();
 		return GetLocalizedName(CurrentBrushName);
