CIS fixes

diff --git a/Engine/Source/Editor/LandscapeEditor/Private/LandscapeEditorDetailCustomization_Layers.h b/Engine/Source/Editor/LandscapeEditor/Private/LandscapeEditorDetailCustomization_Layers.h
--- a/Engine/Source/Editor/LandscapeEditor/Private/LandscapeEditorDetailCustomization_Layers.h
+++ b/Engine/Source/Editor/LandscapeEditor/Private/LandscapeEditorDetailCustomization_Layers.h
@@ -25,6 +25,7 @@ class IDetailLayoutBuilder;
 class SDragAndDropVerticalBox;
 class ALandscapeBlueprintBrushBase;
 class ULandscapeLayerInfoObject;
+class FMenuBuilder;
 
 /**
  * Slate widgets customizer for the layers list in the Landscape Editor
