Fix missing forward declares
#rnx

diff --git a/Engine/Plugins/Experimental/GPULightmass/Source/GPULightmass/Private/GPULightmass.h b/Engine/Plugins/Experimental/GPULightmass/Source/GPULightmass/Private/GPULightmass.h
--- a/Engine/Plugins/Experimental/GPULightmass/Source/GPULightmass/Private/GPULightmass.h
+++ b/Engine/Plugins/Experimental/GPULightmass/Source/GPULightmass/Private/GPULightmass.h
@@ -7,6 +7,8 @@
 #include "GPULightmassModule.h"
 #include "Scene/Scene.h"
 
+class SNotificationItem;
+
 class FGPULightmass : public IStaticLightingSystem
 {
 public:
diff --git a/Engine/Source/Editor/LandscapeEditor/Private/LandscapeEditorDetailCustomization_Layers.h b/Engine/Source/Editor/LandscapeEditor/Private/LandscapeEditorDetailCustomization_Layers.h
--- a/Engine/Source/Editor/LandscapeEditor/Private/LandscapeEditorDetailCustomization_Layers.h
+++ b/Engine/Source/Editor/LandscapeEditor/Private/LandscapeEditorDetailCustomization_Layers.h
@@ -23,6 +23,7 @@ class FDetailWidgetRow;
 class IDetailChildrenBuilder;
 class IDetailLayoutBuilder;
 class SDragAndDropVerticalBox;
+class SInlineEditableTextBlock;
 class ALandscapeBlueprintBrushBase;
 class ULandscapeLayerInfoObject;
 class FMenuBuilder;
