Fix Crash when opening Landscape Tool

#rnx

diff --git a/Engine/Source/Editor/LandscapeEditor/Private/LandscapeEdMode.cpp b/Engine/Source/Editor/LandscapeEditor/Private/LandscapeEdMode.cpp
--- a/Engine/Source/Editor/LandscapeEditor/Private/LandscapeEdMode.cpp
+++ b/Engine/Source/Editor/LandscapeEditor/Private/LandscapeEdMode.cpp
@@ -556,19 +556,14 @@ void FEdModeLandscape::Enter()
 	UpdateTargetList();
 	UpdateShownLayerList();
 
-	// Create the landscape editor window
-	if (!Toolkit.IsValid())
-	{
-		Toolkit = MakeShareable(new FLandscapeToolKit);
-		Toolkit->Init(Owner->GetToolkitHost());
-	}
+	FName ToolkitPalette = NAME_None;
 
 	// Initialize current tool prior to creating the landscape toolkit in case it has a dependency on it
 	if (LandscapeList.Num() == 0)
 	{
 		SetCurrentToolMode("ToolMode_Manage", false);
 		SetCurrentTool("NewLandscape");
-		Toolkit->SetCurrentPalette("ToolMode_Manage");
+		ToolkitPalette = "ToolMode_Manage";
 	}
 	else
 	{
@@ -576,15 +571,23 @@ void FEdModeLandscape::Enter()
 		{
 			SetCurrentToolMode("ToolMode_Sculpt", false);
 			SetCurrentTool("Sculpt");
-			Toolkit->SetCurrentPalette("ToolMode_Sculpt");
+			ToolkitPalette = "ToolMode_Sculpt";
 		}
 		else
 		{
 			SetCurrentTool(CurrentToolMode->CurrentToolName);
-			Toolkit->SetCurrentPalette(CurrentToolMode->ToolModeName);
+			ToolkitPalette = CurrentToolMode->ToolModeName;
 		}
 	}
 
+	// Create the landscape editor window
+	if (!Toolkit.IsValid())
+	{
+		Toolkit = MakeShareable(new FLandscapeToolKit);
+		Toolkit->Init(Owner->GetToolkitHost());
+		Toolkit->SetCurrentPalette(ToolkitPalette);
+	}
+
 	// Force real-time viewports.  We'll back up the current viewport state so we can restore it when the
 	// user exits this mode.
 	const bool bWantRealTime = true;
