Fix missing tool tips in landscape brushes

diff --git a/Engine/Source/Editor/LandscapeEditor/Private/LandscapeEditorDetails.cpp b/Engine/Source/Editor/LandscapeEditor/Private/LandscapeEditorDetails.cpp
--- a/Engine/Source/Editor/LandscapeEditor/Private/LandscapeEditorDetails.cpp
+++ b/Engine/Source/Editor/LandscapeEditor/Private/LandscapeEditorDetails.cpp
@@ -95,17 +95,19 @@ void FLandscapeEditorDetails::CustomizeDetails(IDetailLayoutBuilder& DetailBuild
 		if (Command.IsValid())
 		{
 			BrushSelectors->AddSlot(BrushName)
-			.Icon(Command->GetIcon().GetIcon());
+			.Icon(Command->GetIcon().GetIcon())
+			.ToolTip(Command->GetDescription());
 		}
 	}
+	BrushSelectors->RebuildChildren();
 
 	TSharedPtr<SSegmentedControl<FName>> FalloffSelectors = SNew(SSegmentedControl<FName>)
-	.OnValueChanged(this, &FLandscapeEditorDetails::SetBrushCommand)
-	.Value(this, &FLandscapeEditorDetails::GetCurrentBrushFalloffFName)
-	+ SSegmentedControl<FName>::Slot(FName("Circle_Smooth")).Icon(FLandscapeEditorCommands::Get().CircleBrush_Smooth->GetIcon().GetIcon())
-	+ SSegmentedControl<FName>::Slot(FName("Circle_Linear")).Icon(FLandscapeEditorCommands::Get().CircleBrush_Linear->GetIcon().GetIcon())
-	+ SSegmentedControl<FName>::Slot(FName("Circle_Spherical")).Icon(FLandscapeEditorCommands::Get().CircleBrush_Spherical->GetIcon().GetIcon())
-	+ SSegmentedControl<FName>::Slot(FName("Circle_Tip")).Icon(FLandscapeEditorCommands::Get().CircleBrush_Tip->GetIcon().GetIcon());
+		.OnValueChanged(this, &FLandscapeEditorDetails::SetBrushCommand)
+		.Value(this, &FLandscapeEditorDetails::GetCurrentBrushFalloffFName)
+		+ SSegmentedControl<FName>::Slot(FName("Circle_Smooth")).Icon(FLandscapeEditorCommands::Get().CircleBrush_Smooth->GetIcon().GetIcon()).ToolTip(FLandscapeEditorCommands::Get().CircleBrush_Smooth->GetDescription())
+		+ SSegmentedControl<FName>::Slot(FName("Circle_Linear")).Icon(FLandscapeEditorCommands::Get().CircleBrush_Linear->GetIcon().GetIcon()).ToolTip(FLandscapeEditorCommands::Get().CircleBrush_Linear->GetDescription())
+		+ SSegmentedControl<FName>::Slot(FName("Circle_Spherical")).Icon(FLandscapeEditorCommands::Get().CircleBrush_Spherical->GetIcon().GetIcon()).ToolTip(FLandscapeEditorCommands::Get().CircleBrush_Spherical->GetDescription())
+		+ SSegmentedControl<FName>::Slot(FName("Circle_Tip")).Icon(FLandscapeEditorCommands::Get().CircleBrush_Tip->GetIcon().GetIcon()).ToolTip(FLandscapeEditorCommands::Get().CircleBrush_Tip->GetDescription());
 
 	LandscapeEditorCategory.AddCustomRow(FText::GetEmpty())
 	.Visibility(TAttribute<EVisibility>::Create(TAttribute<EVisibility>::FGetter::CreateSP(this, &FLandscapeEditorDetails::GetBrushSelectorVisibility)))
