Move LandscapeInfo management into LandscapeSubsystem. (Remove usage of World->PerModuleDataObjects)

diff --git a/Engine/Source/Editor/LandscapeEditor/Private/LandscapeEdMode.cpp b/Engine/Source/Editor/LandscapeEditor/Private/LandscapeEdMode.cpp
--- a/Engine/Source/Editor/LandscapeEditor/Private/LandscapeEdMode.cpp
+++ b/Engine/Source/Editor/LandscapeEditor/Private/LandscapeEdMode.cpp
@@ -29,7 +29,7 @@
 #include "InstancedFoliageActor.h"
 #include "EditorWorldExtension.h"
 #include "LandscapeEdModeTools.h"
-#include "LandscapeInfoMap.h"
+#include "LandscapeSubsystem.h"
 
 //Slate dependencies
 #include "Misc/FeedbackContext.h"
@@ -400,16 +400,14 @@ void FEdModeLandscape::Enter()
 
 	if (UWorld* World = GetWorld())
 	{
-		for (auto It = ULandscapeInfoMap::GetLandscapeInfoMap(World).Map.CreateIterator(); It; ++It)
+		ULandscapeSubsystem::ForEachLandscapeInfo(World, [this](ULandscapeInfo* LandscapeInfo)
 		{
-			if (ULandscapeInfo* LandscapeInfo = It.Value())
+			if (ALandscape* Landscape = LandscapeInfo->LandscapeActor.Get())
 			{
-				if (ALandscape* Landscape = !LandscapeInfo->IsPendingKill() ? LandscapeInfo->LandscapeActor.Get() : nullptr)
-				{
-					Landscape->RegisterLandscapeEdMode(this);
-				}
+				Landscape->RegisterLandscapeEdMode(this);
 			}
-		}
+			return true;
+		});
 	}
 
 	OnLevelActorDeletedDelegateHandle = GEngine->OnLevelActorDeleted().AddSP(this, &FEdModeLandscape::OnLevelActorRemoved);
@@ -606,16 +604,14 @@ void FEdModeLandscape::Exit()
 {
 	if (UWorld* World = GetWorld())
 	{
-		for (auto It = ULandscapeInfoMap::GetLandscapeInfoMap(World).Map.CreateIterator(); It; ++It)
+		ULandscapeSubsystem::ForEachLandscapeInfo(World, [](ULandscapeInfo* LandscapeInfo)
 		{
-			if (ULandscapeInfo* LandscapeInfo = It.Value())
+			if (ALandscape* Landscape = LandscapeInfo->LandscapeActor.Get())
 			{
-				if (ALandscape* Landscape = !LandscapeInfo->IsPendingKill() ? LandscapeInfo->LandscapeActor.Get() : nullptr)
-				{
-					Landscape->UnregisterLandscapeEdMode();
-				}
+				Landscape->UnregisterLandscapeEdMode();
 			}
-		}
+			return true;
+		});
 	}
 
 	GEngine->OnLevelActorDeleted().Remove(OnLevelActorDeletedDelegateHandle);
@@ -701,17 +697,24 @@ void FEdModeLandscape::OnPreSaveWorld(uint32 InSaveFlags, const class UWorld* In
 	// Avoid doing this during cooking to keep determinism and we don't want to do this on GameWorlds.
 	if (!InWorld->IsGameWorld() && !GIsCookerLoadingPackage)
 	{
-		ULandscapeInfoMap& LandscapeInfoMap = ULandscapeInfoMap::GetLandscapeInfoMap(InWorld);
-		for (const TPair<FGuid, ULandscapeInfo*>& Pair : LandscapeInfoMap.Map)
+		// The world we are getting called for is the outer of the level being saved. LandscapeSubsystem exists only on the Owning world
+		const UWorld* OwningWorld = InWorld;
+		if (InWorld->PersistentLevel && InWorld->PersistentLevel->OwningWorld != nullptr)
+		{
+			OwningWorld = InWorld->PersistentLevel->OwningWorld;
+		}
+
+		ULandscapeSubsystem::ForEachLandscapeInfo(OwningWorld, [InWorld](ULandscapeInfo* LandscapeInfo)
 		{
-			if (const ULandscapeInfo* LandscapeInfo = Pair.Value)
+			if (ALandscape* LandscapeActor = LandscapeInfo->LandscapeActor.Get())
 			{
-				if (ALandscape* LandscapeActor = LandscapeInfo->LandscapeActor.Get())
+				if (LandscapeActor->GetTypedOuter<UWorld>() == InWorld)
 				{
 					LandscapeActor->OnPreSave();
 				}
 			}
-		}
+			return true;
+		});
 	}
 }
 
@@ -2223,46 +2226,41 @@ int32 FEdModeLandscape::UpdateLandscapeList()
 	if (World)
 	{
 		int32 Index = 0;
-		auto& LandscapeInfoMap = ULandscapeInfoMap::GetLandscapeInfoMap(World);
-
-		for (auto It = LandscapeInfoMap.Map.CreateIterator(); It; ++It)
+		ULandscapeSubsystem::ForEachLandscapeInfo(World, [&CurrentIndex, &Index, this](ULandscapeInfo* LandscapeInfo)
 		{
-			ULandscapeInfo* LandscapeInfo = It.Value();
-			if (LandscapeInfo && !LandscapeInfo->IsPendingKill())
+			if (ALandscape* Landscape = LandscapeInfo->LandscapeActor.Get())
 			{
-				if (ALandscape* Landscape = LandscapeInfo->LandscapeActor.Get())
-				{
-					Landscape->RegisterLandscapeEdMode(this);
-				}
+				Landscape->RegisterLandscapeEdMode(this);
+			}
 
-				ALandscapeProxy* LandscapeProxy = LandscapeInfo->GetLandscapeProxy();
-				if (LandscapeProxy)
+			ALandscapeProxy* LandscapeProxy = LandscapeInfo->GetLandscapeProxy();
+			if (LandscapeProxy)
+			{
+				if (CurrentToolTarget.LandscapeInfo == LandscapeInfo)
 				{
-					if (CurrentToolTarget.LandscapeInfo == LandscapeInfo)
-					{
-						CurrentIndex = Index;
-
-						// Update GizmoActor Landscape Target (is this necessary?)
-						if (CurrentGizmoActor.IsValid())
-						{
-							CurrentGizmoActor->SetTargetLandscape(LandscapeInfo);
-						}
-					}
+					CurrentIndex = Index;
 
-					int32 MinX, MinY, MaxX, MaxY;
-					int32 Width = 0, Height = 0;
-					if (LandscapeInfo->GetLandscapeExtent(MinX, MinY, MaxX, MaxY))
+					// Update GizmoActor Landscape Target (is this necessary?)
+					if (CurrentGizmoActor.IsValid())
 					{
-						Width = MaxX - MinX + 1;
-						Height = MaxY - MinY + 1;
+						CurrentGizmoActor->SetTargetLandscape(LandscapeInfo);
 					}
+				}
 
-					LandscapeList.Add(FLandscapeListInfo(*LandscapeProxy->GetName(), LandscapeInfo,
-						LandscapeInfo->ComponentSizeQuads, LandscapeInfo->ComponentNumSubsections, Width, Height));
-					Index++;
+				int32 MinX, MinY, MaxX, MaxY;
+				int32 Width = 0, Height = 0;
+				if (LandscapeInfo->GetLandscapeExtent(MinX, MinY, MaxX, MaxY))
+				{
+					Width = MaxX - MinX + 1;
+					Height = MaxY - MinY + 1;
 				}
+
+				LandscapeList.Add(FLandscapeListInfo(*LandscapeProxy->GetName(), LandscapeInfo,
+					LandscapeInfo->ComponentSizeQuads, LandscapeInfo->ComponentNumSubsections, Width, Height));
+				Index++;
 			}
-		}
+			return true;
+		});
 	}
 
 	if (CurrentIndex == INDEX_NONE)
diff --git a/Engine/Source/Editor/UnrealEd/Private/UnrealEdSrv.cpp b/Engine/Source/Editor/UnrealEd/Private/UnrealEdSrv.cpp
--- a/Engine/Source/Editor/UnrealEd/Private/UnrealEdSrv.cpp
+++ b/Engine/Source/Editor/UnrealEd/Private/UnrealEdSrv.cpp
@@ -90,7 +90,7 @@
 #include "Editor/ActorPositioning.h"
 #include "Matinee/InterpData.h"
 #include "LandscapeInfo.h"
-#include "LandscapeInfoMap.h"
+#include "LandscapeSubsystem.h"
 #include "Framework/Notifications/NotificationManager.h"
 #include "Widgets/Notifications/SNotificationList.h"
 #include "Logging/LogScopedCategoryAndVerbosityOverride.h"
@@ -591,11 +591,11 @@ bool UUnrealEdEngine::HandleRecreateLandscapeCollisionCommand(const TCHAR* Str,
 {
 	if (!PlayWorld && InWorld && InWorld->GetWorldSettings())
 	{
-		for (auto It = ULandscapeInfoMap::GetLandscapeInfoMap(InWorld).Map.CreateIterator(); It; ++It)
+		ULandscapeSubsystem::ForEachLandscapeInfo(InWorld, [](ULandscapeInfo* LandscapeInfo)
 		{
-			ULandscapeInfo* Info = It.Value();
-			Info->RecreateCollisionComponents();
-		}
+			LandscapeInfo->RecreateCollisionComponents();
+			return true;
+		});
 	}
 	return true;
 }
@@ -604,11 +604,11 @@ bool UUnrealEdEngine::HandleRemoveLandscapeXYOffsetsCommand(const TCHAR* Str, FO
 {
 	if (!PlayWorld && InWorld && InWorld->GetWorldSettings())
 	{
-		for (auto It = ULandscapeInfoMap::GetLandscapeInfoMap(InWorld).Map.CreateIterator(); It; ++It)
+		ULandscapeSubsystem::ForEachLandscapeInfo(InWorld, [](ULandscapeInfo* LandscapeInfo)
 		{
-			ULandscapeInfo* Info = It.Value();
-			Info->RemoveXYOffsets();
-		}
+			LandscapeInfo->RemoveXYOffsets();
+			return true;
+		});
 	}
 	return true;
 }
diff --git a/Engine/Source/Editor/WorldBrowser/Private/Tiles/WorldTileModel.cpp b/Engine/Source/Editor/WorldBrowser/Private/Tiles/WorldTileModel.cpp
--- a/Engine/Source/Editor/WorldBrowser/Private/Tiles/WorldTileModel.cpp
+++ b/Engine/Source/Editor/WorldBrowser/Private/Tiles/WorldTileModel.cpp
@@ -16,7 +16,7 @@
 #include "Tiles/WorldTileCollectionModel.h"
 #include "Engine/WorldComposition.h"
 #include "GameFramework/WorldSettings.h"
-#include "LandscapeInfo.h"
+#include "LandscapeSubsystem.h"
 #include "LandscapeEditorModule.h"
 #include "LandscapeFileFormatInterface.h"
 #include "LandscapeStreamingProxy.h"
@@ -511,7 +511,7 @@ void FWorldTileModel::SetLevelPosition(const FIntVector& InPosition, const FIntP
 	{
 		UpdateLandscapeSectionsOffset(LandscapeOffset); // section offset is 2D 
 		bool bShowWarnings = true;
-		ULandscapeInfo::RecreateLandscapeInfo(LevelCollectionModel.GetWorld(), bShowWarnings);
+		ULandscapeSubsystem::RecreateLandscapeInfos(LevelCollectionModel.GetWorld(), bShowWarnings);
 	}
 	
 	// Transform child levels
diff --git a/Engine/Source/Runtime/Landscape/Classes/LandscapeInfo.h b/Engine/Source/Runtime/Landscape/Classes/LandscapeInfo.h
--- a/Engine/Source/Runtime/Landscape/Classes/LandscapeInfo.h
+++ b/Engine/Source/Runtime/Landscape/Classes/LandscapeInfo.h
@@ -209,11 +209,6 @@ public:
 	/** Resets all actors, proxies, components registrations */
 	LANDSCAPE_API void Reset();
 
-	/** Recreate all LandscapeInfo objects in given world
-	 *  @param  bMapCheck	Whether to warn about landscape errors
-	 */
-	LANDSCAPE_API static void RecreateLandscapeInfo(UWorld* InWorld, bool bMapCheck);
-
 	/** 
 	 *  Fixes up proxies relative position to landscape actor
 	 *  basically makes sure that each LandscapeProxy RootComponent transform reflects LandscapeSectionOffset value
diff --git a/Engine/Source/Runtime/Landscape/Classes/LandscapeInfoMap.h b/Engine/Source/Runtime/Landscape/Classes/LandscapeInfoMap.h
--- a/Engine/Source/Runtime/Landscape/Classes/LandscapeInfoMap.h
+++ b/Engine/Source/Runtime/Landscape/Classes/LandscapeInfoMap.h
@@ -18,21 +18,9 @@ class ULandscapeInfoMap : public UObject
 	GENERATED_UCLASS_BODY()
 
 public:
-	void BeginDestroy() override;
-	void PostDuplicate(bool bDuplicateForPIE) override;
 	void Serialize(FArchive& Ar) override;
 
 	static void AddReferencedObjects(UObject* InThis, FReferenceCollector& Collector);
 
 	TMap<FGuid, ULandscapeInfo*> Map;
-	TWeakObjectPtr<UWorld> World;
-
-	/**
-	* Gets landscape-specific data for given world.
-	*
-	* @param World A pointer to world that return data should be associated with.
-	*
-	* @returns Landscape-specific data associated with given world.
-	*/
-	LANDSCAPE_API static ULandscapeInfoMap& GetLandscapeInfoMap(const UWorld* World);
 };
diff --git a/Engine/Source/Runtime/Landscape/Private/Landscape.cpp b/Engine/Source/Runtime/Landscape/Private/Landscape.cpp
--- a/Engine/Source/Runtime/Landscape/Private/Landscape.cpp
+++ b/Engine/Source/Runtime/Landscape/Private/Landscape.cpp
@@ -20,7 +20,6 @@ Landscape.cpp: Terrain rendering
 #include "ShadowMap.h"
 #include "LandscapeComponent.h"
 #include "LandscapeLayerInfoObject.h"
-#include "LandscapeInfoMap.h"
 #include "EditorSupportDelegates.h"
 #include "LandscapeMeshProxyComponent.h"
 #include "LandscapeRender.h"
@@ -1186,41 +1185,23 @@ ALandscape* ALandscapeStreamingProxy::GetLandscapeActor()
 
 ULandscapeInfo* ALandscapeProxy::CreateLandscapeInfo(bool bMapCheck)
 {
-	ULandscapeInfo* LandscapeInfo = nullptr;
-
 	check(LandscapeGuid.IsValid());
 	UWorld* OwningWorld = GetWorld();
 	check(OwningWorld);
 	
-	auto& LandscapeInfoMap = ULandscapeInfoMap::GetLandscapeInfoMap(OwningWorld);
-	LandscapeInfo = LandscapeInfoMap.Map.FindRef(LandscapeGuid);
-
-	if (!LandscapeInfo)
-	{
-		check(!HasAnyFlags(RF_BeginDestroyed));
-		LandscapeInfo = NewObject<ULandscapeInfo>(GetTransientPackage(), NAME_None, RF_Transactional | RF_Transient);
-		LandscapeInfoMap.Modify(false);
-		LandscapeInfoMap.Map.Add(LandscapeGuid, LandscapeInfo);
-	}
+	ULandscapeInfo* LandscapeInfo = ULandscapeSubsystem::GetLandscapeInfo(OwningWorld, LandscapeGuid, /*bCreate*/ true);
 	check(LandscapeInfo);
 	LandscapeInfo->RegisterActor(this, bMapCheck);
-
 	return LandscapeInfo;
 }
 
 ULandscapeInfo* ALandscapeProxy::GetLandscapeInfo() const
 {
-	ULandscapeInfo* LandscapeInfo = nullptr;
-
 	check(LandscapeGuid.IsValid());
 	UWorld* OwningWorld = GetWorld();
-
-	if (OwningWorld != nullptr)
-	{
-		auto& LandscapeInfoMap = ULandscapeInfoMap::GetLandscapeInfoMap(OwningWorld);
-		LandscapeInfo = LandscapeInfoMap.Map.FindRef(LandscapeGuid);
-	}
-	return LandscapeInfo;
+	check(OwningWorld);
+	
+	return  ULandscapeSubsystem::GetLandscapeInfo(OwningWorld, LandscapeGuid, /*bCreate*/false);
 }
 
 FTransform ALandscapeProxy::LandscapeActorToWorld() const
@@ -2535,7 +2516,7 @@ void ALandscapeProxy::Destroyed()
 
 	if (GIsEditor && !World->IsGameWorld())
 	{
-		ULandscapeInfo::RecreateLandscapeInfo(World, false);
+		ULandscapeSubsystem::RecreateLandscapeInfos(World, false);
 
 		if (SplineComponent)
 		{
@@ -3017,70 +2998,6 @@ void ULandscapeInfo::UpdateComponentLayerWhitelist()
 	});
 }
 
-void ULandscapeInfo::RecreateLandscapeInfo(UWorld* InWorld, bool bMapCheck)
-{
-	check(InWorld);
-
-	ULandscapeInfoMap& LandscapeInfoMap = ULandscapeInfoMap::GetLandscapeInfoMap(InWorld);
-	LandscapeInfoMap.Modify();
-
-	// reset all LandscapeInfo objects
-	for (auto& LandscapeInfoPair : LandscapeInfoMap.Map)
-	{
-		ULandscapeInfo* LandscapeInfo = LandscapeInfoPair.Value;
-
-		if (LandscapeInfo != nullptr)
-		{
-			LandscapeInfo->Modify();
-			LandscapeInfo->Reset();
-		}
-	}
-
-	TMap<FGuid, TArray<ALandscapeProxy*>> ValidLandscapesMap;
-	// Gather all valid landscapes in the world
-	for (ALandscapeProxy* Proxy : TActorRange<ALandscapeProxy>(InWorld))
-	{
-		if (Proxy->GetLevel() &&
-			Proxy->GetLevel()->bIsVisible &&
-			!Proxy->HasAnyFlags(RF_BeginDestroyed) &&
-			!Proxy->IsPendingKill() &&
-			!Proxy->IsPendingKillPending())
-		{
-			ValidLandscapesMap.FindOrAdd(Proxy->GetLandscapeGuid()).Add(Proxy);
-		}
-	}
-
-	// Register landscapes in global landscape map
-	for (auto& ValidLandscapesPair : ValidLandscapesMap)
-	{
-		auto& LandscapeList = ValidLandscapesPair.Value;
-		for (ALandscapeProxy* Proxy : LandscapeList)
-		{
-			Proxy->CreateLandscapeInfo()->RegisterActor(Proxy, bMapCheck);
-		}
-	}
-
-	// Remove empty entries from global LandscapeInfo map
-	for (auto It = LandscapeInfoMap.Map.CreateIterator(); It; ++It)
-	{
-		ULandscapeInfo* Info = It.Value();
-
-		if (Info != nullptr && Info->GetLandscapeProxy() == nullptr)
-		{
-			Info->MarkPendingKill();
-			It.RemoveCurrent();
-		}
-		else if (Info == nullptr) // remove invalid entry
-		{
-			It.RemoveCurrent();
-		}
-	}
-
-	// We need to inform Landscape editor tools about LandscapeInfo updates
-	FEditorSupportDelegates::WorldChange.Broadcast();
-}
-
-
 #endif
 
 void ULandscapeInfo::ForAllLandscapeProxies(TFunctionRef<void(ALandscapeProxy*)> Fn) const
diff --git a/Engine/Source/Runtime/Landscape/Private/LandscapeEdit.cpp b/Engine/Source/Runtime/Landscape/Private/LandscapeEdit.cpp
--- a/Engine/Source/Runtime/Landscape/Private/LandscapeEdit.cpp
+++ b/Engine/Source/Runtime/Landscape/Private/LandscapeEdit.cpp
@@ -3779,7 +3779,7 @@ void ALandscapeProxy::PostEditMove(bool bFinished)
 
 	if (bFinished && !GetWorld()->IsGameWorld())
 	{
-		ULandscapeInfo::RecreateLandscapeInfo(GetWorld(), true);
+		ULandscapeSubsystem::RecreateLandscapeInfos(GetWorld(), true);
 		RecreateComponentsState();
 
 		if (SplineComponent)
@@ -4826,7 +4826,7 @@ void ALandscapeStreamingProxy::PostEditChangeProperty(FPropertyChangedEvent& Pro
 			if (GIsEditor && GetWorld() && !GetWorld()->IsPlayInEditor())
 			{
 				// TODO - only need to refresh the old and new landscape info
-				ULandscapeInfo::RecreateLandscapeInfo(GetWorld(), false);
+				ULandscapeSubsystem::RecreateLandscapeInfos(GetWorld(), false);
 				FixupWeightmaps();
 				InitializeProxyLayersWeightmapUsage();
 			}
diff --git a/Engine/Source/Runtime/Landscape/Private/LandscapeEditLayers.cpp b/Engine/Source/Runtime/Landscape/Private/LandscapeEditLayers.cpp
--- a/Engine/Source/Runtime/Landscape/Private/LandscapeEditLayers.cpp
+++ b/Engine/Source/Runtime/Landscape/Private/LandscapeEditLayers.cpp
@@ -31,7 +31,7 @@ LandscapeEditLayers.cpp: Landscape editing layers mode
 #include "Algo/Count.h"
 #include "LandscapeSettings.h"
 #include "LandscapeRender.h"
-#include "LandscapeInfoMap.h"
+#include "LandscapeSubsystem.h"
 #include "Misc/MessageDialog.h"
 #include "GameFramework/WorldSettings.h"
 #include "UObject/UObjectThreadContext.h"
@@ -129,14 +129,11 @@ private:
 			return;
 		}
 
-		auto& LandscapeInfoMap = ULandscapeInfoMap::GetLandscapeInfoMap(GWorld);
-		for (TPair<FGuid, ULandscapeInfo*>& Pair : LandscapeInfoMap.Map)
+		ULandscapeSubsystem::ForEachLandscapeInfo(GWorld, [](ULandscapeInfo* LandscapeInfo)
 		{
-			if (Pair.Value)
-			{
-				Pair.Value->ClearDirtyData();
-			}
-		}
+			LandscapeInfo->ClearDirtyData();
+			return true;
+		});
 
 		UE_LOG(LogLandscape, Display, TEXT("Landscape.Dirty: Cleared"));
 	}
diff --git a/Engine/Source/Runtime/Landscape/Private/LandscapeGizmoActor.cpp b/Engine/Source/Runtime/Landscape/Private/LandscapeGizmoActor.cpp
--- a/Engine/Source/Runtime/Landscape/Private/LandscapeGizmoActor.cpp
+++ b/Engine/Source/Runtime/Landscape/Private/LandscapeGizmoActor.cpp
@@ -13,7 +13,7 @@
 #include "LandscapeInfo.h"
 #include "Engine/Texture2D.h"
 #include "LandscapeLayerInfoObject.h"
-#include "LandscapeInfoMap.h"
+#include "LandscapeSubsystem.h"
 #include "LandscapeDataAccess.h"
 #include "LandscapeRender.h"
 #include "LandscapeGizmoActiveActor.h"
@@ -650,15 +650,15 @@ void ALandscapeGizmoActiveActor::SetTargetLandscape(ULandscapeInfo* LandscapeInf
 		TargetLandscapeInfo = nullptr;
 		if (GetWorld())
 		{
-			for (const TPair<FGuid, ULandscapeInfo*>& InfoMapPair : ULandscapeInfoMap::GetLandscapeInfoMap(GetWorld()).Map)
+			ULandscapeSubsystem::ForEachLandscapeInfo(GetWorld(), [this](ULandscapeInfo* CandidateInfo)
 			{
-				ULandscapeInfo* CandidateInfo = InfoMapPair.Value;
 				if (CandidateInfo && !CandidateInfo->HasAnyFlags(RF_BeginDestroyed) && CandidateInfo->GetLandscapeProxy() != nullptr)
 				{
 					TargetLandscapeInfo = CandidateInfo;
-					break;
+					return false;
 				}
-			}
+				return true;
+			});
 		}
 	}
 	else
diff --git a/Engine/Source/Runtime/Landscape/Private/LandscapeInfoMap.cpp b/Engine/Source/Runtime/Landscape/Private/LandscapeInfoMap.cpp
--- a/Engine/Source/Runtime/Landscape/Private/LandscapeInfoMap.cpp
+++ b/Engine/Source/Runtime/Landscape/Private/LandscapeInfoMap.cpp
@@ -6,17 +6,9 @@
 
 ULandscapeInfoMap::ULandscapeInfoMap(const FObjectInitializer& ObjectInitializer)
 	: Super(ObjectInitializer)
-	, World(nullptr)
 {
 }
 
-void ULandscapeInfoMap::PostDuplicate(bool bDuplicateForPIE)
-{
-	Super::PostDuplicate(bDuplicateForPIE);
-
-	check(Map.Num() == 0);
-}
-
 void ULandscapeInfoMap::Serialize(FArchive& Ar)
 {
 	Super::Serialize(Ar);
@@ -27,28 +19,8 @@ void ULandscapeInfoMap::Serialize(FArchive& Ar)
 	}
 }
 
-void ULandscapeInfoMap::BeginDestroy()
-{
-	if (World.IsValid())
-	{
-		World->PerModuleDataObjects.Remove(this);
-	}
-
-	Super::BeginDestroy();
-}
-
 void ULandscapeInfoMap::AddReferencedObjects(UObject* InThis, FReferenceCollector& Collector)
 {
 	ULandscapeInfoMap* This = CastChecked<ULandscapeInfoMap>(InThis);
 	Collector.AddReferencedObjects(This->Map, This);
 }
-
-ULandscapeInfoMap& ULandscapeInfoMap::GetLandscapeInfoMap(const UWorld* World)
-{
-	ULandscapeInfoMap *FoundObject = nullptr;
-	World->PerModuleDataObjects.FindItemByClass(&FoundObject);
-
-	checkf(FoundObject, TEXT("ULandscapInfoMap object was not created for this UWorld."));
-
-	return *FoundObject;
-}
diff --git a/Engine/Source/Runtime/Landscape/Private/LandscapeModule.cpp b/Engine/Source/Runtime/Landscape/Private/LandscapeModule.cpp
--- a/Engine/Source/Runtime/Landscape/Private/LandscapeModule.cpp
+++ b/Engine/Source/Runtime/Landscape/Private/LandscapeModule.cpp
@@ -9,7 +9,6 @@
 #include "Materials/MaterialInterface.h"
 #include "LandscapeComponent.h"
 #include "LandscapeVersion.h"
-#include "LandscapeInfoMap.h"
 #include "Materials/MaterialInstance.h"
 
 #include "LandscapeProxy.h"
@@ -27,26 +26,6 @@ public:
 	void ShutdownModule() override;
 };
 
-/**
- * Add landscape-specific per-world data.
- *
- * @param World A pointer to world that this data should be created for.
- */
-void AddPerWorldLandscapeData(UWorld* World)
-{
-	EObjectFlags NewLandscapeDataFlags = RF_NoFlags;
-	if (!World->PerModuleDataObjects.FindItemByClass<ULandscapeInfoMap>())
-	{
-		if (World->HasAnyFlags(RF_Transactional))
-		{
-			NewLandscapeDataFlags = RF_Transactional;
-		}
-		ULandscapeInfoMap* InfoMap = NewObject<ULandscapeInfoMap>(GetTransientPackage(), NAME_None, NewLandscapeDataFlags);
-		InfoMap->World = World;
-		World->PerModuleDataObjects.Add(InfoMap);
-	}
-}
-
 #if WITH_EDITOR
 /**
  * Gets landscape-specific material's static parameters values.
@@ -65,32 +44,6 @@ void LandscapeMaterialsParameterValuesGetter(FStaticParameterSet &OutStaticParam
 bool LandscapeMaterialsParameterSetUpdater(FStaticParameterSet &OutStaticParameterSet, UMaterial* Material);
 #endif // WITH_EDITOR
 
-/**
- * Function that will fire every time a world is created.
- *
- * @param World A world that was created.
- * @param IVS Initialization values.
- */
-void WorldCreationEventFunction(UWorld* World)
-{
-	AddPerWorldLandscapeData(World);
-}
-
-/**
- * Function that will fire every time a world is destroyed.
- *
- * @param World A world that's being destroyed.
- */
-void WorldDestroyEventFunction(UWorld* World)
-{
-	World->PerModuleDataObjects.RemoveAll(
-		[](UObject* Object)
-		{
-			return Object != nullptr && Object->IsA(ULandscapeInfoMap::StaticClass());
-		}
-	);
-}
-
 #if WITH_EDITOR
 /**
  * Gets array of Landscape-specific textures and materials connected with given
@@ -144,36 +97,6 @@ void WorldRenameEventFunction(UWorld* World, const TCHAR* InName, UObject* NewOu
 }
 #endif
 
-/**
- * A function that fires everytime a world is duplicated.
- *
- * If there are some objects duplicated during this event fill out
- * ReplacementMap and ObjectsToFixReferences in order to properly fix
- * references in objects created during this duplication.
- *
- * @param World A world that was duplicated.
- * @param bDuplicateForPIE If this duplication was done for PIE.
- * @param ReplacementMap Replacement map (i.e. old object -> new object).
- * @param ObjectsToFixReferences Array of objects that may contain bad
- *		references to old objects.
- */
-void WorldDuplicateEventFunction(UWorld* World, bool bDuplicateForPIE, TMap<UObject*, UObject*>& ReplacementMap, TArray<UObject*>& ObjectsToFixReferences)
-{
-	int32 Index;
-	ULandscapeInfoMap* InfoMap;
-	if (World->PerModuleDataObjects.FindItemByClass(&InfoMap, &Index))
-	{
-		ULandscapeInfoMap* NewInfoMap = Cast<ULandscapeInfoMap>( StaticDuplicateObject(InfoMap, InfoMap->GetOuter()) );
-		NewInfoMap->World = World;
-
-		World->PerModuleDataObjects[Index] = NewInfoMap;
-	}
-	else
-	{
-		AddPerWorldLandscapeData(World);
-	}
-}
-
 void FLandscapeModule::StartupModule()
 {
 #if WITH_EDITOR
@@ -189,22 +112,11 @@ void FLandscapeModule::StartupModule()
 	);
 #endif // WITH_EDITOR
 
-	FWorldDelegates::OnPostWorldCreation.AddStatic(
-		&WorldCreationEventFunction
-	);
-	FWorldDelegates::OnPreWorldFinishDestroy.AddStatic(
-		&WorldDestroyEventFunction
-	);
-
 #if WITH_EDITOR
 	FWorldDelegates::OnPreWorldRename.AddStatic(
 		&WorldRenameEventFunction
 	);
 #endif // WITH_EDITOR
-
-	FWorldDelegates::OnPostDuplicate.AddStatic(
-		&WorldDuplicateEventFunction
-	);
 }
 
 void FLandscapeModule::ShutdownModule()
diff --git a/Engine/Source/Runtime/Landscape/Private/LandscapeSplines.cpp b/Engine/Source/Runtime/Landscape/Private/LandscapeSplines.cpp
--- a/Engine/Source/Runtime/Landscape/Private/LandscapeSplines.cpp
+++ b/Engine/Source/Runtime/Landscape/Private/LandscapeSplines.cpp
@@ -40,7 +40,7 @@
 #include "Misc/UObjectToken.h"
 #include "Landscape.h"
 #include "LandscapeLayerInfoObject.h"
-#include "LandscapeInfoMap.h"
+#include "LandscapeSubsystem.h"
 #endif
 
 IMPLEMENT_HIT_PROXY(HLandscapeSplineProxy, HHitProxy);
@@ -81,26 +81,25 @@ struct FLandscapeFixSplines
 			return;
 		}
 
-		auto& LandscapeInfoMap = ULandscapeInfoMap::GetLandscapeInfoMap(GWorld);
-		for (TPair<FGuid, ULandscapeInfo*>& Pair : LandscapeInfoMap.Map)
+		ULandscapeSubsystem::ForEachLandscapeInfo(GWorld, [](ULandscapeInfo* LandscapeInfo)
 		{
-			if (Pair.Value)
+			LandscapeInfo->ForAllLandscapeProxies([](ALandscapeProxy* Proxy)
 			{
-				Pair.Value->ForAllLandscapeProxies([](ALandscapeProxy* Proxy)
+				if (Proxy && Proxy->SplineComponent)
 				{
-					if (Proxy && Proxy->SplineComponent)
-					{
-						Proxy->SplineComponent->RebuildAllSplines();
-					}
-				});
+					Proxy->SplineComponent->RebuildAllSplines();
+				}
+			});
 
-				if (Pair.Value->LandscapeActor && Pair.Value->LandscapeActor->HasLayersContent())
+			if (ALandscape* Landscape = LandscapeInfo->LandscapeActor.Get())
+			{
+				if (Landscape->HasLayersContent())
 				{
-					Pair.Value->LandscapeActor->RequestSplineLayerUpdate();
+					Landscape->RequestSplineLayerUpdate();
 				}
 			}
-		}
-
+			return true;
+		});
 	}
 };
 
diff --git a/Engine/Source/Runtime/Landscape/Private/LandscapeSubsystem.cpp b/Engine/Source/Runtime/Landscape/Private/LandscapeSubsystem.cpp
--- a/Engine/Source/Runtime/Landscape/Private/LandscapeSubsystem.cpp
+++ b/Engine/Source/Runtime/Landscape/Private/LandscapeSubsystem.cpp
@@ -9,11 +9,14 @@
 #include "LandscapeProxy.h"
 #include "LandscapeStreamingProxy.h"
 #include "LandscapeInfo.h"
+#include "LandscapeInfoMap.h"
 #include "ProfilingDebugging/CsvProfiler.h"
 #include "WorldPartition/WorldPartitionSubsystem.h"
 #include "ActorPartition/ActorPartitionSubsystem.h"
 #include "Engine/World.h"
 #include "Math/IntRect.h"
+#include "EngineUtils.h"
+#include "EditorSupportDelegates.h"
 
 static int32 GUseStreamingManagerForCameras = 1;
 static FAutoConsoleVariableRef CVarUseStreamingManagerForCameras(
@@ -23,8 +26,13 @@ static FAutoConsoleVariableRef CVarUseStreamingManagerForCameras(
 
 DECLARE_CYCLE_STAT(TEXT("LandscapeSubsystem Tick"), STAT_LandscapeSubsystemTick, STATGROUP_Landscape);
 
-ULandscapeSubsystem::ULandscapeSubsystem()
+ULandscapeSubsystem::ULandscapeSubsystem(const FObjectInitializer& ObjectInitializer)
+#if WITH_EDITOR
+	: GrassMapsBuilder(nullptr)
+#endif
 {
+	LandscapeInfos = ObjectInitializer.CreateDefaultSubobject<ULandscapeInfoMap>(this, "LandscapeInfoMap");
+	LandscapeInfos->SetFlags(RF_Transactional);
 }
 
 ULandscapeSubsystem::~ULandscapeSubsystem()
@@ -41,6 +49,55 @@ void ULandscapeSubsystem::UnregisterActor(ALandscapeProxy* Proxy)
 	Proxies.Remove(Proxy);
 }
 
+void ULandscapeSubsystem::ForEachLandscapeInfo(TFunctionRef<bool(ULandscapeInfo*)> Predicate)
+{
+	for (const TPair<FGuid, ULandscapeInfo*>& Pair : LandscapeInfos->Map)
+	{
+		if (Pair.Value && !Pair.Value->IsPendingKill())
+		{
+			if (!Predicate(Pair.Value))
+			{
+				return;
+			}
+		}
+	}
+}
+
+void ULandscapeSubsystem::ForEachLandscapeInfo(const UWorld* InWorld, TFunctionRef<bool(ULandscapeInfo*)> Predicate)
+{
+	if (ULandscapeSubsystem* LandscapeSubsystem = UWorld::GetSubsystem<ULandscapeSubsystem>(InWorld))
+	{
+		LandscapeSubsystem->ForEachLandscapeInfo(Predicate);
+	}
+}
+
+ULandscapeInfo* ULandscapeSubsystem::GetLandscapeInfo(FGuid LandscapeGuid, bool bCreate)
+{
+	ULandscapeInfo* LandscapeInfo = LandscapeInfos->Map.FindRef(LandscapeGuid);
+	if (LandscapeInfo)
+	{
+		return LandscapeInfo;
+	}
+	else if (bCreate)
+	{
+		LandscapeInfo = NewObject<ULandscapeInfo>(LandscapeInfos, NAME_None, RF_Transactional | RF_Transient);
+		LandscapeInfos->Modify(false);
+		LandscapeInfos->Map.Add(LandscapeGuid, LandscapeInfo);
+	}
+
+	return LandscapeInfo;
+}
+
+ULandscapeInfo* ULandscapeSubsystem::GetLandscapeInfo(const UWorld* InWorld, FGuid LandscapeGuid, bool bCreate)
+{
+	if (ULandscapeSubsystem* LandscapeSubsystem = UWorld::GetSubsystem<ULandscapeSubsystem>(InWorld))
+	{
+		return LandscapeSubsystem->GetLandscapeInfo(LandscapeGuid, bCreate);
+	}
+
+	return nullptr;
+}
+
 void ULandscapeSubsystem::Initialize(FSubsystemCollectionBase& Collection)
 {
 	Super::Initialize(Collection);
@@ -155,6 +212,74 @@ void ULandscapeSubsystem::Tick(float DeltaTime)
 }
 
 #if WITH_EDITOR
+void ULandscapeSubsystem::RecreateLandscapeInfos(bool bMapCheck)
+{
+	LandscapeInfos->Modify();
+
+	// reset all LandscapeInfo objects
+	for (auto& LandscapeInfoPair : LandscapeInfos->Map)
+	{
+		ULandscapeInfo* LandscapeInfo = LandscapeInfoPair.Value;
+
+		if (LandscapeInfo != nullptr)
+		{
+			LandscapeInfo->Modify();
+			LandscapeInfo->Reset();
+		}
+	}
+
+	TMap<FGuid, TArray<ALandscapeProxy*>> ValidLandscapesMap;
+	// Gather all valid landscapes in the world
+	for (ALandscapeProxy* Proxy : TActorRange<ALandscapeProxy>(GetWorld()))
+	{
+		if (Proxy->GetLevel() &&
+			Proxy->GetLevel()->bIsVisible &&
+			!Proxy->HasAnyFlags(RF_BeginDestroyed) &&
+			!Proxy->IsPendingKill() &&
+			!Proxy->IsPendingKillPending())
+		{
+			ValidLandscapesMap.FindOrAdd(Proxy->GetLandscapeGuid()).Add(Proxy);
+		}
+	}
+
+	// Register landscapes in global landscape map
+	for (auto& ValidLandscapesPair : ValidLandscapesMap)
+	{
+		auto& LandscapeList = ValidLandscapesPair.Value;
+		for (ALandscapeProxy* Proxy : LandscapeList)
+		{
+			Proxy->CreateLandscapeInfo()->RegisterActor(Proxy, bMapCheck);
+		}
+	}
+
+	// Remove empty entries from global LandscapeInfo map
+	for (auto It = LandscapeInfos->Map.CreateIterator(); It; ++It)
+	{
+		ULandscapeInfo* Info = It.Value();
+
+		if (Info != nullptr && Info->GetLandscapeProxy() == nullptr)
+		{
+			Info->MarkPendingKill();
+			It.RemoveCurrent();
+		}
+		else if (Info == nullptr) // remove invalid entry
+		{
+			It.RemoveCurrent();
+		}
+	}
+
+	// We need to inform Landscape editor tools about LandscapeInfo updates
+	FEditorSupportDelegates::WorldChange.Broadcast();
+}
+
+void ULandscapeSubsystem::RecreateLandscapeInfos(UWorld* InWorld, bool bMapCheck)
+{
+	if (ULandscapeSubsystem* LandscapeSubsystem = UWorld::GetSubsystem<ULandscapeSubsystem>(InWorld))
+	{
+		LandscapeSubsystem->RecreateLandscapeInfos(bMapCheck);
+	}
+}
+
 void ULandscapeSubsystem::BuildGrassMaps()
 {
 	GrassMapsBuilder->Build();
diff --git a/Engine/Source/Runtime/Landscape/Public/LandscapeSubsystem.h b/Engine/Source/Runtime/Landscape/Public/LandscapeSubsystem.h
--- a/Engine/Source/Runtime/Landscape/Public/LandscapeSubsystem.h
+++ b/Engine/Source/Runtime/Landscape/Public/LandscapeSubsystem.h
@@ -11,6 +11,7 @@
 
 class ALandscapeProxy;
 class ULandscapeInfo;
+class ULandscapeInfoMap;
 
 UCLASS(MinimalAPI)
 class ULandscapeSubsystem : public UWorldSubsystem, public FTickableGameObject
@@ -18,12 +19,18 @@ class ULandscapeSubsystem : public UWorldSubsystem, public FTickableGameObject
 	GENERATED_BODY()
 
 public:
-	ULandscapeSubsystem();
+	ULandscapeSubsystem(const FObjectInitializer& ObjectInitializer);
 	virtual ~ULandscapeSubsystem();
 
 	void RegisterActor(ALandscapeProxy* Proxy);
 	void UnregisterActor(ALandscapeProxy* Proxy);
 
+	void ForEachLandscapeInfo(TFunctionRef<bool(ULandscapeInfo*)> Predicate);
+	LANDSCAPE_API static void ForEachLandscapeInfo(const UWorld* InWorld, TFunctionRef<bool(ULandscapeInfo*)> Predicate);
+
+	ULandscapeInfo* GetLandscapeInfo(FGuid LandscapeGuid, bool bCreate = false);
+	LANDSCAPE_API static ULandscapeInfo* GetLandscapeInfo(const UWorld* InWorld, FGuid LandscapeGuid, bool bCreate = false);
+
 	// Begin FTickableGameObject overrides
 	virtual void Tick(float DeltaTime) override;
 	virtual bool IsTickableInEditor() const override { return true; }
@@ -38,6 +45,9 @@ public:
 	LANDSCAPE_API bool IsGridBased() const;
 	LANDSCAPE_API void UpdateGrid(ULandscapeInfo* LandscapeInfo, uint32 GridSizeInComponents);
 	LANDSCAPE_API ALandscapeProxy* FindOrAddLandscapeProxy(ULandscapeInfo* LandscapeInfo, const FIntPoint& SectionBase);
+
+	void RecreateLandscapeInfos(bool bMapCheck);
+	LANDSCAPE_API static void RecreateLandscapeInfos(UWorld* InWorld, bool bMapCheck);
 #endif
 
 private:
@@ -48,6 +58,9 @@ private:
 
 	TArray<ALandscapeProxy*> Proxies;
 
+	UPROPERTY()
+	ULandscapeInfoMap* LandscapeInfos;
+
 #if WITH_EDITOR
 	class FLandscapeGrassMapsBuilder* GrassMapsBuilder;
 	FLandscapeActorDescFactory LandscapeActorDescFactory;
