When selecting the Landscape gizmo actor, it needs to be the only actor selected, which means that it will change the selection within a selection changed event, triggering an assert with recent selection changes.
To fix this, only deselect actors when the gizmo actor is not the only actor selected, while JamieD makes a fix to prevent nested actor selection changed events.

diff --git a/Engine/Source/Editor/LandscapeEditor/Private/LandscapeEdMode.cpp b/Engine/Source/Editor/LandscapeEditor/Private/LandscapeEdMode.cpp
--- a/Engine/Source/Editor/LandscapeEditor/Private/LandscapeEdMode.cpp
+++ b/Engine/Source/Editor/LandscapeEditor/Private/LandscapeEdMode.cpp
@@ -3228,22 +3228,11 @@ bool FEdModeLandscape::IsSelectionAllowed(AActor* InActor, bool bInSelection) co
 /** FEdMode: Called when the currently selected actor has changed */
 void FEdModeLandscape::ActorSelectionChangeNotify()
 {
-	if (CurrentGizmoActor.IsValid() && CurrentGizmoActor->IsSelected())
+	if (CurrentGizmoActor.IsValid() && CurrentGizmoActor->IsSelected() && (GEditor->GetSelectedActors()->CountSelections<AActor>() != 1))
 	{
 		GEditor->SelectNone(false, true);
 		GEditor->SelectActor(CurrentGizmoActor.Get(), true, false, true);
 	}
-	/*
-		USelection* EditorSelection = GEditor->GetSelectedActors();
-		for ( FSelectionIterator Itor(EditorSelection) ; Itor ; ++Itor )
-		{
-		if (((*Itor)->IsA(ALandscapeGizmoActor::StaticClass())) )
-		{
-		bIsGizmoSelected = true;
-		break;
-		}
-		}
-	*/
 }
 
 void FEdModeLandscape::ActorMoveNotify()
