Landscape: NewLandscape uses FLandscapeConfig presets instead of duplicating.

diff --git a/Engine/Source/Editor/LandscapeEditor/Private/LandscapeEditorDetailCustomization_NewLandscape.cpp b/Engine/Source/Editor/LandscapeEditor/Private/LandscapeEditorDetailCustomization_NewLandscape.cpp
--- a/Engine/Source/Editor/LandscapeEditor/Private/LandscapeEditorDetailCustomization_NewLandscape.cpp
+++ b/Engine/Source/Editor/LandscapeEditor/Private/LandscapeEditorDetailCustomization_NewLandscape.cpp
@@ -20,6 +20,7 @@
 #include "LandscapeEditorObject.h"
 #include "Landscape.h"
 #include "LandscapeEditorUtils.h"
+#include "LandscapeConfigHelper.h"
 #include "NewLandscapeUtils.h"
 
 #include "DetailLayoutBuilder.h"
@@ -597,10 +598,10 @@ TSharedRef<SWidget> FLandscapeEditorDetailCustomization_NewLandscape::GetSection
 {
 	FMenuBuilder MenuBuilder(true, nullptr);
 
-	for (int32 i = 0; i < UE_ARRAY_COUNT(FNewLandscapeUtils::SectionSizes); i++)
+	for (int32 i = 0; i < UE_ARRAY_COUNT(FLandscapeConfig::SubsectionSizeQuadsValues); i++)
 	{
-		MenuBuilder.AddMenuEntry(FText::Format(LOCTEXT("NxNQuads", "{0}\u00D7{0} Quads"), FText::AsNumber(FNewLandscapeUtils::SectionSizes[i])), FText::GetEmpty(),
-			FSlateIcon(), FExecuteAction::CreateStatic(&OnChangeSectionSize, PropertyHandle, FNewLandscapeUtils::SectionSizes[i]));
+		MenuBuilder.AddMenuEntry(FText::Format(LOCTEXT("NxNQuads", "{0}\u00D7{0} Quads"), FText::AsNumber(FLandscapeConfig::SubsectionSizeQuadsValues[i])), FText::GetEmpty(),
+			FSlateIcon(), FExecuteAction::CreateStatic(&OnChangeSectionSize, PropertyHandle, FLandscapeConfig::SubsectionSizeQuadsValues[i]));
 	}
 
 	return MenuBuilder.MakeWidget();
@@ -629,13 +630,13 @@ TSharedRef<SWidget> FLandscapeEditorDetailCustomization_NewLandscape::GetSection
 {
 	FMenuBuilder MenuBuilder(true, nullptr);
 
-	for (int32 i = 0; i < UE_ARRAY_COUNT(FNewLandscapeUtils::NumSections); i++)
+	for (int32 i = 0; i < UE_ARRAY_COUNT(FLandscapeConfig::NumSectionValues); i++)
 	{
 		FFormatNamedArguments Args;
-		Args.Add(TEXT("Width"), FNewLandscapeUtils::NumSections[i]);
-		Args.Add(TEXT("Height"), FNewLandscapeUtils::NumSections[i]);
-		MenuBuilder.AddMenuEntry(FText::Format(FNewLandscapeUtils::NumSections[i] == 1 ? LOCTEXT("1x1Section", "{Width}\u00D7{Height} Section") : LOCTEXT("NxNSections", "{Width}\u00D7{Height} Sections"), Args),
-			FText::GetEmpty(), FSlateIcon(), FExecuteAction::CreateStatic(&OnChangeSectionsPerComponent, PropertyHandle, FNewLandscapeUtils::NumSections[i]));
+		Args.Add(TEXT("Width"), FLandscapeConfig::NumSectionValues[i]);
+		Args.Add(TEXT("Height"), FLandscapeConfig::NumSectionValues[i]);
+		MenuBuilder.AddMenuEntry(FText::Format(FLandscapeConfig::NumSectionValues[i] == 1 ? LOCTEXT("1x1Section", "{Width}\u00D7{Height} Section") : LOCTEXT("NxNSections", "{Width}\u00D7{Height} Sections"), Args),
+			FText::GetEmpty(), FSlateIcon(), FExecuteAction::CreateStatic(&OnChangeSectionsPerComponent, PropertyHandle, FLandscapeConfig::NumSectionValues[i]));
 	}
 
 	return MenuBuilder.MakeWidget();
diff --git a/Engine/Source/Editor/LandscapeEditor/Private/NewLandscapeUtils.cpp b/Engine/Source/Editor/LandscapeEditor/Private/NewLandscapeUtils.cpp
--- a/Engine/Source/Editor/LandscapeEditor/Private/NewLandscapeUtils.cpp
+++ b/Engine/Source/Editor/LandscapeEditor/Private/NewLandscapeUtils.cpp
@@ -8,6 +8,7 @@
 #include "LandscapeEditorModule.h"
 #include "LandscapeEditorUtils.h"
 #include "LandscapeEdMode.h"
+#include "LandscapeConfigHelper.h"
 
 #include "Misc/MessageDialog.h"
 #include "Modules/ModuleManager.h"
@@ -16,9 +17,6 @@
 
 #define LOCTEXT_NAMESPACE "LandscapeEditor.NewLandscape"
 
-const int32 FNewLandscapeUtils::SectionSizes[6] = { 7, 15, 31, 63, 127, 255 };
-const int32 FNewLandscapeUtils::NumSections[2] = { 1, 2 };
-
 void FNewLandscapeUtils::ChooseBestComponentSizeForImport(ULandscapeEditorObject* UISettings)
 {
 	int32 Width = UISettings->ImportLandscape_Width;
@@ -28,12 +26,12 @@ void FNewLandscapeUtils::ChooseBestComponentSizeForImport(ULandscapeEditorObject
 	if(Width > 0 && Height > 0)
 	{
 		// Try to find a section size and number of sections that exactly matches the dimensions of the heightfield
-		for(int32 SectionSizesIdx = UE_ARRAY_COUNT(SectionSizes) - 1; SectionSizesIdx >= 0; SectionSizesIdx--)
+		for(int32 SectionSizesIdx = UE_ARRAY_COUNT(FLandscapeConfig::SubsectionSizeQuadsValues) - 1; SectionSizesIdx >= 0; SectionSizesIdx--)
 		{
-			for(int32 NumSectionsIdx = UE_ARRAY_COUNT(NumSections) - 1; NumSectionsIdx >= 0; NumSectionsIdx--)
+			for(int32 NumSectionsIdx = UE_ARRAY_COUNT(FLandscapeConfig::NumSectionValues) - 1; NumSectionsIdx >= 0; NumSectionsIdx--)
 			{
-				int32 ss = SectionSizes[SectionSizesIdx];
-				int32 ns = NumSections[NumSectionsIdx];
+				int32 ss = FLandscapeConfig::SubsectionSizeQuadsValues[SectionSizesIdx];
+				int32 ns = FLandscapeConfig::NumSectionValues[NumSectionsIdx];
 
 				if(((Width - 1) % (ss * ns)) == 0 && ((Width - 1) / (ss * ns)) <= 32 &&
 					((Height - 1) % (ss * ns)) == 0 && ((Height - 1) / (ss * ns)) <= 32)
@@ -58,19 +56,19 @@ void FNewLandscapeUtils::ChooseBestComponentSizeForImport(ULandscapeEditorObject
 			// if there was no exact match, try increasing the section size until we encompass the whole heightmap
 			const int32 CurrentSectionSize = UISettings->NewLandscape_QuadsPerSection;
 			const int32 CurrentNumSections = UISettings->NewLandscape_SectionsPerComponent;
-			for(int32 SectionSizesIdx = 0; SectionSizesIdx < UE_ARRAY_COUNT(SectionSizes); SectionSizesIdx++)
+			for(int32 SectionSizesIdx = 0; SectionSizesIdx < UE_ARRAY_COUNT(FLandscapeConfig::SubsectionSizeQuadsValues); SectionSizesIdx++)
 			{
-				if(SectionSizes[SectionSizesIdx] < CurrentSectionSize)
+				if(FLandscapeConfig::SubsectionSizeQuadsValues[SectionSizesIdx] < CurrentSectionSize)
 				{
 					continue;
 				}
 
-				const int32 ComponentsX = FMath::DivideAndRoundUp((Width - 1), SectionSizes[SectionSizesIdx] * CurrentNumSections);
-				const int32 ComponentsY = FMath::DivideAndRoundUp((Height - 1), SectionSizes[SectionSizesIdx] * CurrentNumSections);
+				const int32 ComponentsX = FMath::DivideAndRoundUp((Width - 1), FLandscapeConfig::SubsectionSizeQuadsValues[SectionSizesIdx] * CurrentNumSections);
+				const int32 ComponentsY = FMath::DivideAndRoundUp((Height - 1), FLandscapeConfig::SubsectionSizeQuadsValues[SectionSizesIdx] * CurrentNumSections);
 				if(ComponentsX <= 32 && ComponentsY <= 32)
 				{
 					bFoundMatch = true;
-					UISettings->NewLandscape_QuadsPerSection = SectionSizes[SectionSizesIdx];
+					UISettings->NewLandscape_QuadsPerSection = FLandscapeConfig::SubsectionSizeQuadsValues[SectionSizesIdx];
 					//UISettings->NewLandscape_SectionsPerComponent = ;
 					UISettings->NewLandscape_ComponentCount.X = ComponentsX;
 					UISettings->NewLandscape_ComponentCount.Y = ComponentsY;
@@ -83,8 +81,8 @@ void FNewLandscapeUtils::ChooseBestComponentSizeForImport(ULandscapeEditorObject
 		if(!bFoundMatch)
 		{
 			// if the heightmap is very large, fall back to using the largest values we support
-			const int32 MaxSectionSize = SectionSizes[UE_ARRAY_COUNT(SectionSizes) - 1];
-			const int32 MaxNumSubSections = NumSections[UE_ARRAY_COUNT(NumSections) - 1];
+			const int32 MaxSectionSize = FLandscapeConfig::SubsectionSizeQuadsValues[UE_ARRAY_COUNT(FLandscapeConfig::SubsectionSizeQuadsValues) - 1];
+			const int32 MaxNumSubSections = FLandscapeConfig::NumSectionValues[UE_ARRAY_COUNT(FLandscapeConfig::NumSectionValues) - 1];
 			const int32 ComponentsX = FMath::DivideAndRoundUp((Width - 1), MaxSectionSize * MaxNumSubSections);
 			const int32 ComponentsY = FMath::DivideAndRoundUp((Height - 1), MaxSectionSize * MaxNumSubSections);
 
diff --git a/Engine/Source/Editor/LandscapeEditor/Public/NewLandscapeUtils.h b/Engine/Source/Editor/LandscapeEditor/Public/NewLandscapeUtils.h
--- a/Engine/Source/Editor/LandscapeEditor/Public/NewLandscapeUtils.h
+++ b/Engine/Source/Editor/LandscapeEditor/Public/NewLandscapeUtils.h
@@ -16,7 +16,4 @@ public:
 	static void ImportLandscapeData( ULandscapeEditorObject* UISettings, TArray< FLandscapeFileResolution >& ImportResolutions );
 	static TOptional< TArray< FLandscapeImportLayerInfo > > CreateImportLayersInfo( ULandscapeEditorObject* UISettings, int32 NewLandscapePreviewMode );
 	static TArray< uint16 > ComputeHeightData( ULandscapeEditorObject* UISettings, TArray< FLandscapeImportLayerInfo >& ImportLayers, int32 NewLandscapePreviewMode );
-
-	static const int32 SectionSizes[6];
-	static const int32 NumSections[2];
 };
