Landscape: Import class cleanup (common templated fileformat and data)

diff --git a/Engine/Source/Editor/LandscapeEditor/Private/LandscapeEdMode.cpp b/Engine/Source/Editor/LandscapeEditor/Private/LandscapeEdMode.cpp
--- a/Engine/Source/Editor/LandscapeEditor/Private/LandscapeEdMode.cpp
+++ b/Engine/Source/Editor/LandscapeEditor/Private/LandscapeEdMode.cpp
@@ -3309,7 +3309,7 @@ void FEdModeLandscape::ImportData(const FLandscapeTargetListInfo& TargetInfo, co
 
 			FLandscapeFileResolution ImportResolution = {0, 0};
 
-			const FLandscapeHeightmapInfo HeightmapInfo = HeightmapFormat->Validate(*Filename);
+			const FLandscapeFileInfo HeightmapInfo = HeightmapFormat->Validate(*Filename, NAME_None);
 
 			// display error message if there is one, and abort the import
 			if (HeightmapInfo.ResultCode == ELandscapeImportResult::Error)
@@ -3373,7 +3373,7 @@ void FEdModeLandscape::ImportData(const FLandscapeTargetListInfo& TargetInfo, co
 				}
 			}
 
-			FLandscapeHeightmapImportData ImportData = HeightmapFormat->Import(*Filename, ImportResolution);
+			FLandscapeHeightmapImportData ImportData = HeightmapFormat->Import(*Filename, NAME_None, ImportResolution);
 
 			if (ImportData.ResultCode == ELandscapeImportResult::Error)
 			{
@@ -3425,7 +3425,7 @@ void FEdModeLandscape::ImportData(const FLandscapeTargetListInfo& TargetInfo, co
 
 			FLandscapeFileResolution ImportResolution = {0, 0};
 
-			const FLandscapeWeightmapInfo WeightmapInfo = WeightmapFormat->Validate(*Filename, TargetInfo.LayerName);
+			const FLandscapeFileInfo WeightmapInfo = WeightmapFormat->Validate(*Filename, TargetInfo.LayerName);
 
 			// display error message if there is one, and abort the import
 			if (WeightmapInfo.ResultCode == ELandscapeImportResult::Error)
diff --git a/Engine/Source/Editor/LandscapeEditor/Private/LandscapeEditorObject.cpp b/Engine/Source/Editor/LandscapeEditor/Private/LandscapeEditorObject.cpp
--- a/Engine/Source/Editor/LandscapeEditor/Private/LandscapeEditorObject.cpp
+++ b/Engine/Source/Editor/LandscapeEditor/Private/LandscapeEditorObject.cpp
@@ -542,7 +542,7 @@ void ULandscapeEditorObject::ImportLandscapeData()
 	
 	if (HeightmapFormat)
 	{
-		FLandscapeHeightmapImportData HeightmapImportData = HeightmapFormat->Import(*ImportLandscape_HeightmapFilename, {ImportLandscape_Width, ImportLandscape_Height});
+		FLandscapeHeightmapImportData HeightmapImportData = HeightmapFormat->Import(*ImportLandscape_HeightmapFilename, NAME_None, {ImportLandscape_Width, ImportLandscape_Height});
 		ImportLandscape_HeightmapImportResult = HeightmapImportData.ResultCode;
 		ImportLandscape_HeightmapErrorMessage = HeightmapImportData.ErrorMessage;
 		ImportLandscape_Data = MoveTemp(HeightmapImportData.Data);
@@ -612,7 +612,7 @@ void ULandscapeEditorObject::RefreshImportLayersList()
 
 					if (WeightmapFormat)
 					{
-						FLandscapeWeightmapInfo WeightmapImportInfo = WeightmapFormat->Validate(*NewImportLayer.SourceFilePath, NewImportLayer.LayerName);
+						FLandscapeFileInfo WeightmapImportInfo = WeightmapFormat->Validate(*NewImportLayer.SourceFilePath, NewImportLayer.LayerName);
 						NewImportLayer.ImportResult = WeightmapImportInfo.ResultCode;
 						NewImportLayer.ErrorMessage = WeightmapImportInfo.ErrorMessage;
 
@@ -626,7 +626,7 @@ void ULandscapeEditorObject::RefreshImportLayersList()
 					else
 					{
 						NewImportLayer.ImportResult = ELandscapeImportResult::Error;
-						NewImportLayer.ErrorMessage = NSLOCTEXT("LandscapeEditor.NewLandscape", "Import_UnknownFileType", "File type not recognised");
+						NewImportLayer.ErrorMessage = NSLOCTEXT("LandscapeEditor.NewLandscape", "Import_UnknownFileType", "File type not recognized");
 					}
 				}
 			}
diff --git a/Engine/Source/Editor/LandscapeEditor/Private/LandscapeFileFormatPng.cpp b/Engine/Source/Editor/LandscapeEditor/Private/LandscapeFileFormatPng.cpp
--- a/Engine/Source/Editor/LandscapeEditor/Private/LandscapeFileFormatPng.cpp
+++ b/Engine/Source/Editor/LandscapeEditor/Private/LandscapeFileFormatPng.cpp
@@ -18,9 +18,9 @@ FLandscapeHeightmapFileFormat_Png::FLandscapeHeightmapFileFormat_Png()
 	FileTypeInfo.bSupportsExport = true;
 }
 
-FLandscapeHeightmapInfo FLandscapeHeightmapFileFormat_Png::Validate(const TCHAR* HeightmapFilename) const
+FLandscapeFileInfo FLandscapeHeightmapFileFormat_Png::Validate(const TCHAR* HeightmapFilename, FName LayerName) const
 {
-	FLandscapeHeightmapInfo Result;
+	FLandscapeFileInfo Result;
 
 	TArray64<uint8> ImportData;
 	if (!FFileHelper::LoadFileToArray(ImportData, HeightmapFilename, FILEREAD_Silent))
@@ -64,9 +64,9 @@ FLandscapeHeightmapInfo FLandscapeHeightmapFileFormat_Png::Validate(const TCHAR*
 	return Result;
 }
 
-FLandscapeHeightmapImportData FLandscapeHeightmapFileFormat_Png::Import(const TCHAR* HeightmapFilename, FLandscapeFileResolution ExpectedResolution) const
+FLandscapeImportData<uint16> FLandscapeHeightmapFileFormat_Png::Import(const TCHAR* HeightmapFilename, FName LayerName, FLandscapeFileResolution ExpectedResolution) const
 {
-	FLandscapeHeightmapImportData Result;
+	FLandscapeImportData<uint16> Result;
 
 	TArray<uint8> TempData;
 	if (!FFileHelper::LoadFileToArray(TempData, HeightmapFilename, FILEREAD_Silent))
@@ -136,7 +136,7 @@ FLandscapeHeightmapImportData FLandscapeHeightmapFileFormat_Png::Import(const TC
 	return Result;
 }
 
-void FLandscapeHeightmapFileFormat_Png::Export(const TCHAR* HeightmapFilename, TArrayView<const uint16> Data, FLandscapeFileResolution DataResolution, FVector Scale) const
+void FLandscapeHeightmapFileFormat_Png::Export(const TCHAR* HeightmapFilename, FName LayerName, TArrayView<const uint16> Data, FLandscapeFileResolution DataResolution, FVector Scale) const
 {
 	IImageWrapperModule& ImageWrapperModule = FModuleManager::LoadModuleChecked<IImageWrapperModule>("ImageWrapper");
 	TSharedPtr<IImageWrapper> ImageWrapper = ImageWrapperModule.CreateImageWrapper(EImageFormat::PNG);
@@ -157,9 +157,9 @@ FLandscapeWeightmapFileFormat_Png::FLandscapeWeightmapFileFormat_Png()
 	FileTypeInfo.bSupportsExport = true;
 }
 
-FLandscapeWeightmapInfo FLandscapeWeightmapFileFormat_Png::Validate(const TCHAR* WeightmapFilename, FName LayerName) const
+FLandscapeFileInfo FLandscapeWeightmapFileFormat_Png::Validate(const TCHAR* WeightmapFilename, FName LayerName) const
 {
-	FLandscapeWeightmapInfo Result;
+	FLandscapeFileInfo Result;
 
 	TArray64<uint8> ImportData;
 	if (!FFileHelper::LoadFileToArray(ImportData, WeightmapFilename, FILEREAD_Silent))
@@ -194,9 +194,9 @@ FLandscapeWeightmapInfo FLandscapeWeightmapFileFormat_Png::Validate(const TCHAR*
 	return Result;
 }
 
-FLandscapeWeightmapImportData FLandscapeWeightmapFileFormat_Png::Import(const TCHAR* WeightmapFilename, FName LayerName, FLandscapeFileResolution ExpectedResolution) const
+FLandscapeImportData<uint8> FLandscapeWeightmapFileFormat_Png::Import(const TCHAR* WeightmapFilename, FName LayerName, FLandscapeFileResolution ExpectedResolution) const
 {
-	FLandscapeWeightmapImportData Result;
+	FLandscapeImportData<uint8> Result;
 
 	TArray64<uint8> TempData;
 	if (!FFileHelper::LoadFileToArray(TempData, WeightmapFilename, FILEREAD_Silent))
@@ -237,7 +237,7 @@ FLandscapeWeightmapImportData FLandscapeWeightmapFileFormat_Png::Import(const TC
 	return Result;
 }
 
-void FLandscapeWeightmapFileFormat_Png::Export(const TCHAR* WeightmapFilename, FName LayerName, TArrayView<const uint8> Data, FLandscapeFileResolution DataResolution) const
+void FLandscapeWeightmapFileFormat_Png::Export(const TCHAR* WeightmapFilename, FName LayerName, TArrayView<const uint8> Data, FLandscapeFileResolution DataResolution, FVector Scale) const
 {
 	IImageWrapperModule& ImageWrapperModule = FModuleManager::LoadModuleChecked<IImageWrapperModule>("ImageWrapper");
 	TSharedPtr<IImageWrapper> ImageWrapper = ImageWrapperModule.CreateImageWrapper(EImageFormat::PNG);
diff --git a/Engine/Source/Editor/LandscapeEditor/Private/LandscapeFileFormatPng.h b/Engine/Source/Editor/LandscapeEditor/Private/LandscapeFileFormatPng.h
--- a/Engine/Source/Editor/LandscapeEditor/Private/LandscapeFileFormatPng.h
+++ b/Engine/Source/Editor/LandscapeEditor/Private/LandscapeFileFormatPng.h
@@ -20,9 +20,9 @@ public:
 		return FileTypeInfo;
 	}
 
-	virtual FLandscapeHeightmapInfo Validate(const TCHAR* HeightmapFilename) const override;
-	virtual FLandscapeHeightmapImportData Import(const TCHAR* HeightmapFilename, FLandscapeFileResolution ExpectedResolution) const override;
-	virtual void Export(const TCHAR* HeightmapFilename, TArrayView<const uint16> Data, FLandscapeFileResolution DataResolution, FVector Scale) const override;
+	virtual FLandscapeFileInfo Validate(const TCHAR* HeightmapFilename, FName LayerName) const override;
+	virtual FLandscapeImportData<uint16> Import(const TCHAR* HeightmapFilename, FName LayerName, FLandscapeFileResolution ExpectedResolution) const override;
+	virtual void Export(const TCHAR* HeightmapFilename, FName LayerName, TArrayView<const uint16> Data, FLandscapeFileResolution DataResolution, FVector Scale) const override;
 };
 
 //////////////////////////////////////////////////////////////////////////
@@ -40,7 +40,7 @@ public:
 		return FileTypeInfo;
 	}
 
-	virtual FLandscapeWeightmapInfo Validate(const TCHAR* WeightmapFilename, FName LayerName) const override;
-	virtual FLandscapeWeightmapImportData Import(const TCHAR* WeightmapFilename, FName LayerName, FLandscapeFileResolution ExpectedResolution) const override;
-	virtual void Export(const TCHAR* WeightmapFilename, FName LayerName, TArrayView<const uint8> Data, FLandscapeFileResolution DataResolution) const override;
+	virtual FLandscapeFileInfo Validate(const TCHAR* WeightmapFilename, FName LayerName) const override;
+	virtual FLandscapeImportData<uint8> Import(const TCHAR* WeightmapFilename, FName LayerName, FLandscapeFileResolution ExpectedResolution) const override;
+	virtual void Export(const TCHAR* WeightmapFilename, FName LayerName, TArrayView<const uint8> Data, FLandscapeFileResolution DataResolution, FVector Scale) const override;
 };
diff --git a/Engine/Source/Editor/LandscapeEditor/Private/LandscapeFileFormatRaw.cpp b/Engine/Source/Editor/LandscapeEditor/Private/LandscapeFileFormatRaw.cpp
--- a/Engine/Source/Editor/LandscapeEditor/Private/LandscapeFileFormatRaw.cpp
+++ b/Engine/Source/Editor/LandscapeEditor/Private/LandscapeFileFormatRaw.cpp
@@ -45,9 +45,9 @@ FLandscapeHeightmapFileFormat_Raw::FLandscapeHeightmapFileFormat_Raw()
 	FileTypeInfo.bSupportsExport = true;
 }
 
-FLandscapeHeightmapInfo FLandscapeHeightmapFileFormat_Raw::Validate(const TCHAR* HeightmapFilename) const
+FLandscapeFileInfo FLandscapeHeightmapFileFormat_Raw::Validate(const TCHAR* HeightmapFilename, FName LayerName) const
 {
-	FLandscapeHeightmapInfo Result;
+	FLandscapeFileInfo Result;
 
 	int64 ImportFileSize = IFileManager::Get().FileSize(HeightmapFilename);
 
@@ -75,9 +75,9 @@ FLandscapeHeightmapInfo FLandscapeHeightmapFileFormat_Raw::Validate(const TCHAR*
 	return Result;
 }
 
-FLandscapeHeightmapImportData FLandscapeHeightmapFileFormat_Raw::Import(const TCHAR* HeightmapFilename, FLandscapeFileResolution ExpectedResolution) const
+FLandscapeImportData<uint16> FLandscapeHeightmapFileFormat_Raw::Import(const TCHAR* HeightmapFilename, FName LayerName, FLandscapeFileResolution ExpectedResolution) const
 {
-	FLandscapeHeightmapImportData Result;
+	FLandscapeImportData<uint16> Result;
 
 	TArray<uint8> TempData;
 	if (!FFileHelper::LoadFileToArray(TempData, HeightmapFilename, FILEREAD_Silent))
@@ -100,7 +100,7 @@ FLandscapeHeightmapImportData FLandscapeHeightmapFileFormat_Raw::Import(const TC
 	return Result;
 }
 
-void FLandscapeHeightmapFileFormat_Raw::Export(const TCHAR* HeightmapFilename, TArrayView<const uint16> Data, FLandscapeFileResolution DataResolution, FVector Scale) const
+void FLandscapeHeightmapFileFormat_Raw::Export(const TCHAR* HeightmapFilename, FName LayerName, TArrayView<const uint16> Data, FLandscapeFileResolution DataResolution, FVector Scale) const
 {
 	TArray<uint8> TempData;
 	TempData.Empty(DataResolution.Width * DataResolution.Height * 2);
@@ -120,9 +120,9 @@ FLandscapeWeightmapFileFormat_Raw::FLandscapeWeightmapFileFormat_Raw()
 	FileTypeInfo.bSupportsExport = true;
 }
 
-FLandscapeWeightmapInfo FLandscapeWeightmapFileFormat_Raw::Validate(const TCHAR* WeightmapFilename, FName LayerName) const
+FLandscapeFileInfo FLandscapeWeightmapFileFormat_Raw::Validate(const TCHAR* WeightmapFilename, FName LayerName) const
 {
-	FLandscapeWeightmapInfo Result;
+	FLandscapeFileInfo Result;
 
 	int64 ImportFileSize = IFileManager::Get().FileSize(WeightmapFilename);
 
@@ -145,9 +145,9 @@ FLandscapeWeightmapInfo FLandscapeWeightmapFileFormat_Raw::Validate(const TCHAR*
 	return Result;
 }
 
-FLandscapeWeightmapImportData FLandscapeWeightmapFileFormat_Raw::Import(const TCHAR* WeightmapFilename, FName LayerName, FLandscapeFileResolution ExpectedResolution) const
+FLandscapeImportData<uint8> FLandscapeWeightmapFileFormat_Raw::Import(const TCHAR* WeightmapFilename, FName LayerName, FLandscapeFileResolution ExpectedResolution) const
 {
-	FLandscapeWeightmapImportData Result;
+	FLandscapeImportData<uint8> Result;
 
 	TArray<uint8> TempData;
 	if (!FFileHelper::LoadFileToArray(TempData, WeightmapFilename, FILEREAD_Silent))
@@ -168,7 +168,7 @@ FLandscapeWeightmapImportData FLandscapeWeightmapFileFormat_Raw::Import(const TC
 	return Result;
 }
 
-void FLandscapeWeightmapFileFormat_Raw::Export(const TCHAR* WeightmapFilename, FName LayerName, TArrayView<const uint8> Data, FLandscapeFileResolution DataResolution) const
+void FLandscapeWeightmapFileFormat_Raw::Export(const TCHAR* WeightmapFilename, FName LayerName, TArrayView<const uint8> Data, FLandscapeFileResolution DataResolution, FVector Scale) const
 {
 	FFileHelper::SaveArrayToFile(Data, WeightmapFilename);
 }
diff --git a/Engine/Source/Editor/LandscapeEditor/Private/LandscapeFileFormatRaw.h b/Engine/Source/Editor/LandscapeEditor/Private/LandscapeFileFormatRaw.h
--- a/Engine/Source/Editor/LandscapeEditor/Private/LandscapeFileFormatRaw.h
+++ b/Engine/Source/Editor/LandscapeEditor/Private/LandscapeFileFormatRaw.h
@@ -20,9 +20,9 @@ public:
 		return FileTypeInfo;
 	}
 
-	virtual FLandscapeHeightmapInfo Validate(const TCHAR* HeightmapFilename) const override;
-	virtual FLandscapeHeightmapImportData Import(const TCHAR* HeightmapFilename, FLandscapeFileResolution ExpectedResolution) const override;
-	virtual void Export(const TCHAR* HeightmapFilename, TArrayView<const uint16> Data, FLandscapeFileResolution DataResolution, FVector Scale) const override;
+	virtual FLandscapeFileInfo Validate(const TCHAR* HeightmapFilename, FName LayerName) const override;
+	virtual FLandscapeImportData<uint16> Import(const TCHAR* HeightmapFilename, FName LayerName, FLandscapeFileResolution ExpectedResolution) const override;
+	virtual void Export(const TCHAR* HeightmapFilename, FName LayerName, TArrayView<const uint16> Data, FLandscapeFileResolution DataResolution, FVector Scale) const override;
 };
 
 //////////////////////////////////////////////////////////////////////////
@@ -40,7 +40,7 @@ public:
 		return FileTypeInfo;
 	}
 
-	virtual FLandscapeWeightmapInfo Validate(const TCHAR* WeightmapFilename, FName LayerName) const override;
-	virtual FLandscapeWeightmapImportData Import(const TCHAR* WeightmapFilename, FName LayerName, FLandscapeFileResolution ExpectedResolution) const override;
-	virtual void Export(const TCHAR* WeightmapFilename, FName LayerName, TArrayView<const uint8> Data, FLandscapeFileResolution DataResolution) const override;
+	virtual FLandscapeFileInfo Validate(const TCHAR* WeightmapFilename, FName LayerName) const override;
+	virtual FLandscapeImportData<uint8> Import(const TCHAR* WeightmapFilename, FName LayerName, FLandscapeFileResolution ExpectedResolution) const override;
+	virtual void Export(const TCHAR* WeightmapFilename, FName Layername, TArrayView<const uint8> Data, FLandscapeFileResolution DataResolution, FVector Scale) const override;
 };
diff --git a/Engine/Source/Editor/LandscapeEditor/Private/NewLandscapeUtils.cpp b/Engine/Source/Editor/LandscapeEditor/Private/NewLandscapeUtils.cpp
--- a/Engine/Source/Editor/LandscapeEditor/Private/NewLandscapeUtils.cpp
+++ b/Engine/Source/Editor/LandscapeEditor/Private/NewLandscapeUtils.cpp
@@ -119,7 +119,7 @@ void FNewLandscapeUtils::ImportLandscapeData( ULandscapeEditorObject* UISettings
 
 		if(HeightmapFormat)
 		{
-			FLandscapeHeightmapInfo HeightmapImportInfo = HeightmapFormat->Validate(*UISettings->ImportLandscape_HeightmapFilename);
+			FLandscapeFileInfo HeightmapImportInfo = HeightmapFormat->Validate(*UISettings->ImportLandscape_HeightmapFilename, NAME_None);
 			UISettings->ImportLandscape_HeightmapImportResult = HeightmapImportInfo.ResultCode;
 			UISettings->ImportLandscape_HeightmapErrorMessage = HeightmapImportInfo.ErrorMessage;
 			ImportResolutions = MoveTemp(HeightmapImportInfo.PossibleResolutions);
@@ -132,7 +132,7 @@ void FNewLandscapeUtils::ImportLandscapeData( ULandscapeEditorObject* UISettings
 		else
 		{
 			UISettings->ImportLandscape_HeightmapImportResult = ELandscapeImportResult::Error;
-			UISettings->ImportLandscape_HeightmapErrorMessage = LOCTEXT("Import_UnknownFileType", "File type not recognised");
+			UISettings->ImportLandscape_HeightmapErrorMessage = LOCTEXT("Import_UnknownFileType", "File type not recognized");
 		}
 	}
 
@@ -218,7 +218,7 @@ TOptional< TArray< FLandscapeImportLayerInfo > > FNewLandscapeUtils::CreateImpor
 				else
 				{
 					UIImportLayer.ImportResult = ELandscapeImportResult::Error;
-					UIImportLayer.ErrorMessage = LOCTEXT("Import_UnknownFileType", "File type not recognised");
+					UIImportLayer.ErrorMessage = LOCTEXT("Import_UnknownFileType", "File type not recognized");
 				}
 
 				if(UIImportLayer.ImportResult == ELandscapeImportResult::Error)
diff --git a/Engine/Source/Editor/LandscapeEditor/Public/LandscapeEditorModule.h b/Engine/Source/Editor/LandscapeEditor/Public/LandscapeEditorModule.h
--- a/Engine/Source/Editor/LandscapeEditor/Public/LandscapeEditorModule.h
+++ b/Engine/Source/Editor/LandscapeEditor/Public/LandscapeEditorModule.h
@@ -5,9 +5,8 @@
 #include "CoreMinimal.h"
 #include "Modules/ModuleInterface.h"
 #include "Templates/SharedPointer.h"
+#include "LandscapeFileFormatInterface.h"
 
-class ILandscapeHeightmapFileFormat;
-class ILandscapeWeightmapFileFormat;
 class FUICommandList;
 
 /**
diff --git a/Engine/Source/Editor/LandscapeEditor/Public/LandscapeFileFormatInterface.h b/Engine/Source/Editor/LandscapeEditor/Public/LandscapeFileFormatInterface.h
--- a/Engine/Source/Editor/LandscapeEditor/Public/LandscapeFileFormatInterface.h
+++ b/Engine/Source/Editor/LandscapeEditor/Public/LandscapeFileFormatInterface.h
@@ -46,62 +46,41 @@ FORCEINLINE bool operator!=(const FLandscapeFileResolution& Lhs, const FLandscap
 	return !(Lhs == Rhs);
 }
 
-struct FLandscapeHeightmapInfo
+struct FLandscapeFileInfo
 {
-	// Whether the heightmap is usable or has errors/warnings
+	// Whether the the file is usable or has errors/warnings
 	ELandscapeImportResult ResultCode = ELandscapeImportResult::Success;
 
 	// Message to show as the warning/error result
 	FText ErrorMessage;
 
 	// Normally contains a single resolution, but .raw is awful
-	TArray<FLandscapeFileResolution, TInlineAllocator<1>> PossibleResolutions;
+	TArray<FLandscapeFileResolution> PossibleResolutions;
 
 	// The inherent scale of the data format, if it has one, in centimeters
 	// The default for data with no inherent scale is 100,100,0.78125 (100.0/128, shown as 100 in the editor UI)
 	TOptional<FVector> DataScale;
 };
 
-struct FLandscapeWeightmapInfo
+template< class T >
+struct FLandscapeImportData
 {
-	// Whether the weightmap is usable or has errors/warnings
+	// Whether the import data is usable or has errors/warnings
 	ELandscapeImportResult ResultCode = ELandscapeImportResult::Success;
 
 	// Message to show as the warning/error result
 	FText ErrorMessage;
 
-	// Normally contains a single resolution, but .raw is awful
-	TArray<FLandscapeFileResolution, TInlineAllocator<1>> PossibleResolutions;
+	// The Data
+	TArray<T> Data;
 };
 
-struct FLandscapeHeightmapImportData
-{
-	// Whether the heightmap is usable or has errors/warnings
-	ELandscapeImportResult ResultCode = ELandscapeImportResult::Success;
-
-	// Message to show as the warning/error result
-	FText ErrorMessage;
-
-	// The height data!
-	// A value of 32768 is the 0 level (e.g. sea level), lower values are below and higher values are above
-	TArray<uint16> Data;
-};
-
-struct FLandscapeWeightmapImportData
-{
-	// Whether the weightmap is usable or has errors/warnings
-	ELandscapeImportResult ResultCode = ELandscapeImportResult::Success;
-
-	// Message to show as the warning/error result
-	FText ErrorMessage;
-
-	// The weight data!
-	// 255 is fully painted and 0 is unpainted
-	TArray<uint8> Data;
-};
+using FLandscapeHeightmapImportData = FLandscapeImportData<uint16>;
+using FLandscapeWeightmapImportData = FLandscapeImportData<uint8>;
 
 // Interface
-class ILandscapeHeightmapFileFormat
+template< class T >
+class ILandscapeFileFormat
 {
 public:
 	/** Gets info about this format
@@ -112,74 +91,38 @@ public:
 	/** Validate a file for Import
 	 * Gives the file format the opportunity to reject a file or return warnings
 	 * as well as return information about the file for the import UI (e.g. resolution and scale)
-	 * @param HeightmapFilename path to the file to validate for import
+	 * @param Filename path to the file to validate for import
+	 * @param LayerName name of layer that is being imported (in case of Heightmap this will be NAME_None)
 	 * @return information about the file and (optional) error message
 	 */
-	virtual FLandscapeHeightmapInfo Validate(const TCHAR* HeightmapFilename) const = 0;
+	virtual FLandscapeFileInfo Validate(const TCHAR* Filename, FName LayerName) const = 0;
 
 	/** Import a file
-	 * @param HeightmapFilename path to the file to import
+	 * @param Filename path to the file to import
+	 * @param LayerName name of layer being imported (in case of Heightmap this will be NAME_None)
 	 * @param ExpectedResolution resolution selected in the import UI (mostly for the benefit of .raw)
 	 * @return imported data and (optional) error message
 	 */
-	virtual FLandscapeHeightmapImportData Import(const TCHAR* HeightmapFilename, FLandscapeFileResolution ExpectedResolution) const = 0;
+	virtual FLandscapeImportData<T> Import(const TCHAR* Filename, FName LayerName, FLandscapeFileResolution ExpectedResolution) const = 0;
 
 	/** Export a file (if supported)
-	 * @param HeightmapFilename path to the file to export to
+	 * @param Filename path to the file to export to
+	 * @param LayerName name of layer being exported (in case of Heightmap this will be NAME_None)
 	 * @param Data raw data to export
 	 * @param DataResolution resolution of Data
 	 * @param Scale scale of the landscape data, in centimeters
 	 */
-	virtual void Export(const TCHAR* HeightmapFilename, TArrayView<const uint16> Data, FLandscapeFileResolution DataResolution, FVector Scale) const
+	virtual void Export(const TCHAR* Filename, FName LayerName, TArrayView<const T> Data, FLandscapeFileResolution DataResolution, FVector Scale) const
 	{
-		checkf(0, TEXT("File type hasn't implemented support for heightmap export - %s"), *FPaths::GetExtension(HeightmapFilename, true));
+		checkf(0, TEXT("File type hasn't implemented support for export - %s"), *FPaths::GetExtension(Filename, true));
 	}
 
 	/**
 	 * Note: Even though this is an interface class we need a virtual destructor as derived objects are deleted via a pointer to this interface
 	 */
-	virtual ~ILandscapeHeightmapFileFormat() {}
+	virtual ~ILandscapeFileFormat() {}
 };
 
-class ILandscapeWeightmapFileFormat
-{
-public:
-	/** Gets info about this format
-	 * @return information about the file types supported by this file format plugin
-	 */
-	virtual const FLandscapeFileTypeInfo& GetInfo() const = 0;
-
-	/** Validate a file for Import
-	 * Gives the file format the opportunity to reject a file or return warnings
-	 * as well as return information about the file for the import UI (e.g. resolution and scale)
-	 * @param HeightmapFilename path to the file to validate for import
-	 * @param LayerName name of layer that is being imported
-	 * @return information about the file and (optional) error message
-	 */
-	virtual FLandscapeWeightmapInfo Validate(const TCHAR* WeightmapFilename, FName LayerName) const = 0;
-
-	/** Import a file
-	 * @param HeightmapFilename path to the file to import
-	 * @param LayerName name of layer that is being imported
-	 * @param ExpectedResolution resolution selected in the import UI (mostly for the benefit of .raw)
-	 * @return imported data and (optional) error message
-	 */
-	virtual FLandscapeWeightmapImportData Import(const TCHAR* WeightmapFilename, FName LayerName, FLandscapeFileResolution ExpectedResolution) const = 0;
-
-	/** Export a file (if supported)
-	 * @param HeightmapFilename path to the file to export to
-	 * @param LayerName name of layer that is being exported
-	 * @param Data raw data to export
-	 * @param DataResolution resolution of Data
-	 * @param Scale scale of the landscape data, in centimeters
-	 */
-	virtual void Export(const TCHAR* WeightmapFilename, FName LayerName, TArrayView<const uint8> Data, FLandscapeFileResolution DataResolution) const
-	{
-		checkf(0, TEXT("File type hasn't implemented support for weightmap export - %s"), *FPaths::GetExtension(WeightmapFilename, true));
-	}
+using ILandscapeHeightmapFileFormat = ILandscapeFileFormat<uint16>;
+using ILandscapeWeightmapFileFormat = ILandscapeFileFormat<uint8>;
 
-	/**
- 	 * Note: Even though this is an interface class we need a virtual destructor as derived objects are deleted via a pointer to this interface
-	 */
-	virtual ~ILandscapeWeightmapFileFormat() {}
-};
diff --git a/Engine/Source/Editor/WorldBrowser/Private/Tiles/STiledLandscapeImportDlg.cpp b/Engine/Source/Editor/WorldBrowser/Private/Tiles/STiledLandscapeImportDlg.cpp
--- a/Engine/Source/Editor/WorldBrowser/Private/Tiles/STiledLandscapeImportDlg.cpp
+++ b/Engine/Source/Editor/WorldBrowser/Private/Tiles/STiledLandscapeImportDlg.cpp
@@ -537,7 +537,7 @@ FReply STiledLandscapeImportDlg::OnClickedSelectHeightmapTiles()
 						break;
 					}
 
-					FLandscapeHeightmapInfo HeightmapInfo = HeightmapFormat->Validate(*Filename);
+					FLandscapeFileInfo HeightmapInfo = HeightmapFormat->Validate(*Filename, NAME_None);
 					if (HeightmapInfo.ResultCode != ELandscapeImportResult::Success)
 					{
 						StatusMessage = HeightmapInfo.ErrorMessage;
diff --git a/Engine/Source/Editor/WorldBrowser/Private/Tiles/WorldTileCollectionModel.cpp b/Engine/Source/Editor/WorldBrowser/Private/Tiles/WorldTileCollectionModel.cpp
--- a/Engine/Source/Editor/WorldBrowser/Private/Tiles/WorldTileCollectionModel.cpp
+++ b/Engine/Source/Editor/WorldBrowser/Private/Tiles/WorldTileCollectionModel.cpp
@@ -1483,7 +1483,7 @@ static bool ReadHeightmapFile(TArray<uint16>& Result, const FString& Filename, i
 	ILandscapeEditorModule& LandscapeEditorModule = FModuleManager::GetModuleChecked<ILandscapeEditorModule>("LandscapeEditor");
 	const ILandscapeHeightmapFileFormat* HeightmapFormat = LandscapeEditorModule.GetHeightmapFormatByExtension(*FPaths::GetExtension(Filename, true));
 
-	FLandscapeHeightmapImportData ImportData = HeightmapFormat->Import(*Filename, {(uint32)ExpectedWidth, (uint32)ExpectedHeight});
+	FLandscapeHeightmapImportData ImportData = HeightmapFormat->Import(*Filename, NAME_None, {(uint32)ExpectedWidth, (uint32)ExpectedHeight});
 	if (ImportData.ResultCode != ELandscapeImportResult::Error)
 	{
 		Result = MoveTemp(ImportData.Data);
diff --git a/Engine/Source/Editor/WorldBrowser/Private/Tiles/WorldTileModel.cpp b/Engine/Source/Editor/WorldBrowser/Private/Tiles/WorldTileModel.cpp
--- a/Engine/Source/Editor/WorldBrowser/Private/Tiles/WorldTileModel.cpp
+++ b/Engine/Source/Editor/WorldBrowser/Private/Tiles/WorldTileModel.cpp
@@ -350,7 +350,7 @@ bool FWorldTileModel::CanReimportHeightmap() const
 		const FString TargetExtension = FPaths::GetExtension(GetLandscape()->ReimportHeightmapFilePath, true);
 		const ILandscapeHeightmapFileFormat* HeightmapFormat = LandscapeEditorModule.GetHeightmapFormatByExtension(*TargetExtension);
 
-		FLandscapeHeightmapInfo HeightmapInfo = HeightmapFormat->Validate(*GetLandscape()->ReimportHeightmapFilePath);
+		FLandscapeFileInfo HeightmapInfo = HeightmapFormat->Validate(*GetLandscape()->ReimportHeightmapFilePath, NAME_None);
 		if (HeightmapInfo.ResultCode != ELandscapeImportResult::Error)
 		{
 			FIntRect ComponentsRect = GetLandscape()->GetBoundingRect();
diff --git a/Engine/Source/Runtime/Landscape/Private/LandscapeEdit.cpp b/Engine/Source/Runtime/Landscape/Private/LandscapeEdit.cpp
--- a/Engine/Source/Runtime/Landscape/Private/LandscapeEdit.cpp
+++ b/Engine/Source/Runtime/Landscape/Private/LandscapeEdit.cpp
@@ -3584,7 +3584,7 @@ void ULandscapeInfo::ExportHeightmap(const FString& Filename)
 	const ILandscapeHeightmapFileFormat* HeightmapFormat = LandscapeEditorModule.GetHeightmapFormatByExtension(*FPaths::GetExtension(Filename, true));
 	if (HeightmapFormat)
 	{
-		HeightmapFormat->Export(*Filename, HeightData, {(uint32)(MaxX - MinX + 1), (uint32)(MaxY - MinY + 1)}, DrawScale * FVector(1, 1, LANDSCAPE_ZSCALE));
+		HeightmapFormat->Export(*Filename, NAME_None, HeightData, {(uint32)(MaxX - MinX + 1), (uint32)(MaxY - MinY + 1)}, DrawScale * FVector(1, 1, LANDSCAPE_ZSCALE));
 	}
 
 	GWarn->EndSlowTask();
@@ -3617,7 +3617,7 @@ void ULandscapeInfo::ExportLayer(ULandscapeLayerInfoObject* LayerInfo, const FSt
 	const ILandscapeWeightmapFileFormat* WeightmapFormat = LandscapeEditorModule.GetWeightmapFormatByExtension(*FPaths::GetExtension(Filename, true));
 	if (WeightmapFormat)
 	{
-		WeightmapFormat->Export(*Filename, LayerInfo->LayerName, WeightData, {(uint32)(MaxX - MinX + 1), (uint32)(MaxY - MinY + 1)});
+		WeightmapFormat->Export(*Filename, LayerInfo->LayerName, WeightData, {(uint32)(MaxX - MinX + 1), (uint32)(MaxY - MinY + 1)}, DrawScale * FVector(1, 1, LANDSCAPE_ZSCALE));
 	}
 
 	GWarn->EndSlowTask();
