BuildFix: HoudiniEngine

diff --git a/Engine/Source/Editor/LandscapeEditor/Public/LandscapeFileFormatInterface.h b/Engine/Source/Editor/LandscapeEditor/Public/LandscapeFileFormatInterface.h
--- a/Engine/Source/Editor/LandscapeEditor/Public/LandscapeFileFormatInterface.h
+++ b/Engine/Source/Editor/LandscapeEditor/Public/LandscapeFileFormatInterface.h
@@ -78,6 +78,9 @@ struct FLandscapeImportData
 using FLandscapeHeightmapImportData = FLandscapeImportData<uint16>;
 using FLandscapeWeightmapImportData = FLandscapeImportData<uint8>;
 
+using FLandscapeHeightmapInfo = FLandscapeFileInfo;
+using FLandscapeWeightmapInfo = FLandscapeFileInfo;
+
 // Interface
 template< class T >
 class ILandscapeFileFormat
@@ -97,6 +100,8 @@ public:
 	 */
 	virtual FLandscapeFileInfo Validate(const TCHAR* Filename, FName LayerName) const = 0;
 
+	virtual FLandscapeFileInfo Validate(const TCHAR* Filename) const { return Validate(Filename, NAME_None); }
+
 	/** Import a file
 	 * @param Filename path to the file to import
 	 * @param LayerName name of layer being imported (in case of Heightmap this will be NAME_None)
@@ -105,6 +110,8 @@ public:
 	 */
 	virtual FLandscapeImportData<T> Import(const TCHAR* Filename, FName LayerName, FLandscapeFileResolution ExpectedResolution) const = 0;
 
+	virtual FLandscapeImportData<T> Import(const TCHAR* Filename, FLandscapeFileResolution ExpectedResolution) const { return Import(Filename, NAME_None, ExpectedResolution); }
+
 	/** Export a file (if supported)
 	 * @param Filename path to the file to export to
 	 * @param LayerName name of layer being exported (in case of Heightmap this will be NAME_None)
@@ -117,6 +124,11 @@ public:
 		checkf(0, TEXT("File type hasn't implemented support for export - %s"), *FPaths::GetExtension(Filename, true));
 	}
 
+	virtual void Export(const TCHAR* Filename, TArrayView<const T> Data, FLandscapeFileResolution DataResolution, FVector Scale)
+	{
+		Export(Filename, NAME_None, Data, DataResolution, Scale);
+	}
+
 	/**
 	 * Note: Even though this is an interface class we need a virtual destructor as derived objects are deleted via a pointer to this interface
 	 */
