Landscape:
- Import code refactor to support multiple import files and unify import code
- WorldPartition landscape import

diff --git a/Engine/Plugins/Enterprise/DatasmithImporter/Source/DatasmithImporter/Private/DatasmithLandscapeImporter.cpp b/Engine/Plugins/Enterprise/DatasmithImporter/Source/DatasmithImporter/Private/DatasmithLandscapeImporter.cpp
--- a/Engine/Plugins/Enterprise/DatasmithImporter/Source/DatasmithImporter/Private/DatasmithLandscapeImporter.cpp
+++ b/Engine/Plugins/Enterprise/DatasmithImporter/Source/DatasmithImporter/Private/DatasmithLandscapeImporter.cpp
@@ -13,7 +13,7 @@
 #include "LandscapeEditorObject.h"
 #include "LandscapeEditorUtils.h"
 #include "LandscapeFileFormatInterface.h"
-#include "NewLandscapeUtils.h"
+#include "LandscapeImportHelper.h"
 
 #include "Modules/ModuleManager.h"
 #include "Utility/DatasmithImporterUtils.h"
@@ -22,54 +22,66 @@
 
 AActor* FDatasmithLandscapeImporter::ImportLandscapeActor( const TSharedRef< IDatasmithLandscapeElement >& LandscapeActorElement, FDatasmithImportContext& ImportContext, EDatasmithImportActorPolicy ImportActorPolicy )
 {
-	TStrongObjectPtr< ULandscapeEditorObject > LandscapeEditorObject( NewObject< ULandscapeEditorObject >() );
-
-	LandscapeEditorObject->ImportLandscape_HeightmapFilename = LandscapeActorElement->GetHeightmap();
-	LandscapeEditorObject->NewLandscape_Scale = LandscapeActorElement->GetScale();
-
-	TArray< FLandscapeFileResolution > ImportResolutions;
-	FNewLandscapeUtils::ImportLandscapeData( LandscapeEditorObject.Get(), ImportResolutions );
-
-	const int32 ComponentCountX = LandscapeEditorObject->NewLandscape_ComponentCount.X;
-	const int32 ComponentCountY = LandscapeEditorObject->NewLandscape_ComponentCount.Y;
-	const int32 QuadsPerComponent = LandscapeEditorObject->NewLandscape_SectionsPerComponent * LandscapeEditorObject->NewLandscape_QuadsPerSection;
-	const int32 SizeX = ComponentCountX * QuadsPerComponent + 1;
-	const int32 SizeY = ComponentCountY * QuadsPerComponent + 1;
-
-	TOptional< TArray< FLandscapeImportLayerInfo > > ImportLayers = FNewLandscapeUtils::CreateImportLayersInfo( LandscapeEditorObject.Get(), ENewLandscapePreviewMode::ImportLandscape );
-
-	if ( !ImportLayers )
+	const bool bSingleFile = true;
+	FLandscapeImportDescriptor OutImportDescriptor;
+	OutImportDescriptor.Scale = LandscapeActorElement->GetScale();
+	FText OutMessage;
+	ELandscapeImportResult ImportResult = FLandscapeImportHelper::GetHeightmapImportDescriptor(LandscapeActorElement->GetHeightmap(), bSingleFile, OutImportDescriptor, OutMessage);
+	if (ImportResult == ELandscapeImportResult::Error)
 	{
 		return nullptr;
 	}
+	// 0 could be used but this keeps the previous behavior.
+	int32 DescriptorIndex = OutImportDescriptor.FileResolutions.Num() / 2;
+	
+	ULandscapeEditorObject* DefaultValueObject = ULandscapeEditorObject::StaticClass()->GetDefaultObject<ULandscapeEditorObject>();
+	check(DefaultValueObject);
 
-	TArray<uint16> HeightData = FNewLandscapeUtils::ComputeHeightData( LandscapeEditorObject.Get(), ImportLayers.GetValue(), ENewLandscapePreviewMode::ImportLandscape );
+	int32 OutQuadsPerSection = DefaultValueObject->NewLandscape_QuadsPerSection;
+	int32 OutSectionsPerComponent = DefaultValueObject->NewLandscape_SectionsPerComponent;
+	FIntPoint OutComponentCount = DefaultValueObject->NewLandscape_ComponentCount;
 
-	const ELandscapeImportAlphamapType ImportLandscape_AlphamapType = ELandscapeImportAlphamapType::Additive;
+	FLandscapeImportHelper::ChooseBestComponentSizeForImport(OutImportDescriptor.ImportResolutions[DescriptorIndex].Width, OutImportDescriptor.ImportResolutions[DescriptorIndex].Height, OutQuadsPerSection, OutSectionsPerComponent, OutComponentCount);
 
+	TArray<uint16> ImportData;
+	ImportResult = FLandscapeImportHelper::GetHeightmapImportData(OutImportDescriptor, DescriptorIndex, ImportData, OutMessage);
+	if (ImportResult == ELandscapeImportResult::Error)
+	{
+		return nullptr;
+	}
+		
+	const int32 QuadsPerComponent = OutSectionsPerComponent * OutQuadsPerSection;
+	const int32 SizeX = OutComponentCount.X * QuadsPerComponent + 1;
+	const int32 SizeY = OutComponentCount.Y * QuadsPerComponent + 1;
+
+	TArray<uint16> FinalHeightData;
+	FLandscapeImportHelper::ExpandHeightmapImportData(ImportData, FinalHeightData, OutImportDescriptor.ImportResolutions[DescriptorIndex], FLandscapeImportResolution(SizeX, SizeY));
+		
 	const FVector Offset = FTransform( LandscapeActorElement->GetRotation(), FVector::ZeroVector, 
-		LandscapeActorElement->GetScale() ).TransformVector( FVector( -ComponentCountX * QuadsPerComponent / 2, -ComponentCountY * QuadsPerComponent / 2, 0 ) );
+		LandscapeActorElement->GetScale() ).TransformVector( FVector( -OutComponentCount.X * QuadsPerComponent / 2, -OutComponentCount.Y * QuadsPerComponent / 2, 0 ) );
 
 	FVector OriginalTranslation = LandscapeActorElement->GetTranslation();
 	FVector OriginalScale = LandscapeActorElement->GetScale();
 
 	LandscapeActorElement->SetTranslation( LandscapeActorElement->GetTranslation() + Offset );
 	LandscapeActorElement->SetScale( LandscapeActorElement->GetScale() );
-
+		
 	ALandscape* Landscape = Cast< ALandscape >(FDatasmithActorImporter::ImportActor( ALandscape::StaticClass(), LandscapeActorElement, ImportContext, ImportActorPolicy,
-		[ &HeightData, &ImportLayers, SizeX, SizeY, LandscapeEditorObject, ImportLandscape_AlphamapType, ActorScale = LandscapeActorElement->GetScale() ]( AActor* NewActor )
+		[ &FinalHeightData, SizeX, SizeY, OutSectionsPerComponent, OutQuadsPerSection, ActorScale = OutImportDescriptor.Scale ]( AActor* NewActor )
 	{
 		check( Cast< ALandscape >( NewActor ) );
 
 		NewActor->SetActorRelativeScale3D( ActorScale );
 
 		TMap<FGuid, TArray<uint16>> HeightmapDataPerLayers;
-		HeightmapDataPerLayers.Add(FGuid(), HeightData);
-		TMap<FGuid, TArray<FLandscapeImportLayerInfo>> MaterialLayerDataPerLayer;
-		MaterialLayerDataPerLayer.Add(FGuid(), ImportLayers.GetValue());
+		HeightmapDataPerLayers.Add(FGuid(), MoveTemp(FinalHeightData));
 
-		Cast< ALandscape >( NewActor )->Import( FGuid::NewGuid(), 0, 0, SizeX - 1, SizeY - 1, LandscapeEditorObject->NewLandscape_SectionsPerComponent, LandscapeEditorObject->NewLandscape_QuadsPerSection,
-			HeightmapDataPerLayers, nullptr, MaterialLayerDataPerLayer, ImportLandscape_AlphamapType );
+		// The is not Material Layer Import Data
+		TMap<FGuid, TArray<FLandscapeImportLayerInfo>> MaterialLayerDataPerLayer;
+		MaterialLayerDataPerLayer.Add(FGuid(), TArray<FLandscapeImportLayerInfo>());
+		
+		Cast< ALandscape >( NewActor )->Import( FGuid::NewGuid(), 0, 0, SizeX - 1, SizeY - 1, OutSectionsPerComponent, OutQuadsPerSection,
+			HeightmapDataPerLayers, nullptr, MaterialLayerDataPerLayer, ELandscapeImportAlphamapType::Additive);
 	} ) );
 
 	LandscapeActorElement->SetTranslation( OriginalTranslation );
@@ -92,35 +104,13 @@ AActor* FDatasmithLandscapeImporter::ImportLandscapeActor( const TSharedRef< IDa
 	
 	LandscapeTemplate->Apply( Landscape );
 
-	Landscape->ReimportHeightmapFilePath = LandscapeEditorObject->ImportLandscape_HeightmapFilename;
+	Landscape->ReimportHeightmapFilePath = LandscapeActorElement->GetHeightmap();
 
 	ULandscapeInfo* LandscapeInfo = Landscape->CreateLandscapeInfo();
 	LandscapeInfo->UpdateLayerInfoMap(Landscape);
-
-	// Import doesn't fill in the LayerInfo for layers with no data, do that now
-	const TArray< FLandscapeImportLayer >& ImportLandscapeLayersList = LandscapeEditorObject->ImportLandscape_Layers;
-	for(int32 i = 0; i < ImportLandscapeLayersList.Num(); i++)
-	{
-		if(ImportLandscapeLayersList[i].LayerInfo != nullptr)
-		{
-			Landscape->EditorLayerSettings.Add(FLandscapeEditorLayerSettings(ImportLandscapeLayersList[i].LayerInfo, ImportLandscapeLayersList[i].SourceFilePath));
-
-			int32 LayerInfoIndex = LandscapeInfo->GetLayerInfoIndex(ImportLandscapeLayersList[i].LayerName);
-			if(ensure(LayerInfoIndex != INDEX_NONE))
-			{
-				FLandscapeInfoLayerSettings& LayerSettings = LandscapeInfo->Layers[LayerInfoIndex];
-				LayerSettings.LayerInfoObj = ImportLandscapeLayersList[i].LayerInfo;
-			}
-		}
-	}
-
+		
 	Landscape->RegisterAllComponents();
 
-	// Need to explicitly call PostEditChange on the LandscapeMaterial property or the landscape proxy won't update its material
-	FPropertyChangedEvent MaterialPropertyChangedEvent( FindFieldChecked< FProperty >( Landscape->GetClass(), FName("LandscapeMaterial") ) );
-	Landscape->PostEditChangeProperty( MaterialPropertyChangedEvent );
-	Landscape->PostEditChange();
-
 	return Landscape;
 }
 
diff --git a/Engine/Source/Editor/LandscapeEditor/Private/LandscapeEdMode.cpp b/Engine/Source/Editor/LandscapeEditor/Private/LandscapeEdMode.cpp
--- a/Engine/Source/Editor/LandscapeEditor/Private/LandscapeEdMode.cpp
+++ b/Engine/Source/Editor/LandscapeEditor/Private/LandscapeEdMode.cpp
@@ -30,6 +30,7 @@
 #include "EditorWorldExtension.h"
 #include "LandscapeEdModeTools.h"
 #include "LandscapeInfoMap.h"
+#include "LandscapeImportHelper.h"
 
 //Slate dependencies
 #include "Misc/FeedbackContext.h"
@@ -3288,248 +3289,133 @@ void FEdModeLandscape::ReimportData(const FLandscapeTargetListInfo& TargetInfo)
 	}
 }
 
-void FEdModeLandscape::ImportData(const FLandscapeTargetListInfo& TargetInfo, const FString& Filename)
+template<class T>
+void ImportDataInternal(const FLandscapeTargetListInfo& TargetInfo, const FString& Filename, bool bSingleFile, TFunctionRef<void(int32, int32, int32, int32, const TArray<T>&)> SetDataFunc)
 {
 	ULandscapeInfo* LandscapeInfo = TargetInfo.LandscapeInfo.Get();
 	int32 MinX, MinY, MaxX, MaxY;
 	if (LandscapeInfo && LandscapeInfo->GetLandscapeExtent(MinX, MinY, MaxX, MaxY))
 	{
-		const FLandscapeFileResolution LandscapeResolution = {(uint32)(1 + MaxX - MinX), (uint32)(1 + MaxY - MinY)};
-
-		ILandscapeEditorModule& LandscapeEditorModule = FModuleManager::GetModuleChecked<ILandscapeEditorModule>("LandscapeEditor");
+		const FLandscapeImportResolution LandscapeResolution = { (uint32)(1 + MaxX - MinX), (uint32)(1 + MaxY - MinY) };
+		FLandscapeImportDescriptor OutImportDescriptor;
+		FText OutMessage;
 
-		if (TargetInfo.TargetType == ELandscapeToolTargetType::Heightmap)
+		ELandscapeImportResult ImportResult = FLandscapeImportHelper::GetImportDescriptor<T>(Filename, bSingleFile, TargetInfo.LayerName, OutImportDescriptor, OutMessage);
+		if (ImportResult == ELandscapeImportResult::Error)
 		{
-			const ILandscapeHeightmapFileFormat* HeightmapFormat = LandscapeEditorModule.GetHeightmapFormatByExtension(*FPaths::GetExtension(Filename, true));
-
-			if (!HeightmapFormat)
-			{
-				FMessageDialog::Open(EAppMsgType::Ok, NSLOCTEXT("LandscapeEditor.NewLandscape", "Import_UnknownFileType", "File type not recognised"));
-
-				return;
-			}
-
-			FLandscapeFileResolution ImportResolution = {0, 0};
+			FMessageDialog::Open(EAppMsgType::Ok, OutMessage);
+			return;
+		}
 
-			const FLandscapeFileInfo HeightmapInfo = HeightmapFormat->Validate(*Filename, NAME_None);
-
-			// display error message if there is one, and abort the import
-			if (HeightmapInfo.ResultCode == ELandscapeImportResult::Error)
-			{
-				FMessageDialog::Open(EAppMsgType::Ok, HeightmapInfo.ErrorMessage);
-
-				return;
-			}
+		check(OutImportDescriptor.ImportResolutions.Num() > 0);
 
+		int32 DescriptorIndex = 0;
+		// If there is more than one possible res it needs to match perfectly (this constraint could probably removed)
+		if (OutImportDescriptor.ImportResolutions.Num() > 1)
+		{
+			DescriptorIndex = OutImportDescriptor.FindDescriptorIndex(LandscapeResolution.Width, LandscapeResolution.Height);
 			// if the file is a raw format with multiple possibly resolutions, only attempt import if one matches the current landscape
-			if (HeightmapInfo.PossibleResolutions.Num() > 1)
+			if (DescriptorIndex == INDEX_NONE)
 			{
-				if (!HeightmapInfo.PossibleResolutions.Contains(LandscapeResolution))
-				{
-					FFormatNamedArguments Args;
-					Args.Add(TEXT("LandscapeSizeX"), LandscapeResolution.Width);
-					Args.Add(TEXT("LandscapeSizeY"), LandscapeResolution.Height);
-
-					FMessageDialog::Open(EAppMsgType::Ok,
-						FText::Format(NSLOCTEXT("LandscapeEditor.NewLandscape", "Import_HeightmapSizeMismatchRaw", "The heightmap file does not match the current Landscape extent ({LandscapeSizeX}\u00D7{LandscapeSizeY}), and its exact resolution could not be determined"), Args));
-
-					return;
-				}
-				else
-				{
-					ImportResolution = LandscapeResolution;
-				}
-			}
+				FFormatNamedArguments Args;
+				Args.Add(TEXT("LandscapeSizeX"), LandscapeResolution.Width);
+				Args.Add(TEXT("LandscapeSizeY"), LandscapeResolution.Height);
 
-			// display warning message if there is one and allow user to cancel
-			if (HeightmapInfo.ResultCode == ELandscapeImportResult::Warning)
-			{
-				auto Result = FMessageDialog::Open(EAppMsgType::OkCancel, HeightmapInfo.ErrorMessage);
-
-				if (Result != EAppReturnType::Ok)
-				{
-					return;
-				}
-			}
-
-			// if the file is a format with resolution information, warn the user if the resolution doesn't match the current landscape
-			// unlike for raw this is only a warning as we can pad/clip the data if we know what resolution it is
-			if (HeightmapInfo.PossibleResolutions.Num() == 1)
-			{
-				ImportResolution = HeightmapInfo.PossibleResolutions[0];
-				if (ImportResolution != LandscapeResolution)
-				{
-					FFormatNamedArguments Args;
-					Args.Add(TEXT("FileSizeX"), ImportResolution.Width);
-					Args.Add(TEXT("FileSizeY"), ImportResolution.Height);
-					Args.Add(TEXT("LandscapeSizeX"), LandscapeResolution.Width);
-					Args.Add(TEXT("LandscapeSizeY"), LandscapeResolution.Height);
-
-					auto Result = FMessageDialog::Open(EAppMsgType::OkCancel,
-						FText::Format(NSLOCTEXT("LandscapeEditor.NewLandscape", "Import_HeightmapSizeMismatch", "The heightmap file's size ({FileSizeX}\u00D7{FileSizeY}) does not match the current Landscape extent ({LandscapeSizeX}\u00D7{LandscapeSizeY}), if you continue it will be padded/clipped to fit"), Args));
-
-					if (Result != EAppReturnType::Ok)
-					{
-						return;
-					}
-				}
-			}
-
-			FLandscapeHeightmapImportData ImportData = HeightmapFormat->Import(*Filename, NAME_None, ImportResolution);
-
-			if (ImportData.ResultCode == ELandscapeImportResult::Error)
-			{
-				FMessageDialog::Open(EAppMsgType::Ok, ImportData.ErrorMessage);
+				FMessageDialog::Open(EAppMsgType::Ok,
+					FText::Format(NSLOCTEXT("LandscapeEditor.Import", "Import_SizeMismatchRaw", "The file resolution does not match the current Landscape extent ({LandscapeSizeX}\u00D7{LandscapeSizeY}), and its exact resolution could not be determined"), Args));
 
 				return;
 			}
-
-			{
-				ALandscape* Landscape = GetLandscape();
-				FScopedSetLandscapeEditingLayer Scope(Landscape, GetCurrentLayerGuid(), [&] { check(Landscape); Landscape->RequestLayersContentUpdate(ELandscapeLayerUpdateMode::Update_Heightmap_All); });
-
-				TArray<uint16> Data;
-				if (ImportResolution != LandscapeResolution)
-				{
-					// Cloned from FLandscapeEditorDetailCustomization_NewLandscape.OnCreateButtonClicked
-					// so that reimports behave the same as the initial import :)
-
-					const int32 OffsetX = (int32)(LandscapeResolution.Width - ImportResolution.Width) / 2;
-					const int32 OffsetY = (int32)(LandscapeResolution.Height - ImportResolution.Height) / 2;
-
-					Data.SetNumUninitialized(LandscapeResolution.Width * LandscapeResolution.Height * sizeof(uint16));
-
-					LandscapeEditorUtils::ExpandData<uint16>(Data.GetData(), ImportData.Data.GetData(),
-						0, 0, ImportResolution.Width - 1, ImportResolution.Height - 1,
-						-OffsetX, -OffsetY, LandscapeResolution.Width - OffsetX - 1, LandscapeResolution.Height - OffsetY - 1);
-				}
-				else
-				{
-					Data = MoveTemp(ImportData.Data);
-				}
-
-				FScopedTransaction Transaction(LOCTEXT("Undo_ImportHeightmap", "Importing Landscape Heightmap"));
-
-				FHeightmapAccessor<false> HeightmapAccessor(LandscapeInfo);
-				HeightmapAccessor.SetData(MinX, MinY, MaxX, MaxY, Data.GetData());
-			}
 		}
-		else
+				
+		// display warning message if there is one and allow user to cancel
+		if (ImportResult == ELandscapeImportResult::Warning)
 		{
-			const ILandscapeWeightmapFileFormat* WeightmapFormat = LandscapeEditorModule.GetWeightmapFormatByExtension(*FPaths::GetExtension(Filename, true));
+			auto Result = FMessageDialog::Open(EAppMsgType::OkCancel, OutMessage);
 
-			if (!WeightmapFormat)
+			if (Result != EAppReturnType::Ok)
 			{
-				FMessageDialog::Open(EAppMsgType::Ok, NSLOCTEXT("LandscapeEditor.NewLandscape", "Import_UnknownFileType", "File type not recognised"));
-
 				return;
 			}
+		}
 
-			FLandscapeFileResolution ImportResolution = {0, 0};
+		const FLandscapeImportResolution ImportResolution = OutImportDescriptor.ImportResolutions[DescriptorIndex];
+		bool bResolutionMismatch = false;
+		// if the file is a format with resolution information, warn the user if the resolution doesn't match the current landscape
+		// unlike for raw this is only a warning as we can pad/clip the data if we know what resolution it is
+		if (ImportResolution != LandscapeResolution)
+		{
+			bResolutionMismatch = true;
 
-			const FLandscapeFileInfo WeightmapInfo = WeightmapFormat->Validate(*Filename, TargetInfo.LayerName);
+			FFormatNamedArguments Args;
+			Args.Add(TEXT("ImportSizeX"), ImportResolution.Width);
+			Args.Add(TEXT("ImportSizeY"), ImportResolution.Height);
+			Args.Add(TEXT("LandscapeSizeX"), LandscapeResolution.Width);
+			Args.Add(TEXT("LandscapeSizeY"), LandscapeResolution.Height);
 
-			// display error message if there is one, and abort the import
-			if (WeightmapInfo.ResultCode == ELandscapeImportResult::Error)
-			{
-				FMessageDialog::Open(EAppMsgType::Ok, WeightmapInfo.ErrorMessage);
+			auto Result = FMessageDialog::Open(EAppMsgType::OkCancel,
+				FText::Format(NSLOCTEXT("LandscapeEditor.Import", "Import_SizeMismatch", "The import size ({ImportSizeX}\u00D7{ImportSizeY}) does not match the current Landscape extent ({LandscapeSizeX}\u00D7{LandscapeSizeY}), if you continue it will be padded/clipped to fit"), Args));
 
+			if (Result != EAppReturnType::Ok)
+			{
 				return;
 			}
+		}
 
-			// if the file is a raw format with multiple possibly resolutions, only attempt import if one matches the current landscape
-			if (WeightmapInfo.PossibleResolutions.Num() > 1)
-			{
-				if (!WeightmapInfo.PossibleResolutions.Contains(LandscapeResolution))
-				{
-					FFormatNamedArguments Args;
-					Args.Add(TEXT("LandscapeSizeX"), LandscapeResolution.Width);
-					Args.Add(TEXT("LandscapeSizeY"), LandscapeResolution.Height);
-
-					FMessageDialog::Open(EAppMsgType::Ok,
-						FText::Format(NSLOCTEXT("LandscapeEditor.NewLandscape", "Import_LayerSizeMismatch_ResNotDetermined", "The layer file does not match the current Landscape extent ({LandscapeSizeX}\u00D7{LandscapeSizeY}), and its exact resolution could not be determined"), Args));
-
-					return;
-				}
-				else
-				{
-					ImportResolution = LandscapeResolution;
-				}
-			}
+		TArray<T> ImportData;
+		ImportResult = FLandscapeImportHelper::GetImportData<T>(OutImportDescriptor, DescriptorIndex, TargetInfo.LayerName, ImportData, OutMessage);
+		if (ImportResult == ELandscapeImportResult::Error)
+		{
+			FMessageDialog::Open(EAppMsgType::Ok, OutMessage);
+			return;
+		}
 
-			// display warning message if there is one and allow user to cancel
-			if (WeightmapInfo.ResultCode == ELandscapeImportResult::Warning)
+		// Expand if necessary...
+		{
+			TArray<T> FinalData;
+			if (bResolutionMismatch)
 			{
-				auto Result = FMessageDialog::Open(EAppMsgType::OkCancel, WeightmapInfo.ErrorMessage);
-
-				if (Result != EAppReturnType::Ok)
-				{
-					return;
-				}
+				FLandscapeImportHelper::ExpandImportData<T>(ImportData, FinalData, ImportResolution, LandscapeResolution);
 			}
-
-			// if the file is a format with resolution information, warn the user if the resolution doesn't match the current landscape
-			// unlike for raw this is only a warning as we can pad/clip the data if we know what resolution it is
-			if (WeightmapInfo.PossibleResolutions.Num() == 1)
-			{
-				ImportResolution = WeightmapInfo.PossibleResolutions[0];
-				if (ImportResolution != LandscapeResolution)
-				{
-					FFormatNamedArguments Args;
-					Args.Add(TEXT("FileSizeX"), ImportResolution.Width);
-					Args.Add(TEXT("FileSizeY"), ImportResolution.Height);
-					Args.Add(TEXT("LandscapeSizeX"), LandscapeResolution.Width);
-					Args.Add(TEXT("LandscapeSizeY"), LandscapeResolution.Height);
-
-					auto Result = FMessageDialog::Open(EAppMsgType::OkCancel,
-						FText::Format(NSLOCTEXT("LandscapeEditor.NewLandscape", "Import_LayerSizeMismatch_WillClamp", "The layer file's size ({FileSizeX}\u00D7{FileSizeY}) does not match the current Landscape extent ({LandscapeSizeX}\u00D7{LandscapeSizeY}), if you continue it will be padded/clipped to fit"), Args));
-
-					if (Result != EAppReturnType::Ok)
-					{
-						return;
-					}
-				}
-			}
-
-			FLandscapeWeightmapImportData ImportData = WeightmapFormat->Import(*Filename, TargetInfo.LayerName, ImportResolution);
-
-			if (ImportData.ResultCode == ELandscapeImportResult::Error)
+			else
 			{
-				FMessageDialog::Open(EAppMsgType::Ok, ImportData.ErrorMessage);
-
-				return;
+				FinalData = MoveTemp(ImportData);
 			}
+									
+			SetDataFunc(MinX, MinY, MaxX, MaxY, FinalData);
+		}
+	}			
+}
 
-			{
-				ALandscape* Landscape = GetLandscape();
-				FScopedSetLandscapeEditingLayer Scope(Landscape, GetCurrentLayerGuid(), [&] { check(Landscape); Landscape->RequestLayersContentUpdate(ELandscapeLayerUpdateMode::Update_Weightmap_All); });
-
-				TArray<uint8> Data;
-				if (ImportResolution != LandscapeResolution)
-				{
-					// Cloned from FLandscapeEditorDetailCustomization_NewLandscape.OnCreateButtonClicked
-					// so that reimports behave the same as the initial import :)
-
-					const int32 OffsetX = (int32)(LandscapeResolution.Width - ImportResolution.Width) / 2;
-					const int32 OffsetY = (int32)(LandscapeResolution.Height - ImportResolution.Height) / 2;
+void FEdModeLandscape::ImportData(const FLandscapeTargetListInfo& TargetInfo, const FString& Filename)
+{
+	if (TargetInfo.TargetType == ELandscapeToolTargetType::Heightmap)
+	{
+		ImportDataInternal<uint16>(TargetInfo, Filename, UseSingleFileImport(), [this, &TargetInfo](int32 MinX, int32 MinY, int32 MaxX, int32 MaxY, const TArray<uint16>& Data)
+		{
+			ULandscapeInfo* LandscapeInfo = TargetInfo.LandscapeInfo.Get();
+			ALandscape* Landscape = LandscapeInfo->LandscapeActor.Get();
+			FScopedSetLandscapeEditingLayer Scope(Landscape, GetCurrentLayerGuid(), [&] { check(Landscape); Landscape->RequestLayersContentUpdate(ELandscapeLayerUpdateMode::Update_Heightmap_All); });
 
-					Data.SetNumUninitialized(LandscapeResolution.Width * LandscapeResolution.Height * sizeof(uint8));
+			FScopedTransaction Transaction(LOCTEXT("Undo_ImportHeightmap", "Importing Landscape Heightmap"));
 
-					LandscapeEditorUtils::ExpandData<uint8>(Data.GetData(), ImportData.Data.GetData(),
-						0, 0, ImportResolution.Width - 1, ImportResolution.Height - 1,
-						-OffsetX, -OffsetY, LandscapeResolution.Width - OffsetX - 1, LandscapeResolution.Height - OffsetY - 1);
-				}
-				else
-				{
-					Data = MoveTemp(ImportData.Data);
-				}
+			FHeightmapAccessor<false> HeightmapAccessor(LandscapeInfo);
+			HeightmapAccessor.SetData(MinX, MinY, MaxX, MaxY, Data.GetData());
+		});
+	}
+	else
+	{
+		ImportDataInternal<uint8>(TargetInfo, Filename, UseSingleFileImport(), [this, &TargetInfo](int32 MinX, int32 MinY, int32 MaxX, int32 MaxY, const TArray<uint8>& Data)
+		{
+			ULandscapeInfo* LandscapeInfo = TargetInfo.LandscapeInfo.Get();
+			ALandscape* Landscape = LandscapeInfo->LandscapeActor.Get();
+			FScopedSetLandscapeEditingLayer Scope(Landscape, GetCurrentLayerGuid(), [&] { check(Landscape); Landscape->RequestLayersContentUpdate(ELandscapeLayerUpdateMode::Update_Weightmap_All); });
 
-				FScopedTransaction Transaction(LOCTEXT("Undo_ImportWeightmap", "Importing Landscape Layer"));
+			FScopedTransaction Transaction(LOCTEXT("Undo_ImportWeightmap", "Importing Landscape Layer"));
 
-				FAlphamapAccessor<false, false> AlphamapAccessor(LandscapeInfo, TargetInfo.LayerInfoObj.Get());
-				AlphamapAccessor.SetData(MinX, MinY, MaxX, MaxY, Data.GetData(), ELandscapeLayerPaintingRestriction::None);
-			}
-		}
+			FAlphamapAccessor<false, false> AlphamapAccessor(LandscapeInfo, TargetInfo.LayerInfoObj.Get());
+			AlphamapAccessor.SetData(MinX, MinY, MaxX, MaxY, Data.GetData(), ELandscapeLayerPaintingRestriction::None);
+		});
 	}
 }
 
diff --git a/Engine/Source/Editor/LandscapeEditor/Private/LandscapeEdMode.h b/Engine/Source/Editor/LandscapeEditor/Private/LandscapeEdMode.h
--- a/Engine/Source/Editor/LandscapeEditor/Private/LandscapeEdMode.h
+++ b/Engine/Source/Editor/LandscapeEditor/Private/LandscapeEdMode.h
@@ -603,6 +603,7 @@ public:
 
 	void ReimportData(const FLandscapeTargetListInfo& TargetInfo);
 	void ImportData(const FLandscapeTargetListInfo& TargetInfo, const FString& Filename);
+	bool UseSingleFileImport() const { return !IsGridBased(); }
 
 	/** Resample landscape to a different resolution or change the component size */
 	ALandscape* ChangeComponentSetting(int32 NumComponentsX, int32 NumComponentsY, int32 InNumSubsections, int32 InSubsectionSizeQuads, bool bResample);
diff --git a/Engine/Source/Editor/LandscapeEditor/Private/LandscapeEditorDetailCustomization_NewLandscape.cpp b/Engine/Source/Editor/LandscapeEditor/Private/LandscapeEditorDetailCustomization_NewLandscape.cpp
--- a/Engine/Source/Editor/LandscapeEditor/Private/LandscapeEditorDetailCustomization_NewLandscape.cpp
+++ b/Engine/Source/Editor/LandscapeEditor/Private/LandscapeEditorDetailCustomization_NewLandscape.cpp
@@ -21,7 +21,7 @@
 #include "Landscape.h"
 #include "LandscapeEditorUtils.h"
 #include "LandscapeConfigHelper.h"
-#include "NewLandscapeUtils.h"
+#include "LandscapeImportHelper.h"
 
 #include "DetailLayoutBuilder.h"
 #include "IDetailChildrenBuilder.h"
@@ -817,9 +817,10 @@ FReply FLandscapeEditorDetailCustomization_NewLandscape::OnCreateButtonClicked()
 		const int32 SizeX = ComponentCountX * QuadsPerComponent + 1;
 		const int32 SizeY = ComponentCountY * QuadsPerComponent + 1;
 
-		TOptional< TArray< FLandscapeImportLayerInfo > > MaterialImportLayers = FNewLandscapeUtils::CreateImportLayersInfo(UISettings, LandscapeEdMode->NewLandscapePreviewMode );
+		TArray<FLandscapeImportLayerInfo> MaterialImportLayers;
+		ELandscapeImportResult LayerImportResult = LandscapeEdMode->NewLandscapePreviewMode == ENewLandscapePreviewMode::NewLandscape ? UISettings->CreateNewLayersInfo(MaterialImportLayers) : UISettings->CreateImportLayersInfo(MaterialImportLayers);
 
-		if ( !MaterialImportLayers)
+		if (LayerImportResult == ELandscapeImportResult::Error)
 		{
 			return FReply::Handled();
 		}
@@ -827,9 +828,19 @@ FReply FLandscapeEditorDetailCustomization_NewLandscape::OnCreateButtonClicked()
 		TMap<FGuid, TArray<uint16>> HeightDataPerLayers;
 		TMap<FGuid, TArray<FLandscapeImportLayerInfo>> MaterialLayerDataPerLayers;
 
-		HeightDataPerLayers.Add(FGuid(), FNewLandscapeUtils::ComputeHeightData(UISettings, MaterialImportLayers.GetValue(), LandscapeEdMode->NewLandscapePreviewMode));
+		TArray<uint16> OutHeightData;
+		if (LandscapeEdMode->NewLandscapePreviewMode == ENewLandscapePreviewMode::NewLandscape)
+		{
+			UISettings->InitializeDefaultHeightData(OutHeightData);
+		}
+		else
+		{
+			UISettings->ExpandImportData(OutHeightData, MaterialImportLayers);
+		}
+
+		HeightDataPerLayers.Add(FGuid(), OutHeightData);
 		// ComputeHeightData will also modify/expand material layers data, which is why we create MaterialLayerDataPerLayers after calling ComputeHeightData
-		MaterialLayerDataPerLayers.Add(FGuid(), MaterialImportLayers.GetValue());
+		MaterialLayerDataPerLayers.Add(FGuid(), MoveTemp(MaterialImportLayers));
 
 		FScopedTransaction Transaction(LOCTEXT("Undo", "Creating New Landscape"));
 
@@ -847,12 +858,13 @@ FReply FLandscapeEditorDetailCustomization_NewLandscape::OnCreateButtonClicked()
 		// >= 8192x8192 -> LOD3
 		Landscape->StaticLightingLOD = FMath::DivideAndRoundUp(FMath::CeilLogTwo((SizeX * SizeY) / (2048 * 2048) + 1), (uint32)2);
 
+		FString ReimportHeightmapFilePath;
 		if (LandscapeEdMode->NewLandscapePreviewMode == ENewLandscapePreviewMode::ImportLandscape)
 		{
-			Landscape->ReimportHeightmapFilePath = UISettings->ImportLandscape_HeightmapFilename;
+			ReimportHeightmapFilePath = UISettings->ImportLandscape_HeightmapFilename;
 		}
 
-		Landscape->Import(FGuid::NewGuid(), 0, 0, SizeX - 1, SizeY - 1, UISettings->NewLandscape_SectionsPerComponent, UISettings->NewLandscape_QuadsPerSection, HeightDataPerLayers, nullptr, MaterialLayerDataPerLayers, UISettings->ImportLandscape_AlphamapType);
+		Landscape->Import(FGuid::NewGuid(), 0, 0, SizeX - 1, SizeY - 1, UISettings->NewLandscape_SectionsPerComponent, UISettings->NewLandscape_QuadsPerSection, HeightDataPerLayers, *ReimportHeightmapFilePath, MaterialLayerDataPerLayers, UISettings->ImportLandscape_AlphamapType);
 
 		ULandscapeInfo* LandscapeInfo = Landscape->GetLandscapeInfo();
 		check(LandscapeInfo);
@@ -926,7 +938,7 @@ FReply FLandscapeEditorDetailCustomization_NewLandscape::OnFitImportDataButtonCl
 	FEdModeLandscape* LandscapeEdMode = GetEditorMode();
 	if (LandscapeEdMode != nullptr)
 	{
-		ChooseBestComponentSizeForImport(LandscapeEdMode);
+		LandscapeEdMode->UISettings->ChooseBestComponentSizeForImport();
 	}
 
 	return FReply::Handled();
@@ -1023,7 +1035,7 @@ void FLandscapeEditorDetailCustomization_NewLandscape::OnImportHeightmapFilename
 	FEdModeLandscape* LandscapeEdMode = GetEditorMode();
 	if (LandscapeEdMode != nullptr)
 	{
-		FNewLandscapeUtils::ImportLandscapeData(LandscapeEdMode->UISettings, ImportResolutions);
+		LandscapeEdMode->UISettings->OnImportHeightmapFilenameChanged();
 	}
 }
 
@@ -1064,12 +1076,16 @@ TSharedRef<SWidget> FLandscapeEditorDetailCustomization_NewLandscape::GetImportL
 {
 	FMenuBuilder MenuBuilder(true, nullptr);
 
-	for (int32 i = 0; i < ImportResolutions.Num(); i++)
+	FEdModeLandscape* LandscapeEdMode = GetEditorMode();
+	if (LandscapeEdMode != nullptr)
 	{
-		FFormatNamedArguments Args;
-		Args.Add(TEXT("Width"), ImportResolutions[i].Width);
-		Args.Add(TEXT("Height"), ImportResolutions[i].Height);
-		MenuBuilder.AddMenuEntry(FText::Format(LOCTEXT("ImportResolution_Format", "{Width}\u00D7{Height}"), Args), FText(), FSlateIcon(), FExecuteAction::CreateSP(this, &FLandscapeEditorDetailCustomization_NewLandscape::OnChangeImportLandscapeResolution, i));
+		for (int32 i = 0; i < LandscapeEdMode->UISettings->HeightmapImportDescriptor.ImportResolutions.Num(); i++)
+		{
+			FFormatNamedArguments Args;
+			Args.Add(TEXT("Width"), LandscapeEdMode->UISettings->HeightmapImportDescriptor.ImportResolutions[i].Width);
+			Args.Add(TEXT("Height"), LandscapeEdMode->UISettings->HeightmapImportDescriptor.ImportResolutions[i].Height);
+			MenuBuilder.AddMenuEntry(FText::Format(LOCTEXT("ImportResolution_Format", "{Width}\u00D7{Height}"), Args), FText(), FSlateIcon(), FExecuteAction::CreateSP(this, &FLandscapeEditorDetailCustomization_NewLandscape::OnChangeImportLandscapeResolution, i));
+		}
 	}
 
 	return MenuBuilder.MakeWidget();
@@ -1080,10 +1096,7 @@ void FLandscapeEditorDetailCustomization_NewLandscape::OnChangeImportLandscapeRe
 	FEdModeLandscape* LandscapeEdMode = GetEditorMode();
 	if (LandscapeEdMode != nullptr)
 	{
-		LandscapeEdMode->UISettings->ImportLandscape_Width = ImportResolutions[Index].Width;
-		LandscapeEdMode->UISettings->ImportLandscape_Height = ImportResolutions[Index].Height;
-		LandscapeEdMode->UISettings->ClearImportLandscapeData();
-		ChooseBestComponentSizeForImport(LandscapeEdMode);
+		LandscapeEdMode->UISettings->OnChangeImportLandscapeResolution(Index);
 	}
 }
 
@@ -1110,11 +1123,6 @@ FText FLandscapeEditorDetailCustomization_NewLandscape::GetImportLandscapeResolu
 	return FText::GetEmpty();
 }
 
-void FLandscapeEditorDetailCustomization_NewLandscape::ChooseBestComponentSizeForImport(FEdModeLandscape* LandscapeEdMode)
-{
-	FNewLandscapeUtils::ChooseBestComponentSizeForImport(LandscapeEdMode->UISettings);
-}
-
 EVisibility FLandscapeEditorDetailCustomization_NewLandscape::GetMaterialTipVisibility() const
 {
 	FEdModeLandscape* LandscapeEdMode = GetEditorMode();
diff --git a/Engine/Source/Editor/LandscapeEditor/Private/LandscapeEditorDetailCustomization_NewLandscape.h b/Engine/Source/Editor/LandscapeEditor/Private/LandscapeEditorDetailCustomization_NewLandscape.h
--- a/Engine/Source/Editor/LandscapeEditor/Private/LandscapeEditorDetailCustomization_NewLandscape.h
+++ b/Engine/Source/Editor/LandscapeEditor/Private/LandscapeEditorDetailCustomization_NewLandscape.h
@@ -75,7 +75,6 @@ public:
 
 	bool GetImportButtonIsEnabled() const;
 	FReply OnFitImportDataButtonClicked();
-	void ChooseBestComponentSizeForImport(FEdModeLandscape* LandscapeEdMode);
 
 	FText GetOverallResolutionTooltip() const;
 
@@ -83,7 +82,6 @@ public:
 	EVisibility GetMaterialTipVisibility() const;
 
 protected:
-	TArray<FLandscapeFileResolution> ImportResolutions;
 	bool bUsingSlider;
 };
 
diff --git a/Engine/Source/Editor/LandscapeEditor/Private/LandscapeEditorDetailCustomization_TargetLayers.cpp b/Engine/Source/Editor/LandscapeEditor/Private/LandscapeEditorDetailCustomization_TargetLayers.cpp
--- a/Engine/Source/Editor/LandscapeEditor/Private/LandscapeEditorDetailCustomization_TargetLayers.cpp
+++ b/Engine/Source/Editor/LandscapeEditor/Private/LandscapeEditorDetailCustomization_TargetLayers.cpp
@@ -904,24 +904,29 @@ TSharedPtr<SWidget> FLandscapeEditorCustomNodeBuilder_TargetLayers::OnTargetLaye
 	if (Target->TargetType == ELandscapeToolTargetType::Heightmap || Target->LayerInfoObj != NULL)
 	{
 		FMenuBuilder MenuBuilder(true, NULL);
-
+		
 		MenuBuilder.BeginSection("LandscapeEditorLayerActions", LOCTEXT("LayerContextMenu.Heading", "Layer Actions"));
 		{
-			// Export
-			FUIAction ExportAction = FUIAction(FExecuteAction::CreateStatic(&FLandscapeEditorCustomNodeBuilder_TargetLayers::OnExportLayer, Target));
-			MenuBuilder.AddMenuEntry(LOCTEXT("LayerContextMenu.Export", "Export to file"), FText(), FSlateIcon(), ExportAction);
+			// In Grid Based Landscape Import/Export/Reimport is done through the Import Tool
+			FEdModeLandscape* LandscapeEdMode = GetEditorMode();
+			if (LandscapeEdMode && !LandscapeEdMode->IsGridBased())
+			{
+				// Export
+				FUIAction ExportAction = FUIAction(FExecuteAction::CreateStatic(&FLandscapeEditorCustomNodeBuilder_TargetLayers::OnExportLayer, Target));
+				MenuBuilder.AddMenuEntry(LOCTEXT("LayerContextMenu.Export", "Export to file"), FText(), FSlateIcon(), ExportAction);
 
-			// Import
-			FUIAction ImportAction = FUIAction(FExecuteAction::CreateStatic(&FLandscapeEditorCustomNodeBuilder_TargetLayers::OnImportLayer, Target));
-			MenuBuilder.AddMenuEntry(LOCTEXT("LayerContextMenu.Import", "Import from file"), FText(), FSlateIcon(), ImportAction);
+				// Import
+				FUIAction ImportAction = FUIAction(FExecuteAction::CreateStatic(&FLandscapeEditorCustomNodeBuilder_TargetLayers::OnImportLayer, Target));
+				MenuBuilder.AddMenuEntry(LOCTEXT("LayerContextMenu.Import", "Import from file"), FText(), FSlateIcon(), ImportAction);
 
-			// Reimport
-			const FString& ReimportPath = Target->GetReimportFilePath();
+				// Reimport
+				const FString& ReimportPath = Target->GetReimportFilePath();
 
-			if (!ReimportPath.IsEmpty())
-			{
-				FUIAction ReImportAction = FUIAction(FExecuteAction::CreateStatic(&FLandscapeEditorCustomNodeBuilder_TargetLayers::OnReimportLayer, Target));
-				MenuBuilder.AddMenuEntry(FText::Format(LOCTEXT("LayerContextMenu.ReImport", "Reimport from {0}"), FText::FromString(ReimportPath)), FText(), FSlateIcon(), ReImportAction);
+				if (!ReimportPath.IsEmpty())
+				{
+					FUIAction ReImportAction = FUIAction(FExecuteAction::CreateStatic(&FLandscapeEditorCustomNodeBuilder_TargetLayers::OnReimportLayer, Target));
+					MenuBuilder.AddMenuEntry(FText::Format(LOCTEXT("LayerContextMenu.ReImport", "Reimport from {0}"), FText::FromString(ReimportPath)), FText(), FSlateIcon(), ReImportAction);
+				}
 			}
 
 			if (Target->TargetType == ELandscapeToolTargetType::Weightmap)
@@ -962,6 +967,7 @@ void FLandscapeEditorCustomNodeBuilder_TargetLayers::OnExportLayer(const TShared
 	FEdModeLandscape* LandscapeEdMode = GetEditorMode();
 	if (LandscapeEdMode)
 	{
+		check(!LandscapeEdMode->IsGridBased());
 		IDesktopPlatform* DesktopPlatform = FDesktopPlatformModule::Get();
 
 		ULandscapeInfo* LandscapeInfo = Target->LandscapeInfo.Get();
@@ -1024,6 +1030,7 @@ void FLandscapeEditorCustomNodeBuilder_TargetLayers::OnImportLayer(const TShared
 	FEdModeLandscape* LandscapeEdMode = GetEditorMode();
 	if (LandscapeEdMode)
 	{
+		check(!LandscapeEdMode->IsGridBased());
 		IDesktopPlatform* DesktopPlatform = FDesktopPlatformModule::Get();
 
 		ULandscapeInfo* LandscapeInfo = Target->LandscapeInfo.Get();
@@ -1079,6 +1086,7 @@ void FLandscapeEditorCustomNodeBuilder_TargetLayers::OnReimportLayer(const TShar
 	FEdModeLandscape* LandscapeEdMode = GetEditorMode();
 	if (LandscapeEdMode)
 	{
+		check(!LandscapeEdMode->IsGridBased());
 		LandscapeEdMode->ReimportData(*Target);
 	}
 }
diff --git a/Engine/Source/Editor/LandscapeEditor/Private/LandscapeEditorObject.cpp b/Engine/Source/Editor/LandscapeEditor/Private/LandscapeEditorObject.cpp
--- a/Engine/Source/Editor/LandscapeEditor/Private/LandscapeEditorObject.cpp
+++ b/Engine/Source/Editor/LandscapeEditor/Private/LandscapeEditorObject.cpp
@@ -1,12 +1,14 @@
 // Copyright Epic Games, Inc. All Rights Reserved.
 
 #include "LandscapeEditorObject.h"
+#include "LandscapeDataAccess.h"
 #include "Engine/Texture2D.h"
 #include "HAL/FileManager.h"
 #include "Modules/ModuleManager.h"
 #include "UObject/ConstructorHelpers.h"
 #include "LandscapeEditorModule.h"
 #include "LandscapeRender.h"
+#include "LandscapeImportHelper.h"
 #include "LandscapeMaterialInstanceConstant.h"
 #include "Misc/ConfigCacheIni.h"
 #include "EngineUtils.h"
@@ -535,30 +537,203 @@ bool ULandscapeEditorObject::SetAlphaTexture(UTexture2D* InTexture, EColorChanne
 	return Result;
 }
 
-void ULandscapeEditorObject::ImportLandscapeData()
+void ULandscapeEditorObject::ChooseBestComponentSizeForImport()
 {
-	ILandscapeEditorModule& LandscapeEditorModule = FModuleManager::GetModuleChecked<ILandscapeEditorModule>("LandscapeEditor");
-	const ILandscapeHeightmapFileFormat* HeightmapFormat = LandscapeEditorModule.GetHeightmapFormatByExtension(*FPaths::GetExtension(ImportLandscape_HeightmapFilename, true));
-	
-	if (HeightmapFormat)
+	FLandscapeImportHelper::ChooseBestComponentSizeForImport(ImportLandscape_Width, ImportLandscape_Height, NewLandscape_QuadsPerSection, NewLandscape_SectionsPerComponent, NewLandscape_ComponentCount);
+}
+
+bool ULandscapeEditorObject::UseSingleFileImport() const
+{
+	if (ParentMode)
 	{
-		FLandscapeHeightmapImportData HeightmapImportData = HeightmapFormat->Import(*ImportLandscape_HeightmapFilename, NAME_None, {ImportLandscape_Width, ImportLandscape_Height});
-		ImportLandscape_HeightmapImportResult = HeightmapImportData.ResultCode;
-		ImportLandscape_HeightmapErrorMessage = HeightmapImportData.ErrorMessage;
-		ImportLandscape_Data = MoveTemp(HeightmapImportData.Data);
+		return ParentMode->UseSingleFileImport();
 	}
-	else if (!ImportLandscape_HeightmapFilename.IsEmpty())
+
+	return true;
+}
+
+void ULandscapeEditorObject::OnImportHeightmapFilenameChanged()
+{
+	ClearImportLandscapeData();
+	HeightmapImportDescriptorIndex = 0;
+	ImportLandscape_Width = 0;
+	ImportLandscape_Height = 0;
+
+	ImportLandscape_HeightmapImportResult = ELandscapeImportResult::Success;
+	ImportLandscape_HeightmapErrorMessage = FText();
+
+	ImportLandscape_HeightmapImportResult =
+		FLandscapeImportHelper::GetHeightmapImportDescriptor(ImportLandscape_HeightmapFilename, UseSingleFileImport(), HeightmapImportDescriptor, ImportLandscape_HeightmapErrorMessage);
+	if (ImportLandscape_HeightmapImportResult != ELandscapeImportResult::Error)
 	{
-		ImportLandscape_HeightmapImportResult = ELandscapeImportResult::Error;
-		ImportLandscape_HeightmapErrorMessage = NSLOCTEXT("LandscapeEditor.NewLandscape", "Import_UnknownFileType", "File type not recognised");
+		ImportLandscape_Width = HeightmapImportDescriptor.ImportResolutions[HeightmapImportDescriptorIndex].Width;
+		ImportLandscape_Height = HeightmapImportDescriptor.ImportResolutions[HeightmapImportDescriptorIndex].Height;
+		ChooseBestComponentSizeForImport();
+		ImportLandscapeData();
 	}
+}
 
+void ULandscapeEditorObject::OnChangeImportLandscapeResolution(int32 DescriptorIndex)
+{
+	check(DescriptorIndex >= 0 && DescriptorIndex < HeightmapImportDescriptor.ImportResolutions.Num());
+	HeightmapImportDescriptorIndex = DescriptorIndex;
+	ImportLandscape_Width = HeightmapImportDescriptor.ImportResolutions[HeightmapImportDescriptorIndex].Width;
+	ImportLandscape_Height = HeightmapImportDescriptor.ImportResolutions[HeightmapImportDescriptorIndex].Height;
+	ClearImportLandscapeData();
+	ChooseBestComponentSizeForImport();
+}
+
+void ULandscapeEditorObject::ImportLandscapeData()
+{
+	ImportLandscape_HeightmapImportResult = FLandscapeImportHelper::GetHeightmapImportData(HeightmapImportDescriptor, HeightmapImportDescriptorIndex, ImportLandscape_Data, ImportLandscape_HeightmapErrorMessage);
 	if (ImportLandscape_HeightmapImportResult == ELandscapeImportResult::Error)
 	{
 		ImportLandscape_Data.Empty();
 	}
 }
 
+ELandscapeImportResult ULandscapeEditorObject::CreateImportLayersInfo(TArray<FLandscapeImportLayerInfo>& OutImportLayerInfos)
+{
+	const int32 QuadsPerComponent = NewLandscape_SectionsPerComponent * NewLandscape_QuadsPerSection;
+	const int32 SizeX = NewLandscape_ComponentCount.X * QuadsPerComponent + 1;
+	const int32 SizeY = NewLandscape_ComponentCount.Y * QuadsPerComponent + 1;
+
+	const uint32 ImportSizeX = ImportLandscape_Width;
+	const uint32 ImportSizeY = ImportLandscape_Height;
+
+	if (ImportLandscape_HeightmapImportResult == ELandscapeImportResult::Error)
+	{
+		// Cancel import
+		return ELandscapeImportResult::Error;
+	}
+
+	OutImportLayerInfos.Reserve(ImportLandscape_Layers.Num());
+
+	// Fill in LayerInfos array and allocate data
+	for (FLandscapeImportLayer& UIImportLayer : ImportLandscape_Layers)
+	{
+		OutImportLayerInfos.Add((const FLandscapeImportLayer&)UIImportLayer); //slicing is fine here
+		FLandscapeImportLayerInfo& ImportLayer = OutImportLayerInfos.Last();
+
+		if (ImportLayer.LayerInfo != nullptr && ImportLayer.SourceFilePath != "")
+		{
+			FLandscapeImportDescriptor OutImportDescriptor;
+			UIImportLayer.ImportResult = FLandscapeImportHelper::GetWeightmapImportDescriptor(ImportLayer.SourceFilePath, UseSingleFileImport(), ImportLayer.LayerName, OutImportDescriptor, UIImportLayer.ErrorMessage);
+			if (UIImportLayer.ImportResult == ELandscapeImportResult::Error)
+			{
+				FMessageDialog::Open(EAppMsgType::Ok, UIImportLayer.ErrorMessage);
+				return ELandscapeImportResult::Error;
+			}
+
+			// Use same import index as Heightmap
+			int32 FoundIndex = INDEX_NONE;
+			for (int32 Index = 0; Index < OutImportDescriptor.FileResolutions.Num(); ++Index)
+			{
+				if (OutImportDescriptor.ImportResolutions[Index] == HeightmapImportDescriptor.ImportResolutions[HeightmapImportDescriptorIndex])
+				{
+					FoundIndex = Index;
+					break;
+				}
+			}
+
+			if (FoundIndex == INDEX_NONE)
+			{
+				UIImportLayer.ImportResult = ELandscapeImportResult::Error;
+				UIImportLayer.ErrorMessage = NSLOCTEXT("LandscapeEditor.NewLandscape", "Import_WeightHeightResolutionMismatch", "Weightmap import resolution isn't same as Heightmap resolution.");
+				FMessageDialog::Open(EAppMsgType::Ok, UIImportLayer.ErrorMessage);
+				return ELandscapeImportResult::Error;
+			}
+
+			UIImportLayer.ImportResult = FLandscapeImportHelper::GetWeightmapImportData(OutImportDescriptor, FoundIndex, ImportLayer.LayerName, ImportLayer.LayerData, UIImportLayer.ErrorMessage);
+			if (UIImportLayer.ImportResult == ELandscapeImportResult::Error)
+			{
+				FMessageDialog::Open(EAppMsgType::Ok, UIImportLayer.ErrorMessage);
+				return ELandscapeImportResult::Error;
+			}
+		}
+	}
+
+	return ELandscapeImportResult::Success;
+}
+
+ELandscapeImportResult ULandscapeEditorObject::CreateNewLayersInfo(TArray<FLandscapeImportLayerInfo>& OutNewLayerInfos)
+{
+	const int32 QuadsPerComponent = NewLandscape_SectionsPerComponent * NewLandscape_QuadsPerSection;
+	const int32 SizeX = NewLandscape_ComponentCount.X * QuadsPerComponent + 1;
+	const int32 SizeY = NewLandscape_ComponentCount.Y * QuadsPerComponent + 1;
+
+	OutNewLayerInfos.Reset(ImportLandscape_Layers.Num());
+
+	// Fill in LayerInfos array and allocate data
+	for (const FLandscapeImportLayer& UIImportLayer : ImportLandscape_Layers)
+	{
+		FLandscapeImportLayerInfo ImportLayer = FLandscapeImportLayerInfo(UIImportLayer.LayerName);
+		ImportLayer.LayerInfo = UIImportLayer.LayerInfo;
+		ImportLayer.SourceFilePath = "";
+		ImportLayer.LayerData = TArray<uint8>();
+		OutNewLayerInfos.Add(MoveTemp(ImportLayer));
+	}
+
+	// Fill the first weight-blended layer to 100%
+	if (FLandscapeImportLayerInfo* FirstBlendedLayer = OutNewLayerInfos.FindByPredicate([](const FLandscapeImportLayerInfo& ImportLayer) { return ImportLayer.LayerInfo && !ImportLayer.LayerInfo->bNoWeightBlend; }))
+	{
+		const int32 DataSize = SizeX * SizeY;
+		FirstBlendedLayer->LayerData.AddUninitialized(DataSize);
+
+		uint8* ByteData = FirstBlendedLayer->LayerData.GetData();
+		FMemory::Memset(ByteData, 255, DataSize);
+	}
+
+	return ELandscapeImportResult::Success;
+}
+
+void ULandscapeEditorObject::InitializeDefaultHeightData(TArray<uint16>& OutData)
+{
+	const int32 QuadsPerComponent = NewLandscape_SectionsPerComponent * NewLandscape_QuadsPerSection;
+	const int32 SizeX = NewLandscape_ComponentCount.X * QuadsPerComponent + 1;
+	const int32 SizeY = NewLandscape_ComponentCount.Y * QuadsPerComponent + 1;
+	const int32 TotalSize = SizeX * SizeY;
+	// Initialize heightmap data
+	OutData.Reset();
+	OutData.AddUninitialized(TotalSize);
+	
+	TArray<uint16> StrideData;
+	StrideData.AddUninitialized(SizeX);
+	// Initialize blank heightmap data
+	for (int32 X = 0; X < SizeX; ++X)
+	{
+		StrideData[X] = LandscapeDataAccess::MidValue;
+	}
+	for (int32 Y = 0; Y < SizeY; ++Y)
+	{
+		FMemory::Memcpy(&OutData[Y * SizeX], StrideData.GetData(), sizeof(uint16) * SizeX);
+	}
+}
+
+void ULandscapeEditorObject::ExpandImportData(TArray<uint16>& OutHeightData, TArray<FLandscapeImportLayerInfo>& OutImportLayerInfos)
+{
+	const TArray<uint16>& ImportData = GetImportLandscapeData();
+	if (ImportData.Num())
+	{
+		const int32 QuadsPerComponent = NewLandscape_SectionsPerComponent * NewLandscape_QuadsPerSection;
+		FLandscapeImportResolution RequiredResolution(NewLandscape_ComponentCount.X * QuadsPerComponent + 1, NewLandscape_ComponentCount.Y * QuadsPerComponent + 1);
+		FLandscapeImportResolution ImportResolution(ImportLandscape_Width, ImportLandscape_Height);
+
+		FLandscapeImportHelper::ExpandHeightmapImportData(ImportData, OutHeightData, ImportResolution, RequiredResolution);
+
+		for (int32 LayerIdx = 0; LayerIdx < OutImportLayerInfos.Num(); ++LayerIdx)
+		{
+			TArray<uint8>& OutImportLayerData = OutImportLayerInfos[LayerIdx].LayerData;
+			TArray<uint8> OutLayerData;
+			if (OutImportLayerData.Num())
+			{
+				FLandscapeImportHelper::ExpandWeightmapImportData(OutImportLayerData, OutLayerData, ImportResolution, RequiredResolution);
+				OutImportLayerData = MoveTemp(OutLayerData);
+			}
+		}
+	}
+}
+
 void ULandscapeEditorObject::RefreshImportLayersList()
 {
 	UTexture2D* ThumbnailWeightmap = LoadObject<UTexture2D>(NULL, TEXT("/Engine/EditorLandscapeResources/LandscapeThumbnailWeightmap.LandscapeThumbnailWeightmap"), NULL, LOAD_None, NULL);
@@ -607,27 +782,17 @@ void ULandscapeEditorObject::RefreshImportLayersList()
 				}
 				else
 				{
-					ILandscapeEditorModule& LandscapeEditorModule = FModuleManager::GetModuleChecked<ILandscapeEditorModule>("LandscapeEditor");
-					const ILandscapeWeightmapFileFormat* WeightmapFormat = LandscapeEditorModule.GetWeightmapFormatByExtension(*FPaths::GetExtension(NewImportLayer.SourceFilePath, true));
-
-					if (WeightmapFormat)
+					FLandscapeImportDescriptor OutImportDescriptor;
+					NewImportLayer.ImportResult = FLandscapeImportHelper::GetWeightmapImportDescriptor(NewImportLayer.SourceFilePath, UseSingleFileImport(), NewImportLayer.LayerName, OutImportDescriptor, NewImportLayer.ErrorMessage);
+					if (NewImportLayer.ImportResult != ELandscapeImportResult::Error)
 					{
-						FLandscapeFileInfo WeightmapImportInfo = WeightmapFormat->Validate(*NewImportLayer.SourceFilePath, NewImportLayer.LayerName);
-						NewImportLayer.ImportResult = WeightmapImportInfo.ResultCode;
-						NewImportLayer.ErrorMessage = WeightmapImportInfo.ErrorMessage;
-
-						if (WeightmapImportInfo.ResultCode != ELandscapeImportResult::Error &&
-							!WeightmapImportInfo.PossibleResolutions.Contains(FLandscapeFileResolution{ImportLandscape_Width, ImportLandscape_Height}))
+						int32 DescriptorIndex = OutImportDescriptor.FindDescriptorIndex(ImportLandscape_Width, ImportLandscape_Height);
+						if (DescriptorIndex == INDEX_NONE)
 						{
 							NewImportLayer.ImportResult = ELandscapeImportResult::Error;
-							NewImportLayer.ErrorMessage = NSLOCTEXT("LandscapeEditor.NewLandscape", "Import_LayerSizeMismatch", "Size of the layer file does not match size of heightmap file");
+							NewImportLayer.ErrorMessage = NSLOCTEXT("LandscapeEditor.NewLandscape", "Import_LayerSizeMismatch", "Size of the layer import size does not match size of heightmap import size");
 						}
 					}
-					else
-					{
-						NewImportLayer.ImportResult = ELandscapeImportResult::Error;
-						NewImportLayer.ErrorMessage = NSLOCTEXT("LandscapeEditor.NewLandscape", "Import_UnknownFileType", "File type not recognized");
-					}
 				}
 			}
 		}
diff --git a/Engine/Source/Editor/LandscapeEditor/Private/LandscapeImportHelper.cpp b/Engine/Source/Editor/LandscapeEditor/Private/LandscapeImportHelper.cpp
new file mode 100644
--- /dev/null
+++ b/Engine/Source/Editor/LandscapeEditor/Private/LandscapeImportHelper.cpp
@@ -0,0 +1,367 @@
+// Copyright Epic Games, Inc. All Rights Reserved.
+
+#include "LandscapeImportHelper.h"
+#include "LandscapeEditorModule.h"
+#include "LandscapeDataAccess.h"
+#include "LandscapeConfigHelper.h"
+#include "Modules/ModuleManager.h"
+#include "HAL/FileManager.h"
+#include "Misc/Paths.h"
+
+#define LOCTEXT_NAMESPACE "LandscapeImportHelper"
+
+bool FLandscapeImportHelper::ExtractCoordinates(FString BaseFilename, FIntPoint& OutCoord, FString& OutBaseFilePattern)
+{
+	//We expect file name in form: <tilename>_x<number>_y<number>
+	int32 XPos = BaseFilename.Find(TEXT("_x"), ESearchCase::IgnoreCase, ESearchDir::FromEnd);
+	int32 YPos = BaseFilename.Find(TEXT("_y"), ESearchCase::IgnoreCase, ESearchDir::FromEnd);
+	if (XPos != INDEX_NONE && YPos != INDEX_NONE && XPos < YPos)
+	{
+		FString XCoord = BaseFilename.Mid(XPos + 2, YPos - (XPos + 2));
+		FString YCoord = BaseFilename.Mid(YPos + 2, BaseFilename.Len() - (YPos + 2));
+
+		if (XCoord.IsNumeric() && YCoord.IsNumeric())
+		{
+			OutBaseFilePattern = BaseFilename.Mid(0, XPos);
+			TTypeFromString<int32>::FromString(OutCoord.X, *XCoord);
+			TTypeFromString<int32>::FromString(OutCoord.Y, *YCoord);
+			return true;
+		}
+	}
+
+	return false;
+}
+
+void FLandscapeImportHelper::GetMatchingFiles(const FString& FilePathPattern, TArray<FString>& OutFileToImport)
+{
+	IFileManager::Get().IterateDirectoryRecursively(*FPaths::GetPath(FilePathPattern), [&OutFileToImport, &FilePathPattern](const TCHAR* FilenameOrDirectory, bool bIsDirectory)
+	{
+		if (!bIsDirectory)
+		{
+			FString Filename(FilenameOrDirectory);
+			if (Filename.StartsWith(FilePathPattern))
+			{
+				OutFileToImport.Add(Filename);
+			}
+		}
+		return true;
+	});
+}
+
+template<class T>
+ELandscapeImportResult GetImportDataInternal(const FLandscapeImportDescriptor& ImportDescriptor, int32 DescriptorIndex, FName LayerName, T DefaultValue, TArray<T>& OutData, FText& OutMessage)
+{
+	if (DescriptorIndex < 0 || DescriptorIndex >= ImportDescriptor.ImportResolutions.Num())
+	{
+		OutMessage = LOCTEXT("Import_InvalidDescriptorIndex", "Invalid Descriptor Index");
+		return ELandscapeImportResult::Error;
+	}
+		
+	if (!ImportDescriptor.FileDescriptors.Num() || ImportDescriptor.ImportResolutions.Num() != ImportDescriptor.FileResolutions.Num())
+	{
+		OutMessage = LOCTEXT("Import_InvalidDescriptor", "Invalid Descriptor");
+		return ELandscapeImportResult::Error;
+	}
+	
+	int32 TotalWidth = ImportDescriptor.ImportResolutions[DescriptorIndex].Width;
+	int32 TotalHeight = ImportDescriptor.ImportResolutions[DescriptorIndex].Height;
+
+	OutData.Reset();
+	OutData.SetNumZeroed(TotalWidth * TotalHeight);
+	// Initialize All to default value so that non-covered regions have data
+	TArray<T> StrideData;
+	StrideData.SetNumUninitialized(TotalWidth);
+	for (int32 X = 0; X < TotalWidth; ++X)
+	{
+		StrideData[X] = DefaultValue;
+	}
+	for (int32 Y = 0; Y < TotalHeight; ++Y)
+	{
+		FMemory::Memcpy(&OutData[Y * TotalWidth], StrideData.GetData(), sizeof(T) * TotalWidth);
+	}
+
+	ELandscapeImportResult Result = ELandscapeImportResult::Success;
+	// Import Regions
+	ILandscapeEditorModule& LandscapeEditorModule = FModuleManager::GetModuleChecked<ILandscapeEditorModule>("LandscapeEditor");
+	const ILandscapeFileFormat<T>* FileFormat = LandscapeEditorModule.GetFormatByExtension<T>(*FPaths::GetExtension(ImportDescriptor.FileDescriptors[0].FilePath, true));
+	check(FileFormat);
+
+	int32 FileWidth = ImportDescriptor.FileResolutions[DescriptorIndex].Width;
+	int32 FileHeight = ImportDescriptor.FileResolutions[DescriptorIndex].Height;
+	
+	for (const FLandscapeImportFileDescriptor& FileDescriptor : ImportDescriptor.FileDescriptors)
+	{
+		FLandscapeImportData<T> ImportData = FileFormat->Import(*FileDescriptor.FilePath, LayerName, ImportDescriptor.FileResolutions[DescriptorIndex]);
+		OutMessage = ImportData.ErrorMessage;
+		Result = ImportData.ResultCode;
+		if (ImportData.ResultCode == ELandscapeImportResult::Error)
+		{
+			break;
+		}
+		
+		int32 StartX = FileDescriptor.Coord.X * FileWidth;
+		int32 StartY = FileDescriptor.Coord.Y * FileHeight;
+		
+		for (int32 Y = 0; Y < FileHeight; ++Y)
+		{		
+			int32 DestY = StartY + Y;
+			FMemory::Memcpy(&OutData[DestY * TotalWidth + StartX], &ImportData.Data[Y * FileWidth], FileWidth * sizeof(T));
+		}
+	}
+
+	return Result;
+}
+
+template<class T>
+ELandscapeImportResult GetImportDescriptorInternal(const FString& FilePath, bool bSingleFile, FName LayerName, FLandscapeImportDescriptor& OutImportDescriptor, FText& OutMessage)
+{
+	OutImportDescriptor.Reset();
+	if (FilePath.IsEmpty())
+	{
+		OutMessage = LOCTEXT("Import_InvalidPath", "Invalid file");
+		return ELandscapeImportResult::Error;
+	}
+		
+	FIntPoint OutCoord;
+	FIntPoint MinCoord(INT32_MAX, INT32_MAX); // All coords should be rebased to the min
+	FIntPoint MaxCoord(INT32_MIN, INT32_MIN);
+	FString OutFileImportPattern;
+	FString FilePathPattern;
+	TArray<FString> OutFilesToImport;
+	if (!bSingleFile && FLandscapeImportHelper::ExtractCoordinates(FPaths::GetBaseFilename(FilePath), OutCoord, OutFileImportPattern))
+	{
+		FilePathPattern = FPaths::GetPath(FilePath) / OutFileImportPattern;
+		FLandscapeImportHelper::GetMatchingFiles(FilePathPattern, OutFilesToImport);
+	}
+	else
+	{
+		bSingleFile = true;
+		OutFilesToImport.Add(FilePath);
+	}
+
+	TArray<FLandscapeFileResolution> ImportResolutions;
+
+	ILandscapeEditorModule& LandscapeEditorModule = FModuleManager::GetModuleChecked<ILandscapeEditorModule>("LandscapeEditor");
+	const ILandscapeFileFormat<T>* FileFormat = nullptr;
+	for (const FString& ImportFilename : OutFilesToImport)
+	{
+		const bool bFirst = FileFormat == nullptr;
+		const ILandscapeFileFormat<T>* CurrentFileFormat = LandscapeEditorModule.GetFormatByExtension<T>(*FPaths::GetExtension(ImportFilename, true));
+		if (FileFormat != nullptr && FileFormat != CurrentFileFormat)
+		{
+			OutMessage = LOCTEXT("Import_MismatchFileType", "Not all files have the same file type");
+			return ELandscapeImportResult::Error;
+		}
+		FileFormat = CurrentFileFormat;
+
+		if (FileFormat)
+		{
+			FLandscapeFileInfo FileInfo = FileFormat->Validate(*ImportFilename);
+			if (FileInfo.ResultCode == ELandscapeImportResult::Error)
+			{
+				OutMessage = FileInfo.ErrorMessage;
+				return FileInfo.ResultCode;
+			}
+
+			FString OutLocalFileImportPattern;
+			if (!bSingleFile && !FLandscapeImportHelper::ExtractCoordinates(FPaths::GetBaseFilename(ImportFilename), OutCoord, OutLocalFileImportPattern))
+			{
+				OutMessage = FText::Format(LOCTEXT("Import_InvalidFilename", "File '{0}' doesn't have proper pattern(ex: {1}_x0_y0.{2})"),FText::FromString(FPaths::GetBaseFilename(ImportFilename)), FText::FromString(OutFileImportPattern), FText::FromString(FPaths::GetExtension(FilePath)));
+				return ELandscapeImportResult::Error;
+			}
+			MinCoord.X = FMath::Min(OutCoord.X, MinCoord.X);
+			MinCoord.Y = FMath::Min(OutCoord.Y, MinCoord.Y);
+			MaxCoord.X = FMath::Max(OutCoord.X, MaxCoord.X);
+			MaxCoord.Y = FMath::Max(OutCoord.Y, MaxCoord.Y);
+
+			FLandscapeImportFileDescriptor FileDescriptor(ImportFilename, OutCoord);
+			OutImportDescriptor.FileDescriptors.Add(FileDescriptor);
+
+			if (bFirst)
+			{
+				// Resolutions will need to match for all files (keep the first one to compare)
+				OutImportDescriptor.FileResolutions = MoveTemp(FileInfo.PossibleResolutions);
+								
+				if (FileInfo.DataScale.IsSet())
+				{
+					OutImportDescriptor.Scale = FileInfo.DataScale.GetValue();
+					OutImportDescriptor.Scale.Z *= LANDSCAPE_INV_ZSCALE;
+				}
+			}
+			else
+			{
+				if (OutImportDescriptor.FileResolutions != FileInfo.PossibleResolutions)
+				{
+					OutMessage = LOCTEXT("Import_MismatchResolution", "Not all files have the same resolution");
+					return ELandscapeImportResult::Error;
+				}
+				else if (FileInfo.DataScale.IsSet())
+				{
+					FVector CurrentScale = FileInfo.DataScale.GetValue();
+					CurrentScale.Z *= LANDSCAPE_INV_ZSCALE;
+
+					if (!OutImportDescriptor.Scale.Equals(CurrentScale))
+					{
+						OutMessage = LOCTEXT("Import_MismatchScale", "Not all files have the same data scale");
+						return ELandscapeImportResult::Error;
+					}
+				}
+			}
+		}
+		else
+		{
+			OutMessage = LOCTEXT("Import_UnknownFileType", "File type not recognized");
+			return ELandscapeImportResult::Error;
+		}
+	}
+
+	check(OutImportDescriptor.FileDescriptors.Num());
+	// Rebase with MinCoord
+	for (FLandscapeImportFileDescriptor& FileDescriptor : OutImportDescriptor.FileDescriptors)
+	{
+		FileDescriptor.Coord -= MinCoord;
+	}
+	MaxCoord -= MinCoord;
+
+	// Compute Import Total Size
+	for (const FLandscapeFileResolution& Resolution : OutImportDescriptor.FileResolutions)
+	{
+		OutImportDescriptor.ImportResolutions.Add(FLandscapeImportResolution((MaxCoord.X+1)*Resolution.Width, (MaxCoord.Y+1)*Resolution.Height));
+	}
+
+	return ELandscapeImportResult::Success;
+}
+
+template<class T>
+void ExpandImportDataInternal(const TArray<T>& InData, TArray<T>& OutData, const FLandscapeImportResolution& CurrentResolution, const FLandscapeImportResolution& RequiredResolution)
+{
+	check(InData.Num() == CurrentResolution.Width * CurrentResolution.Height);
+	// Center Import data
+	const int32 OffsetX = (int32)(RequiredResolution.Width - CurrentResolution.Width) / 2;
+	const int32 OffsetY = (int32)(RequiredResolution.Height - CurrentResolution.Height) / 2;
+
+	FIntRect SrcRegion(0, 0, CurrentResolution.Width - 1, CurrentResolution.Height - 1);
+	FIntRect DestRegion(-OffsetX, -OffsetY, RequiredResolution.Width - OffsetX - 1, RequiredResolution.Height - OffsetY - 1);
+	FLandscapeConfigHelper::ExpandData<T>(InData, OutData, SrcRegion, DestRegion, true);
+}
+
+ELandscapeImportResult FLandscapeImportHelper::GetHeightmapImportDescriptor(const FString& FilePath, bool bSingleFile, FLandscapeImportDescriptor& OutImportDescriptor, FText& OutMessage)
+{
+	return GetImportDescriptorInternal<uint16>(FilePath, bSingleFile, NAME_None, OutImportDescriptor, OutMessage);
+}
+
+ELandscapeImportResult FLandscapeImportHelper::GetWeightmapImportDescriptor(const FString& FilePath, bool bSingleFile, FName LayerName, FLandscapeImportDescriptor& OutImportDescriptor, FText& OutMessage)
+{
+	return GetImportDescriptorInternal<uint8>(FilePath, bSingleFile, LayerName, OutImportDescriptor, OutMessage);
+}
+
+ELandscapeImportResult FLandscapeImportHelper::GetHeightmapImportData(const FLandscapeImportDescriptor& ImportDescriptor, int32 DescriptorIndex, TArray<uint16>& OutData, FText& OutMessage)
+{
+	return GetImportDataInternal<uint16>(ImportDescriptor, DescriptorIndex, NAME_None, LandscapeDataAccess::MidValue, OutData, OutMessage);
+}
+
+ELandscapeImportResult FLandscapeImportHelper::GetWeightmapImportData(const FLandscapeImportDescriptor& ImportDescriptor, int32 DescriptorIndex, FName LayerName, TArray<uint8>& OutData, FText& OutMessage)
+{
+	return GetImportDataInternal<uint8>(ImportDescriptor, DescriptorIndex, LayerName, 0, OutData, OutMessage);
+}
+
+void FLandscapeImportHelper::ExpandWeightmapImportData(const TArray<uint8>& InData, TArray<uint8>& OutData, const FLandscapeImportResolution& CurrentResolution, const FLandscapeImportResolution& RequiredResolution)
+{
+	ExpandImportDataInternal<uint8>(InData, OutData, CurrentResolution, RequiredResolution);
+}
+
+void FLandscapeImportHelper::ExpandHeightmapImportData(const TArray<uint16>& InData, TArray<uint16>& OutData, const FLandscapeImportResolution& CurrentResolution, const FLandscapeImportResolution& RequiredResolution)
+{
+	ExpandImportDataInternal<uint16>(InData, OutData, CurrentResolution, RequiredResolution);
+}
+
+void FLandscapeImportHelper::ChooseBestComponentSizeForImport(int32 Width, int32 Height, int32& InOutQuadsPerSection, int32& InOutSectionsPerComponent, FIntPoint& OutComponentCount)
+{
+	bool bValidSubsectionSizeParam = false;
+	bool bValidQuadsPerSectionParam = false;
+	check(Width > 0 && Height > 0);
+		
+	bool bFoundMatch = false;
+	// Try to find a section size and number of sections that exactly matches the dimensions of the heightfield
+	for (int32 SectionSizesIdx = UE_ARRAY_COUNT(FLandscapeConfig::SubsectionSizeQuadsValues) - 1; SectionSizesIdx >= 0; SectionSizesIdx--)
+	{
+		for (int32 NumSectionsIdx = UE_ARRAY_COUNT(FLandscapeConfig::NumSectionValues) - 1; NumSectionsIdx >= 0; NumSectionsIdx--)
+		{
+			int32 ss = FLandscapeConfig::SubsectionSizeQuadsValues[SectionSizesIdx];
+			int32 ns = FLandscapeConfig::NumSectionValues[NumSectionsIdx];
+
+			// Check if the passed in values are found in the array of valid values
+			bValidSubsectionSizeParam |= (InOutSectionsPerComponent == ns);
+			bValidQuadsPerSectionParam |= (InOutQuadsPerSection == ss);
+
+			if (((Width - 1) % (ss * ns)) == 0 && ((Width - 1) / (ss * ns)) <= 32 &&
+				((Height - 1) % (ss * ns)) == 0 && ((Height - 1) / (ss * ns)) <= 32)
+			{
+				bFoundMatch = true;
+				InOutQuadsPerSection = ss;
+				InOutSectionsPerComponent = ns;
+				OutComponentCount.X = (Width - 1) / (ss * ns);
+				OutComponentCount.Y = (Height - 1) / (ss * ns);
+				break;
+			}
+		}
+		if (bFoundMatch)
+		{
+			break;
+		}
+	}
+
+	if (!bFoundMatch)
+	{
+		if (!bValidSubsectionSizeParam)
+		{
+			InOutSectionsPerComponent = FLandscapeConfig::NumSectionValues[0];
+		}
+
+		if (!bValidQuadsPerSectionParam)
+		{
+			InOutSectionsPerComponent = FLandscapeConfig::SubsectionSizeQuadsValues[0];
+		}
+
+		// if there was no exact match, try increasing the section size until we encompass the whole heightmap
+		const int32 CurrentSectionSize = InOutQuadsPerSection;
+		const int32 CurrentNumSections = InOutSectionsPerComponent;
+		for (int32 SectionSizesIdx = 0; SectionSizesIdx < UE_ARRAY_COUNT(FLandscapeConfig::SubsectionSizeQuadsValues); SectionSizesIdx++)
+		{
+			if (FLandscapeConfig::SubsectionSizeQuadsValues[SectionSizesIdx] < CurrentSectionSize)
+			{
+				continue;
+			}
+
+			const int32 ComponentsX = FMath::DivideAndRoundUp((Width - 1), FLandscapeConfig::SubsectionSizeQuadsValues[SectionSizesIdx] * CurrentNumSections);
+			const int32 ComponentsY = FMath::DivideAndRoundUp((Height - 1), FLandscapeConfig::SubsectionSizeQuadsValues[SectionSizesIdx] * CurrentNumSections);
+			if (ComponentsX <= 32 && ComponentsY <= 32)
+			{
+				bFoundMatch = true;
+				InOutQuadsPerSection = FLandscapeConfig::SubsectionSizeQuadsValues[SectionSizesIdx];
+				OutComponentCount.X = ComponentsX;
+				OutComponentCount.Y = ComponentsY;
+				break;
+			}
+		}
+	}
+
+	if (!bFoundMatch)
+	{
+		// if the heightmap is very large, fall back to using the largest values we support
+		const int32 MaxSectionSize = FLandscapeConfig::SubsectionSizeQuadsValues[UE_ARRAY_COUNT(FLandscapeConfig::SubsectionSizeQuadsValues) - 1];
+		const int32 MaxNumSubSections = FLandscapeConfig::NumSectionValues[UE_ARRAY_COUNT(FLandscapeConfig::NumSectionValues) - 1];
+		const int32 ComponentsX = FMath::DivideAndRoundUp((Width - 1), MaxSectionSize * MaxNumSubSections);
+		const int32 ComponentsY = FMath::DivideAndRoundUp((Height - 1), MaxSectionSize * MaxNumSubSections);
+
+		bFoundMatch = true;
+		InOutQuadsPerSection = MaxSectionSize;
+		InOutSectionsPerComponent = MaxNumSubSections;
+		OutComponentCount.X = ComponentsX;
+		OutComponentCount.Y = ComponentsY;
+	}
+
+	check(bFoundMatch);
+}
+
+#undef LOCTEXT_NAMESPACE
diff --git a/Engine/Source/Editor/LandscapeEditor/Private/NewLandscapeUtils.cpp b/Engine/Source/Editor/LandscapeEditor/Private/NewLandscapeUtils.cpp
--- a/Engine/Source/Editor/LandscapeEditor/Private/NewLandscapeUtils.cpp
+++ b/Engine/Source/Editor/LandscapeEditor/Private/NewLandscapeUtils.cpp
@@ -19,83 +19,7 @@
 
 void FNewLandscapeUtils::ChooseBestComponentSizeForImport(ULandscapeEditorObject* UISettings)
 {
-	int32 Width = UISettings->ImportLandscape_Width;
-	int32 Height = UISettings->ImportLandscape_Height;
-
-	bool bFoundMatch = false;
-	if(Width > 0 && Height > 0)
-	{
-		// Try to find a section size and number of sections that exactly matches the dimensions of the heightfield
-		for(int32 SectionSizesIdx = UE_ARRAY_COUNT(FLandscapeConfig::SubsectionSizeQuadsValues) - 1; SectionSizesIdx >= 0; SectionSizesIdx--)
-		{
-			for(int32 NumSectionsIdx = UE_ARRAY_COUNT(FLandscapeConfig::NumSectionValues) - 1; NumSectionsIdx >= 0; NumSectionsIdx--)
-			{
-				int32 ss = FLandscapeConfig::SubsectionSizeQuadsValues[SectionSizesIdx];
-				int32 ns = FLandscapeConfig::NumSectionValues[NumSectionsIdx];
-
-				if(((Width - 1) % (ss * ns)) == 0 && ((Width - 1) / (ss * ns)) <= 32 &&
-					((Height - 1) % (ss * ns)) == 0 && ((Height - 1) / (ss * ns)) <= 32)
-				{
-					bFoundMatch = true;
-					UISettings->NewLandscape_QuadsPerSection = ss;
-					UISettings->NewLandscape_SectionsPerComponent = ns;
-					UISettings->NewLandscape_ComponentCount.X = (Width - 1) / (ss * ns);
-					UISettings->NewLandscape_ComponentCount.Y = (Height - 1) / (ss * ns);
-					UISettings->NewLandscape_ClampSize();
-					break;
-				}
-			}
-			if(bFoundMatch)
-			{
-				break;
-			}
-		}
-
-		if(!bFoundMatch)
-		{
-			// if there was no exact match, try increasing the section size until we encompass the whole heightmap
-			const int32 CurrentSectionSize = UISettings->NewLandscape_QuadsPerSection;
-			const int32 CurrentNumSections = UISettings->NewLandscape_SectionsPerComponent;
-			for(int32 SectionSizesIdx = 0; SectionSizesIdx < UE_ARRAY_COUNT(FLandscapeConfig::SubsectionSizeQuadsValues); SectionSizesIdx++)
-			{
-				if(FLandscapeConfig::SubsectionSizeQuadsValues[SectionSizesIdx] < CurrentSectionSize)
-				{
-					continue;
-				}
-
-				const int32 ComponentsX = FMath::DivideAndRoundUp((Width - 1), FLandscapeConfig::SubsectionSizeQuadsValues[SectionSizesIdx] * CurrentNumSections);
-				const int32 ComponentsY = FMath::DivideAndRoundUp((Height - 1), FLandscapeConfig::SubsectionSizeQuadsValues[SectionSizesIdx] * CurrentNumSections);
-				if(ComponentsX <= 32 && ComponentsY <= 32)
-				{
-					bFoundMatch = true;
-					UISettings->NewLandscape_QuadsPerSection = FLandscapeConfig::SubsectionSizeQuadsValues[SectionSizesIdx];
-					//UISettings->NewLandscape_SectionsPerComponent = ;
-					UISettings->NewLandscape_ComponentCount.X = ComponentsX;
-					UISettings->NewLandscape_ComponentCount.Y = ComponentsY;
-					UISettings->NewLandscape_ClampSize();
-					break;
-				}
-			}
-		}
-
-		if(!bFoundMatch)
-		{
-			// if the heightmap is very large, fall back to using the largest values we support
-			const int32 MaxSectionSize = FLandscapeConfig::SubsectionSizeQuadsValues[UE_ARRAY_COUNT(FLandscapeConfig::SubsectionSizeQuadsValues) - 1];
-			const int32 MaxNumSubSections = FLandscapeConfig::NumSectionValues[UE_ARRAY_COUNT(FLandscapeConfig::NumSectionValues) - 1];
-			const int32 ComponentsX = FMath::DivideAndRoundUp((Width - 1), MaxSectionSize * MaxNumSubSections);
-			const int32 ComponentsY = FMath::DivideAndRoundUp((Height - 1), MaxSectionSize * MaxNumSubSections);
-
-			bFoundMatch = true;
-			UISettings->NewLandscape_QuadsPerSection = MaxSectionSize;
-			UISettings->NewLandscape_SectionsPerComponent = MaxNumSubSections;
-			UISettings->NewLandscape_ComponentCount.X = ComponentsX;
-			UISettings->NewLandscape_ComponentCount.Y = ComponentsY;
-			UISettings->NewLandscape_ClampSize();
-		}
-
-		check(bFoundMatch);
-	}
+	UISettings->ChooseBestComponentSizeForImport();
 }
 
 void FNewLandscapeUtils::ImportLandscapeData( ULandscapeEditorObject* UISettings, TArray< FLandscapeFileResolution >& ImportResolutions )
@@ -111,181 +35,66 @@ void FNewLandscapeUtils::ImportLandscapeData( ULandscapeEditorObject* UISettings
 	UISettings->ClearImportLandscapeData();
 	UISettings->ImportLandscape_HeightmapImportResult = ELandscapeImportResult::Success;
 	UISettings->ImportLandscape_HeightmapErrorMessage = FText();
-
+	int32 IndexToUse = 0;
 	if(!UISettings->ImportLandscape_HeightmapFilename.IsEmpty())
 	{
-		ILandscapeEditorModule& LandscapeEditorModule = FModuleManager::GetModuleChecked<ILandscapeEditorModule>("LandscapeEditor");
-		const ILandscapeHeightmapFileFormat* HeightmapFormat = LandscapeEditorModule.GetHeightmapFormatByExtension(*FPaths::GetExtension(UISettings->ImportLandscape_HeightmapFilename, true));
+		FLandscapeImportDescriptor OutImportDescriptor;
+		UISettings->ImportLandscape_HeightmapImportResult =
+			FLandscapeImportHelper::GetHeightmapImportDescriptor(UISettings->ImportLandscape_HeightmapFilename, UISettings->UseSingleFileImport(), OutImportDescriptor, UISettings->ImportLandscape_HeightmapErrorMessage);
 
-		if(HeightmapFormat)
-		{
-			FLandscapeFileInfo HeightmapImportInfo = HeightmapFormat->Validate(*UISettings->ImportLandscape_HeightmapFilename, NAME_None);
-			UISettings->ImportLandscape_HeightmapImportResult = HeightmapImportInfo.ResultCode;
-			UISettings->ImportLandscape_HeightmapErrorMessage = HeightmapImportInfo.ErrorMessage;
-			ImportResolutions = MoveTemp(HeightmapImportInfo.PossibleResolutions);
-			if(HeightmapImportInfo.DataScale.IsSet())
-			{
-				UISettings->NewLandscape_Scale = HeightmapImportInfo.DataScale.GetValue();
-				UISettings->NewLandscape_Scale.Z *= LANDSCAPE_INV_ZSCALE;
-			}
-		}
-		else
+		if (UISettings->ImportLandscape_HeightmapImportResult != ELandscapeImportResult::Error)
 		{
-			UISettings->ImportLandscape_HeightmapImportResult = ELandscapeImportResult::Error;
-			UISettings->ImportLandscape_HeightmapErrorMessage = LOCTEXT("Import_UnknownFileType", "File type not recognized");
+			IndexToUse = OutImportDescriptor.FileResolutions.Num() / 2;
+			// In a single file import the region should match the file resolution
+			check(OutImportDescriptor.ImportResolutions[IndexToUse].Width == OutImportDescriptor.FileResolutions[IndexToUse].Width &&
+				OutImportDescriptor.ImportResolutions[IndexToUse].Height == OutImportDescriptor.FileResolutions[IndexToUse].Height);
+			ImportResolutions = MoveTemp(OutImportDescriptor.FileResolutions);
 		}
 	}
 
 	if(ImportResolutions.Num() > 0)
 	{
-		int32 i = ImportResolutions.Num() / 2;
-		UISettings->ImportLandscape_Width = ImportResolutions[i].Width;
-		UISettings->ImportLandscape_Height = ImportResolutions[i].Height;
+		UISettings->ImportLandscape_Width = ImportResolutions[IndexToUse].Width;
+		UISettings->ImportLandscape_Height = ImportResolutions[IndexToUse].Height;
+		UISettings->ChooseBestComponentSizeForImport();
 		UISettings->ImportLandscapeData();
-		ChooseBestComponentSizeForImport(UISettings);
 	}
 }
 
 TOptional< TArray< FLandscapeImportLayerInfo > > FNewLandscapeUtils::CreateImportLayersInfo( ULandscapeEditorObject* UISettings, int32 NewLandscapePreviewMode )
 {
-	const int32 ComponentCountX = UISettings->NewLandscape_ComponentCount.X;
-	const int32 ComponentCountY = UISettings->NewLandscape_ComponentCount.Y;
-	const int32 QuadsPerComponent = UISettings->NewLandscape_SectionsPerComponent * UISettings->NewLandscape_QuadsPerSection;
-	const int32 SizeX = ComponentCountX * QuadsPerComponent + 1;
-	const int32 SizeY = ComponentCountY * QuadsPerComponent + 1;
-
 	TArray<FLandscapeImportLayerInfo> ImportLayers;
-
+	ELandscapeImportResult Result = ELandscapeImportResult::Success;
 	if(NewLandscapePreviewMode == ENewLandscapePreviewMode::NewLandscape)
 	{
-		const auto& ImportLandscapeLayersList = UISettings->ImportLandscape_Layers;
-		ImportLayers.Reserve(ImportLandscapeLayersList.Num());
-
-		// Fill in LayerInfos array and allocate data
-		for(const FLandscapeImportLayer& UIImportLayer : ImportLandscapeLayersList)
-		{
-			FLandscapeImportLayerInfo ImportLayer = FLandscapeImportLayerInfo(UIImportLayer.LayerName);
-			ImportLayer.LayerInfo = UIImportLayer.LayerInfo;
-			ImportLayer.SourceFilePath = "";
-			ImportLayer.LayerData = TArray<uint8>();
-			ImportLayers.Add(MoveTemp(ImportLayer));
-		}
-
-		// Fill the first weight-blended layer to 100%
-		if(FLandscapeImportLayerInfo* FirstBlendedLayer = ImportLayers.FindByPredicate([](const FLandscapeImportLayerInfo& ImportLayer) { return ImportLayer.LayerInfo && !ImportLayer.LayerInfo->bNoWeightBlend; }))
-		{
-			FirstBlendedLayer->LayerData.AddUninitialized(SizeX * SizeY);
-
-			uint8* ByteData = FirstBlendedLayer->LayerData.GetData();
-			for(int32 i = 0; i < SizeX * SizeY; i++)
-			{
-				ByteData[i] = 255;
-			}
-		}
+		Result = UISettings->CreateNewLayersInfo(ImportLayers);
 	}
 	else if(NewLandscapePreviewMode == ENewLandscapePreviewMode::ImportLandscape)
 	{
-		const uint32 ImportSizeX = UISettings->ImportLandscape_Width;
-		const uint32 ImportSizeY = UISettings->ImportLandscape_Height;
-
-		if(UISettings->ImportLandscape_HeightmapImportResult == ELandscapeImportResult::Error)
-		{
-			// Cancel import
-			return TOptional< TArray< FLandscapeImportLayerInfo > >();
-		}
-
-		TArray<FLandscapeImportLayer>& ImportLandscapeLayersList = UISettings->ImportLandscape_Layers;
-		ImportLayers.Reserve(ImportLandscapeLayersList.Num());
-
-		// Fill in LayerInfos array and allocate data
-		for(FLandscapeImportLayer& UIImportLayer : ImportLandscapeLayersList)
-		{
-			ImportLayers.Add((const FLandscapeImportLayer&)UIImportLayer); //slicing is fine here
-			FLandscapeImportLayerInfo& ImportLayer = ImportLayers.Last();
-
-			if(ImportLayer.LayerInfo != nullptr && ImportLayer.SourceFilePath != "")
-			{
-				ILandscapeEditorModule& LandscapeEditorModule = FModuleManager::GetModuleChecked<ILandscapeEditorModule>("LandscapeEditor");
-				const ILandscapeWeightmapFileFormat* WeightmapFormat = LandscapeEditorModule.GetWeightmapFormatByExtension(*FPaths::GetExtension(ImportLayer.SourceFilePath, true));
-
-				if(WeightmapFormat)
-				{
-					FLandscapeWeightmapImportData WeightmapImportData = WeightmapFormat->Import(*ImportLayer.SourceFilePath, ImportLayer.LayerName, { ImportSizeX, ImportSizeY });
-					UIImportLayer.ImportResult = WeightmapImportData.ResultCode;
-					UIImportLayer.ErrorMessage = WeightmapImportData.ErrorMessage;
-					ImportLayer.LayerData = MoveTemp(WeightmapImportData.Data);
-				}
-				else
-				{
-					UIImportLayer.ImportResult = ELandscapeImportResult::Error;
-					UIImportLayer.ErrorMessage = LOCTEXT("Import_UnknownFileType", "File type not recognized");
-				}
-
-				if(UIImportLayer.ImportResult == ELandscapeImportResult::Error)
-				{
-					ImportLayer.LayerData.Empty();
-					FMessageDialog::Open(EAppMsgType::Ok, UIImportLayer.ErrorMessage);
+		Result = UISettings->CreateImportLayersInfo(ImportLayers);
+	}
 
-					// Cancel import
-					return TOptional< TArray< FLandscapeImportLayerInfo > >();
-				}
-			}
-		}
+	if (Result == ELandscapeImportResult::Error)
+	{
+		return TOptional<TArray<FLandscapeImportLayerInfo>>();
 	}
 
-	return ImportLayers;
+	return MoveTemp(ImportLayers);
 }
 
-TArray< uint16 > FNewLandscapeUtils::ComputeHeightData( ULandscapeEditorObject* UISettings, TArray< FLandscapeImportLayerInfo >& ImportLayers, int32 NewLandscapePreviewMode )
+TArray<uint16> FNewLandscapeUtils::ComputeHeightData(ULandscapeEditorObject* UISettings, TArray< FLandscapeImportLayerInfo >& ImportLayers, int32 NewLandscapePreviewMode)
 {
-	const int32 ComponentCountX = UISettings->NewLandscape_ComponentCount.X;
-	const int32 ComponentCountY = UISettings->NewLandscape_ComponentCount.Y;
-	const int32 QuadsPerComponent = UISettings->NewLandscape_SectionsPerComponent * UISettings->NewLandscape_QuadsPerSection;
-	const int32 SizeX = ComponentCountX * QuadsPerComponent + 1;
-	const int32 SizeY = ComponentCountY * QuadsPerComponent + 1;
-
-	const uint32 ImportSizeX = UISettings->ImportLandscape_Width;
-	const uint32 ImportSizeY = UISettings->ImportLandscape_Height;
-
-	// Initialize heightmap data
-	TArray<uint16> Data;
-	Data.AddUninitialized(SizeX * SizeY);
-	uint16* WordData = Data.GetData();
-
-	// Initialize blank heightmap data
-	for(int32 i = 0; i < SizeX * SizeY; i++)
+	TArray<uint16> HeightData;
+	if (NewLandscapePreviewMode == ENewLandscapePreviewMode::NewLandscape)
 	{
-		WordData[i] = 32768;
+		UISettings->InitializeDefaultHeightData(HeightData);
 	}
-
-	if(NewLandscapePreviewMode == ENewLandscapePreviewMode::ImportLandscape)
+	else
 	{
-		const TArray<uint16>& ImportData = UISettings->GetImportLandscapeData();
-		if(ImportData.Num() != 0)
-		{
-			const int32 OffsetX = (int32)(SizeX - ImportSizeX) / 2;
-			const int32 OffsetY = (int32)(SizeY - ImportSizeY) / 2;
-
-			// Heightmap
-			Data = LandscapeEditorUtils::ExpandData(ImportData,
-				0, 0, ImportSizeX - 1, ImportSizeY - 1,
-				-OffsetX, -OffsetY, SizeX - OffsetX - 1, SizeY - OffsetY - 1);
-
-			// Layers
-			for(int32 LayerIdx = 0; LayerIdx < ImportLayers.Num(); LayerIdx++)
-			{
-				TArray<uint8>& ImportLayerData = ImportLayers[LayerIdx].LayerData;
-				if(ImportLayerData.Num())
-				{
-					ImportLayerData = LandscapeEditorUtils::ExpandData(ImportLayerData,
-						0, 0, ImportSizeX - 1, ImportSizeY - 1,
-						-OffsetX, -OffsetY, SizeX - OffsetX - 1, SizeY - OffsetY - 1);
-				}
-			}
-		}
+		UISettings->ExpandImportData(HeightData, ImportLayers);
 	}
 
-	return Data;
+	return MoveTemp(HeightData);
 }
 
 #undef LOCTEXT_NAMESPACE
diff --git a/Engine/Source/Editor/LandscapeEditor/Public/LandscapeEditorModule.h b/Engine/Source/Editor/LandscapeEditor/Public/LandscapeEditorModule.h
--- a/Engine/Source/Editor/LandscapeEditor/Public/LandscapeEditorModule.h
+++ b/Engine/Source/Editor/LandscapeEditor/Public/LandscapeEditorModule.h
@@ -6,6 +6,8 @@
 #include "Modules/ModuleInterface.h"
 #include "Templates/SharedPointer.h"
 #include "LandscapeFileFormatInterface.h"
+#include "Templates/UnrealTypeTraits.h"
+#include "Templates/EnableIf.h"
 
 class FUICommandList;
 
@@ -31,5 +33,17 @@ public:
 	virtual const ILandscapeHeightmapFileFormat* GetHeightmapFormatByExtension(const TCHAR* Extension) const = 0;
 	virtual const ILandscapeWeightmapFileFormat* GetWeightmapFormatByExtension(const TCHAR* Extension) const = 0;
 
+	template<typename T>
+	typename TEnableIf<TIsSame<T, uint16>::Value, const ILandscapeHeightmapFileFormat*>::Type GetFormatByExtension(const TCHAR* Extension)
+	{
+		return GetHeightmapFormatByExtension(Extension);
+	}
+
+	template<typename T>
+	typename TEnableIf<TIsSame<T, uint8>::Value, const ILandscapeWeightmapFileFormat*>::Type GetFormatByExtension(const TCHAR* Extension)
+	{
+		return GetWeightmapFormatByExtension(Extension);
+	}
+
 	virtual TSharedPtr<FUICommandList> GetLandscapeLevelViewportCommandList() const = 0;
 };
diff --git a/Engine/Source/Editor/LandscapeEditor/Public/LandscapeEditorObject.h b/Engine/Source/Editor/LandscapeEditor/Public/LandscapeEditorObject.h
--- a/Engine/Source/Editor/LandscapeEditor/Public/LandscapeEditorObject.h
+++ b/Engine/Source/Editor/LandscapeEditor/Public/LandscapeEditorObject.h
@@ -9,6 +9,7 @@
 #include "Editor/LandscapeEditor/Private/LandscapeEdMode.h"
 #include "LandscapeFileFormatInterface.h"
 #include "LandscapeBlueprintBrush.h"
+#include "LandscapeImportHelper.h"
 
 #include "LandscapeEditorObject.generated.h"
 
@@ -522,6 +523,8 @@ class ULandscapeEditorObject : public UObject
 	UPROPERTY(NonTransactional)
 	uint32 ImportLandscape_Height;
 
+	FLandscapeImportDescriptor HeightmapImportDescriptor;
+	int32 HeightmapImportDescriptorIndex;
 private:
 	UPROPERTY(NonTransactional)
 	TArray<uint16> ImportLandscape_Data;
@@ -644,11 +647,17 @@ public:
 
 	const TArray<uint16>& GetImportLandscapeData() const { return ImportLandscape_Data; }
 	void ClearImportLandscapeData() { ImportLandscape_Data.Empty(); }
-
+	void ChooseBestComponentSizeForImport();
 	void ImportLandscapeData();
 	void RefreshImportLayersList();
-	
+	ELandscapeImportResult CreateImportLayersInfo(TArray<FLandscapeImportLayerInfo>& OutImportLayerInfos);
+	ELandscapeImportResult CreateNewLayersInfo(TArray<FLandscapeImportLayerInfo>& OutNewLayerInfos);
+	void InitializeDefaultHeightData(TArray<uint16>& OutData);
+	void ExpandImportData(TArray<uint16>& OutHeightData, TArray<FLandscapeImportLayerInfo>& OutImportLayerInfos);
 	void UpdateComponentLayerWhitelist();
+	bool UseSingleFileImport() const;
+	void OnChangeImportLandscapeResolution(int32 DescriptorIndex);
+	void OnImportHeightmapFilenameChanged();
 
 	int32 ClampLandscapeSize(int32 InComponentsCount) const
 	{
diff --git a/Engine/Source/Editor/LandscapeEditor/Public/LandscapeImportHelper.h b/Engine/Source/Editor/LandscapeEditor/Public/LandscapeImportHelper.h
new file mode 100644
--- /dev/null
+++ b/Engine/Source/Editor/LandscapeEditor/Public/LandscapeImportHelper.h
@@ -0,0 +1,142 @@
+// Copyright Epic Games, Inc. All Rights Reserved.
+
+#pragma once
+
+#include "CoreMinimal.h"
+#include "UObject/ObjectMacros.h"
+#include "LandscapeFileFormatInterface.h"
+#include "Templates/UnrealTypeTraits.h"
+#include "Templates/EnableIf.h"
+
+struct FLandscapeImportFileDescriptor
+{
+	FLandscapeImportFileDescriptor(const FString& InFilePath, const FIntPoint& InCoord)
+		: Coord(InCoord), FilePath(InFilePath){}
+
+	/* Which tile does that descriptor represent */
+	FIntPoint Coord;
+	/* File path */
+	FString FilePath;
+};
+
+struct FLandscapeImportResolution
+{
+	FLandscapeImportResolution(uint32 InWidth, uint32 InHeight) : Width(InWidth), Height(InHeight) {}
+
+	uint32 Width;
+	uint32 Height;
+};
+
+FORCEINLINE bool operator==(const FLandscapeImportResolution& Lhs, const FLandscapeImportResolution& Rhs)
+{
+	return Lhs.Width == Rhs.Width && Lhs.Height == Rhs.Height;
+}
+
+FORCEINLINE bool operator!=(const FLandscapeImportResolution& Lhs, const FLandscapeImportResolution& Rhs)
+{
+	return !(Lhs == Rhs);
+}
+
+struct FLandscapeImportDescriptor
+{
+	FLandscapeImportDescriptor()
+		: Scale(100, 100, 100) {}
+
+	void Reset()
+	{
+		ImportResolutions.Reset();
+		FileResolutions.Reset();
+		FileDescriptors.Reset();
+		Scale = FVector(100, 100, 100);
+	}
+
+	int32 FindDescriptorIndex(int32 ImportWidth, int32 ImportHeight)
+	{
+		for (int32 Index = 0; Index < ImportResolutions.Num(); ++Index)
+		{
+			if (ImportResolutions[Index].Width == ImportWidth && ImportResolutions[Index].Height == ImportHeight)
+			{
+				return Index;
+			}
+		}
+
+		return INDEX_NONE;
+	}
+
+	/* Landscape Import Resolution based on File Coords + Resolutions */
+	TArray<FLandscapeImportResolution> ImportResolutions;
+	/* Single File Resolutions */
+	TArray<FLandscapeFileResolution> FileResolutions;
+	/* Files contributing to this descriptor */
+	TArray<FLandscapeImportFileDescriptor> FileDescriptors;
+	/* Scale of Import data */
+	FVector Scale;
+};
+
+class LANDSCAPEEDITOR_API FLandscapeImportHelper
+{
+public:
+	/**
+	 * @param FilePath File to import when bSingleFile is true.
+	 * @param bSingleFile If false FilePath will be used as the seed to find matching "_xA_yB" files in the same directory.
+	 * @param OutImportDescriptor Descriptor output which contains all the information needed to call GetHeightmapImportData.
+	 * @param OutMessage In case the result of the operation is  an error or warning, this will contain the reason.
+	 * @return If the operation was a Success/Error/Warning. 
+	 */
+	static ELandscapeImportResult GetHeightmapImportDescriptor(const FString& FilePath, bool bSingleFile, FLandscapeImportDescriptor& OutImportDescriptor, FText& OutMessage);
+
+	/**
+	 * @param FilePath File to import when bSingleFile is true.
+	 * @param bSingleFile If false FilePath will be used as the seed to find matching "_xA_yB" files in the same directory.
+	 * @param LayerName Name of the Weightmap layer being imported. This can be used by the file formats.
+	 * @param OutImportDescriptor Descriptor output which contains all the information needed to call GetHeightmapImportData.
+	 * @param OutMessage In case the result of the operation is  an error or warning, this will contain the reason.
+	 * @return If the operation was a Success/Error/Warning.
+	 */
+	static ELandscapeImportResult GetWeightmapImportDescriptor(const FString& FilePath, bool bSingleFile, FName LayerName, FLandscapeImportDescriptor& OutImportDescriptor, FText& OutMessage);
+
+	static ELandscapeImportResult GetHeightmapImportData(const FLandscapeImportDescriptor& ImportDescriptor, int32 DescriptorIndex, TArray<uint16>& OutData, FText& OutMessage);
+	static ELandscapeImportResult GetWeightmapImportData(const FLandscapeImportDescriptor& ImportDescriptor, int32 DescriptorIndex, FName LayerName, TArray<uint8>& OutData, FText& OutMessage);
+
+	static void ExpandWeightmapImportData(const TArray<uint8>& InData, TArray<uint8>& OutData, const FLandscapeImportResolution& CurrentResolution, const FLandscapeImportResolution& RequiredResolution);
+	static void ExpandHeightmapImportData(const TArray<uint16>& InData, TArray<uint16>& OutData, const FLandscapeImportResolution& CurrentResolution, const FLandscapeImportResolution& RequiredResolution);
+
+	static void ChooseBestComponentSizeForImport(int32 Width, int32 Height, int32& InOutQuadsPerSection, int32& InOutSectionsPerComponent, FIntPoint& OutComponentCount);
+
+	static bool ExtractCoordinates(FString BaseFilename, FIntPoint& OutCoord, FString& OutBaseFilePattern);
+	static void GetMatchingFiles(const FString& FilePathPattern, TArray<FString>& OutFilePaths);
+
+
+	template<typename T>
+	static typename TEnableIf<TIsSame<T, uint16>::Value, ELandscapeImportResult>::Type GetImportDescriptor(const FString& FilePath, bool bSingleFile, FName LayerName, FLandscapeImportDescriptor& OutImportDescriptor, FText& OutMessage)
+	{
+		return GetHeightmapImportDescriptor(FilePath, bSingleFile, OutImportDescriptor, OutMessage);
+	}
+
+	template<typename T>
+	static typename TEnableIf<TIsSame<T, uint16>::Value, ELandscapeImportResult>::Type GetImportData(const FLandscapeImportDescriptor& ImportDescriptor, int32 DescriptorIndex, FName LayerName, TArray<T>& OutData, FText& OutMessage)
+	{
+		return GetHeightmapImportData(ImportDescriptor, DescriptorIndex, OutData, OutMessage);
+	}
+	template<typename T>
+	static typename TEnableIf<TIsSame<T, uint16>::Value, void>::Type ExpandImportData(const TArray<T>& InData, TArray<T>& OutData, const FLandscapeImportResolution& CurrentResolution, const FLandscapeImportResolution& RequiredResolution)
+	{
+		return ExpandHeightmapImportData(InData, OutData, CurrentResolution, RequiredResolution);
+	}
+
+	template<typename T>
+	static typename TEnableIf<TIsSame<T, uint8>::Value, ELandscapeImportResult>::Type GetImportDescriptor(const FString& FilePath, bool bSingleFile, FName LayerName, FLandscapeImportDescriptor& OutImportDescriptor, FText& OutMessage)
+	{
+		return GetWeightmapImportDescriptor(FilePath, bSingleFile, LayerName, OutImportDescriptor, OutMessage);
+	}
+	template<typename T>
+	static typename TEnableIf<TIsSame<T, uint8>::Value, ELandscapeImportResult>::Type GetImportData(const FLandscapeImportDescriptor& ImportDescriptor, int32 DescriptorIndex, FName LayerName, TArray<T>& OutData, FText& OutMessage)
+	{
+		return GetWeightmapImportData(ImportDescriptor, DescriptorIndex, LayerName, OutData, OutMessage);
+	}
+	template<typename T>
+	static typename TEnableIf<TIsSame<T, uint8>::Value, void>::Type ExpandImportData(const TArray<T>& InData, TArray<T>& OutData, const FLandscapeImportResolution& CurrentResolution, const FLandscapeImportResolution& RequiredResolution)
+	{
+		return ExpandWeightmapImportData(InData, OutData, CurrentResolution, RequiredResolution);
+	}
+};
\ No newline at end of file
diff --git a/Engine/Source/Editor/LandscapeEditor/Public/NewLandscapeUtils.h b/Engine/Source/Editor/LandscapeEditor/Public/NewLandscapeUtils.h
--- a/Engine/Source/Editor/LandscapeEditor/Public/NewLandscapeUtils.h
+++ b/Engine/Source/Editor/LandscapeEditor/Public/NewLandscapeUtils.h
@@ -9,6 +9,8 @@ struct FLandscapeFileResolution;
 struct FLandscapeImportLayerInfo;
 class ULandscapeEditorObject;
 
+// This class has been replaced by FLandscapeImportHelper & ULandscapeEditorObject methods.
+// @todo_ow: deprecate (probably just delete since it was added for Datasmith which was updated to the new import code)
 class LANDSCAPEEDITOR_API FNewLandscapeUtils
 {
 public:
diff --git a/Engine/Source/Runtime/Landscape/Public/LandscapeConfigHelper.h b/Engine/Source/Runtime/Landscape/Public/LandscapeConfigHelper.h
--- a/Engine/Source/Runtime/Landscape/Public/LandscapeConfigHelper.h
+++ b/Engine/Source/Runtime/Landscape/Public/LandscapeConfigHelper.h
@@ -78,6 +78,7 @@ private:
 	static void MoveSplinesToLandscape(ULandscapeInfo* InLandscapeInfo, ALandscapeProxy* InLandscape, float InScaleFactor);
 	static void MoveFoliageToLandscape(ULandscapeInfo* InLandscapeInfo, ULandscapeInfo* InNewLandscapeInfo);
 
+public:
 	template<typename T>
 	static void CopyData(const TArray<T>& InData, TArray<T>& OutData, const FIntRect& InSrcRegion, const FIntRect& InDestRegion, bool bInResample)
 	{
@@ -92,24 +93,36 @@ private:
 	}
 
 	template<typename T>
-	static void ExpandData(const TArray<T>& InData, TArray<T>& OutData, const FIntRect& InSrcRegion, const FIntRect& InDestRegion)
+	static void ExpandData(const TArray<T>& InData, TArray<T>& OutData, const FIntRect& InSrcRegion, const FIntRect& InDestRegion, bool bOffset = false)
 	{
 		// Regions are in ComponentQuads and we want Vertices (+1)
 		const int32 SrcWidth = InSrcRegion.Width() + 1;
 		const int32 SrcHeight = InSrcRegion.Height() + 1;
 		const int32 DestWidth = InDestRegion.Width() + 1;
 		const int32 DestHeight = InDestRegion.Height() + 1;
+		const int32 OffsetX = bOffset ? InDestRegion.Min.X - InSrcRegion.Min.X : 0;
+		const int32 OffsetY = bOffset ? InDestRegion.Min.Y - InSrcRegion.Min.X : 0;
 		
 		OutData.Empty(DestWidth * DestHeight);
 		OutData.AddUninitialized(DestWidth * DestHeight);
 
 		for (int32 Y = 0; Y < DestHeight; ++Y)
 		{
-			const int32 OldY = FMath::Clamp<int32>(Y, 0, SrcHeight - 1);
+			const int32 OldY = FMath::Clamp<int32>(Y + OffsetY, 0, SrcHeight - 1);
 						
-			// Copy one row of the old data
-			FMemory::Memcpy(&OutData[Y * DestWidth], &InData[OldY * SrcWidth], FMath::Min<int32>(SrcWidth, DestWidth) * sizeof(T));
-		
+			// Pad anything to the left
+			const T PadLeft = InData[OldY * SrcWidth];
+			for (int32 X = 0; X < -OffsetX; ++X)
+			{
+				OutData[Y * DestWidth + X] = PadLeft;
+			}
+
+			{
+				const int32 X = FMath::Max(0, -OffsetX);
+				const int32 OldX = FMath::Clamp<int32>(X + OffsetX, 0, SrcWidth - 1);
+				FMemory::Memcpy(&OutData[Y * DestWidth + X], &InData[OldY * SrcWidth + OldX], FMath::Min<int32>(SrcWidth, DestWidth) * sizeof(T));
+			}
+
 			const T PadRight = InData[OldY * SrcWidth + SrcWidth - 1];
 			for (int32 X = SrcWidth; X < DestWidth; ++X)
 			{
