Landscape: New Import/Export Tool allows importing/exporting multiple files to support large worlds (WorldPartition)

diff --git a/Engine/Plugins/Enterprise/DatasmithImporter/Source/DatasmithImporter/Private/DatasmithLandscapeImporter.cpp b/Engine/Plugins/Enterprise/DatasmithImporter/Source/DatasmithImporter/Private/DatasmithLandscapeImporter.cpp
--- a/Engine/Plugins/Enterprise/DatasmithImporter/Source/DatasmithImporter/Private/DatasmithLandscapeImporter.cpp
+++ b/Engine/Plugins/Enterprise/DatasmithImporter/Source/DatasmithImporter/Private/DatasmithLandscapeImporter.cpp
@@ -11,7 +11,6 @@
 #include "LandscapeDataAccess.h"
 #include "LandscapeEditorModule.h"
 #include "LandscapeEditorObject.h"
-#include "LandscapeEditorUtils.h"
 #include "LandscapeFileFormatInterface.h"
 #include "LandscapeImportHelper.h"
 
@@ -55,7 +54,7 @@ AActor* FDatasmithLandscapeImporter::ImportLandscapeActor( const TSharedRef< IDa
 	const int32 SizeY = OutComponentCount.Y * QuadsPerComponent + 1;
 
 	TArray<uint16> FinalHeightData;
-	FLandscapeImportHelper::ExpandHeightmapImportData(ImportData, FinalHeightData, OutImportDescriptor.ImportResolutions[DescriptorIndex], FLandscapeImportResolution(SizeX, SizeY));
+	FLandscapeImportHelper::TransformHeightmapImportData(ImportData, FinalHeightData, OutImportDescriptor.ImportResolutions[DescriptorIndex], FLandscapeImportResolution(SizeX, SizeY), ELandscapeImportTransformType::ExpandCentered);
 		
 	const FVector Offset = FTransform( LandscapeActorElement->GetRotation(), FVector::ZeroVector, 
 		LandscapeActorElement->GetScale() ).TransformVector( FVector( -OutComponentCount.X * QuadsPerComponent / 2, -OutComponentCount.Y * QuadsPerComponent / 2, 0 ) );
diff --git a/Engine/Source/Editor/EditorStyle/Private/StarshipStyle.cpp b/Engine/Source/Editor/EditorStyle/Private/StarshipStyle.cpp
--- a/Engine/Source/Editor/EditorStyle/Private/StarshipStyle.cpp
+++ b/Engine/Source/Editor/EditorStyle/Private/StarshipStyle.cpp
@@ -6297,7 +6297,8 @@ void FStarshipEditorStyle::FStyle::SetupLandscapeEditorStyle()
 			Set("LandscapeEditor.NewLandscape.Small",       	new IMAGE_BRUSH("Icons/Landscape/Landscape_NewLandscape_x40", 		Icon20x20));
 			Set("LandscapeEditor.ResizeLandscape",				new IMAGE_BRUSH("Icons/Landscape/Landscape_Resize_x40",				Icon20x20));
 			Set("LandscapeEditor.ResizeLandscape.Small",		new IMAGE_BRUSH("Icons/Landscape/Landscape_Resize_x40",				Icon20x20));
-
+			Set("LandscapeEditor.ImportExportTool",				new IMAGE_BRUSH("Icons/Landscape/Landscape_ImportExport_x40",		Icon20x20));
+			Set("LandscapeEditor.ImportExportTool.Small",		new IMAGE_BRUSH("Icons/Landscape/Landscape_ImportExport_x40",		Icon20x20));
 
 			Set("LandscapeEditor.SculptTool",       			new IMAGE_BRUSH("Icons/Landscape/Landscape_Sculpt_x40",            Icon20x20));
 			Set("LandscapeEditor.SculptTool.Small",       		new IMAGE_BRUSH("Icons/Landscape/Landscape_Sculpt_x40",            Icon20x20));
diff --git a/Engine/Source/Editor/LandscapeEditor/Private/LandscapeEdMode.cpp b/Engine/Source/Editor/LandscapeEditor/Private/LandscapeEdMode.cpp
--- a/Engine/Source/Editor/LandscapeEditor/Private/LandscapeEdMode.cpp
+++ b/Engine/Source/Editor/LandscapeEditor/Private/LandscapeEdMode.cpp
@@ -31,6 +31,7 @@
 #include "LandscapeEdModeTools.h"
 #include "LandscapeInfoMap.h"
 #include "LandscapeImportHelper.h"
+#include "LandscapeConfigHelper.h"
 
 //Slate dependencies
 #include "Misc/FeedbackContext.h"
@@ -57,42 +58,6 @@
 
 DEFINE_LOG_CATEGORY(LogLandscapeEdMode);
 
-struct HNewLandscapeGrabHandleProxy : public HHitProxy
-{
-	DECLARE_HIT_PROXY();
-
-	ELandscapeEdge::Type Edge;
-
-	HNewLandscapeGrabHandleProxy(ELandscapeEdge::Type InEdge) :
-		HHitProxy(HPP_Wireframe),
-		Edge(InEdge)
-	{
-	}
-
-	virtual EMouseCursor::Type GetMouseCursor() override
-	{
-		switch (Edge)
-		{
-		case ELandscapeEdge::X_Negative:
-		case ELandscapeEdge::X_Positive:
-			return EMouseCursor::ResizeLeftRight;
-		case ELandscapeEdge::Y_Negative:
-		case ELandscapeEdge::Y_Positive:
-			return EMouseCursor::ResizeUpDown;
-		case ELandscapeEdge::X_Negative_Y_Negative:
-		case ELandscapeEdge::X_Positive_Y_Positive:
-			return EMouseCursor::ResizeSouthEast;
-		case ELandscapeEdge::X_Negative_Y_Positive:
-		case ELandscapeEdge::X_Positive_Y_Negative:
-			return EMouseCursor::ResizeSouthWest;
-		}
-
-		return EMouseCursor::SlashedCircle;
-	}
-};
-
-IMPLEMENT_HIT_PROXY(HNewLandscapeGrabHandleProxy, HHitProxy)
-
 void FLandscapeTool::SetEditRenderType()
 {
 	GLandscapeEditRenderMode = ELandscapeEditRenderMode::SelectRegion | (GLandscapeEditRenderMode & ELandscapeEditRenderMode::BitMaskForMask);
@@ -125,8 +90,7 @@ FEdModeLandscape::FEdModeLandscape()
 	, CurrentToolIndex(INDEX_NONE)
 	, CurrentBrushSetIndex(0)
 	, NewLandscapePreviewMode(ENewLandscapePreviewMode::None)
-	, DraggingEdge(ELandscapeEdge::None)
-	, DraggingEdge_Remainder(0)
+	, ImportExportMode(EImportExportMode::Import)
 	, CurrentGizmoActor(nullptr)
 	, CopyPasteTool(nullptr)
 	, SplinesTool(nullptr)
@@ -161,6 +125,7 @@ FEdModeLandscape::FEdModeLandscape()
 	InitializeTool_Retopologize();
 	InitializeTool_NewLandscape();
 	InitializeTool_ResizeLandscape();
+	InitializeTool_ImportExport();
 	InitializeTool_Select();
 	InitializeTool_AddComponent();
 	InitializeTool_DeleteComponent();
@@ -261,6 +226,7 @@ void FEdModeLandscape::UpdateToolModes()
 	ToolMode_Manage->ValidTools.Add(TEXT("MoveToLevel"));
 	ToolMode_Manage->ValidTools.Add(TEXT("ResizeLandscape"));
 	ToolMode_Manage->ValidTools.Add(TEXT("Splines"));
+	ToolMode_Manage->ValidTools.Add(TEXT("ImportExport"));
 	
 	// Restore
 	FName* PreviousToolName = PreviousTools.Find(ToolMode_Manage->ToolModeName);
@@ -1672,42 +1638,6 @@ bool FEdModeLandscape::InputKey(FEditorViewportClient* ViewportClient, FViewport
 		}
 	}
 
-	if (NewLandscapePreviewMode != ENewLandscapePreviewMode::None)
-	{
-		if (Key == EKeys::LeftMouseButton)
-		{
-			// Press mouse button
-			if (Event == IE_Pressed && !IsAltDown(Viewport))
-			{
-				// See if we clicked on a new landscape handle..
-				int32 HitX = Viewport->GetMouseX();
-				int32 HitY = Viewport->GetMouseY();
-				HHitProxy*	HitProxy = Viewport->GetHitProxy(HitX, HitY);
-				if (HitProxy)
-				{
-					if (HitProxy->IsA(HNewLandscapeGrabHandleProxy::StaticGetType()))
-					{
-						HNewLandscapeGrabHandleProxy* EdgeProxy = (HNewLandscapeGrabHandleProxy*)HitProxy;
-						DraggingEdge = EdgeProxy->Edge;
-						DraggingEdge_Remainder = 0;
-
-						return false; // false to let FEditorViewportClient.InputKey start mouse tracking and enable InputDelta() so we can use it
-					}
-				}
-			}
-			else if (Event == IE_Released)
-			{
-				if (DraggingEdge)
-				{
-					DraggingEdge = ELandscapeEdge::None;
-					DraggingEdge_Remainder = 0;
-
-					return false; // false to let FEditorViewportClient.InputKey end mouse tracking
-				}
-			}
-		}
-	}
-	else
 	{
 		// Override Key Input for Selection Brush
 		if (CurrentBrush)
@@ -1852,95 +1782,6 @@ bool FEdModeLandscape::InputDelta(FEditorViewportClient* InViewportClient, FView
 		return false;
 	}
 
-	if (NewLandscapePreviewMode != ENewLandscapePreviewMode::None)
-	{
-		if (InViewportClient->GetCurrentWidgetAxis() != EAxisList::None)
-		{
-			FVector DeltaScale = InScale;
-			DeltaScale.X = DeltaScale.Y = (FMath::Abs(InScale.X) > FMath::Abs(InScale.Y)) ? InScale.X : InScale.Y;
-
-			UISettings->Modify();
-			UISettings->NewLandscape_Location += InDrag;
-			UISettings->NewLandscape_Rotation += InRot;
-			UISettings->NewLandscape_Scale += DeltaScale;
-
-			return true;
-		}
-		else if (DraggingEdge != ELandscapeEdge::None)
-		{
-			FVector HitLocation;
-			LandscapePlaneTrace(InViewportClient, FPlane(UISettings->NewLandscape_Location, FVector(0, 0, 1)), HitLocation);
-
-			FTransform Transform(UISettings->NewLandscape_Rotation, UISettings->NewLandscape_Location, UISettings->NewLandscape_Scale * UISettings->NewLandscape_QuadsPerSection * UISettings->NewLandscape_SectionsPerComponent);
-			HitLocation = Transform.InverseTransformPosition(HitLocation);
-
-			UISettings->Modify();
-			switch (DraggingEdge)
-			{
-			case ELandscapeEdge::X_Negative:
-			case ELandscapeEdge::X_Negative_Y_Negative:
-			case ELandscapeEdge::X_Negative_Y_Positive:
-			{
-				const int32 InitialComponentCountX = UISettings->NewLandscape_ComponentCount.X;
-				const int32 Delta = FMath::RoundToInt(HitLocation.X + (float)InitialComponentCountX / 2);
-				UISettings->NewLandscape_ComponentCount.X = InitialComponentCountX - Delta;
-				UISettings->NewLandscape_ClampSize();
-				const int32 ActualDelta = UISettings->NewLandscape_ComponentCount.X - InitialComponentCountX;
-				UISettings->NewLandscape_Location -= Transform.TransformVector(FVector(((float)ActualDelta / 2), 0, 0));
-			}
-				break;
-			case ELandscapeEdge::X_Positive:
-			case ELandscapeEdge::X_Positive_Y_Negative:
-			case ELandscapeEdge::X_Positive_Y_Positive:
-			{
-				const int32 InitialComponentCountX = UISettings->NewLandscape_ComponentCount.X;
-				int32 Delta = FMath::RoundToInt(HitLocation.X - (float)InitialComponentCountX / 2);
-				UISettings->NewLandscape_ComponentCount.X = InitialComponentCountX + Delta;
-				UISettings->NewLandscape_ClampSize();
-				const int32 ActualDelta = UISettings->NewLandscape_ComponentCount.X - InitialComponentCountX;
-				UISettings->NewLandscape_Location += Transform.TransformVector(FVector(((float)ActualDelta / 2), 0, 0));
-			}
-				break;
-			case  ELandscapeEdge::Y_Negative:
-			case  ELandscapeEdge::Y_Positive:
-				break;
-			}
-
-			switch (DraggingEdge)
-			{
-			case ELandscapeEdge::Y_Negative:
-			case ELandscapeEdge::X_Negative_Y_Negative:
-			case ELandscapeEdge::X_Positive_Y_Negative:
-			{
-				const int32 InitialComponentCountY = UISettings->NewLandscape_ComponentCount.Y;
-				int32 Delta = FMath::RoundToInt(HitLocation.Y + (float)InitialComponentCountY / 2);
-				UISettings->NewLandscape_ComponentCount.Y = InitialComponentCountY - Delta;
-				UISettings->NewLandscape_ClampSize();
-				const int32 ActualDelta = UISettings->NewLandscape_ComponentCount.Y - InitialComponentCountY;
-				UISettings->NewLandscape_Location -= Transform.TransformVector(FVector(0, (float)ActualDelta / 2, 0));
-			}
-				break;
-			case ELandscapeEdge::Y_Positive:
-			case ELandscapeEdge::X_Negative_Y_Positive:
-			case ELandscapeEdge::X_Positive_Y_Positive:
-			{
-				const int32 InitialComponentCountY = UISettings->NewLandscape_ComponentCount.Y;
-				int32 Delta = FMath::RoundToInt(HitLocation.Y - (float)InitialComponentCountY / 2);
-				UISettings->NewLandscape_ComponentCount.Y = InitialComponentCountY + Delta;
-				UISettings->NewLandscape_ClampSize();
-				const int32 ActualDelta = UISettings->NewLandscape_ComponentCount.Y - InitialComponentCountY;
-				UISettings->NewLandscape_Location += Transform.TransformVector(FVector(0, (float)ActualDelta / 2, 0));
-			}
-				break;
-			case  ELandscapeEdge::X_Negative:
-			case  ELandscapeEdge::X_Positive:
-				break;
-			}
-
-			return true;
-		}
-	}
-
 	if (CurrentTool && CurrentTool->InputDelta(InViewportClient, InViewport, InDrag, InRot, InScale))
 	{
 		return true;
@@ -2838,205 +2679,6 @@ void FEdModeLandscape::Render(const FSceneView* View, FViewport* Viewport, FPrim
 		return;
 	}
 
-	if (NewLandscapePreviewMode != ENewLandscapePreviewMode::None)
-	{
-		static const float        CornerSize = 0.33f;
-		static const FLinearColor CornerColour(1.0f, 1.0f, 0.5f);
-		static const FLinearColor EdgeColour(1.0f, 1.0f, 0.0f);
-		static const FLinearColor ComponentBorderColour(0.0f, 0.85f, 0.0f);
-		static const FLinearColor SectionBorderColour(0.0f, 0.4f, 0.0f);
-		static const FLinearColor InnerColour(0.0f, 0.25f, 0.0f);
-
-		const ELevelViewportType ViewportType = ((FEditorViewportClient*)Viewport->GetClient())->ViewportType;
-
-		const int32 ComponentCountX = UISettings->NewLandscape_ComponentCount.X;
-		const int32 ComponentCountY = UISettings->NewLandscape_ComponentCount.Y;
-		const int32 QuadsPerComponent = UISettings->NewLandscape_SectionsPerComponent * UISettings->NewLandscape_QuadsPerSection;
-		const float ComponentSize = QuadsPerComponent;
-		const FVector Offset = UISettings->NewLandscape_Location + FTransform(UISettings->NewLandscape_Rotation, FVector::ZeroVector, UISettings->NewLandscape_Scale).TransformVector(FVector(-ComponentCountX * ComponentSize / 2, -ComponentCountY * ComponentSize / 2, 0));
-		const FTransform Transform = FTransform(UISettings->NewLandscape_Rotation, Offset, UISettings->NewLandscape_Scale);
-
-		if (NewLandscapePreviewMode == ENewLandscapePreviewMode::ImportLandscape)
-		{
-			const TArray<uint16>& ImportHeights = UISettings->GetImportLandscapeData();
-			if (ImportHeights.Num() != 0)
-			{
-				const float InvQuadsPerComponent = 1.0f / (float)QuadsPerComponent;
-				const int32 SizeX = ComponentCountX * QuadsPerComponent + 1;
-				const int32 SizeY = ComponentCountY * QuadsPerComponent + 1;
-				const int32 ImportSizeX = UISettings->ImportLandscape_Width;
-				const int32 ImportSizeY = UISettings->ImportLandscape_Height;
-				const int32 OffsetX = (SizeX - ImportSizeX) / 2;
-				const int32 OffsetY = (SizeY - ImportSizeY) / 2;
-
-				for (int32 ComponentY = 0; ComponentY < ComponentCountY; ComponentY++)
-				{
-					const int32 Y0 = ComponentY * QuadsPerComponent;
-					const int32 Y1 = (ComponentY + 1) * QuadsPerComponent;
-
-					const int32 ImportY0 = FMath::Clamp<int32>(Y0 - OffsetY, 0, ImportSizeY - 1);
-					const int32 ImportY1 = FMath::Clamp<int32>(Y1 - OffsetY, 0, ImportSizeY - 1);
-
-					for (int32 ComponentX = 0; ComponentX < ComponentCountX; ComponentX++)
-					{
-						const int32 X0 = ComponentX * QuadsPerComponent;
-						const int32 X1 = (ComponentX + 1) * QuadsPerComponent;
-						const int32 ImportX0 = FMath::Clamp<int32>(X0 - OffsetX, 0, ImportSizeX - 1);
-						const int32 ImportX1 = FMath::Clamp<int32>(X1 - OffsetX, 0, ImportSizeX - 1);
-						const float Z00 = ((float)ImportHeights[ImportX0 + ImportY0 * ImportSizeX] - 32768.0f) * LANDSCAPE_ZSCALE;
-						const float Z01 = ((float)ImportHeights[ImportX0 + ImportY1 * ImportSizeX] - 32768.0f) * LANDSCAPE_ZSCALE;
-						const float Z10 = ((float)ImportHeights[ImportX1 + ImportY0 * ImportSizeX] - 32768.0f) * LANDSCAPE_ZSCALE;
-						const float Z11 = ((float)ImportHeights[ImportX1 + ImportY1 * ImportSizeX] - 32768.0f) * LANDSCAPE_ZSCALE;
-
-						if (ComponentX == 0)
-						{
-							PDI->SetHitProxy(new HNewLandscapeGrabHandleProxy(ELandscapeEdge::X_Negative));
-							PDI->DrawLine(Transform.TransformPosition(FVector(X0, Y0, Z00)), Transform.TransformPosition(FVector(X0, Y1, Z01)), ComponentBorderColour, SDPG_Foreground);
-							PDI->SetHitProxy(NULL);
-						}
-
-						if (ComponentX == ComponentCountX - 1)
-						{
-							PDI->SetHitProxy(new HNewLandscapeGrabHandleProxy(ELandscapeEdge::X_Positive));
-							PDI->DrawLine(Transform.TransformPosition(FVector(X1, Y0, Z10)), Transform.TransformPosition(FVector(X1, Y1, Z11)), ComponentBorderColour, SDPG_Foreground);
-							PDI->SetHitProxy(NULL);
-						}
-						else
-						{
-							PDI->DrawLine(Transform.TransformPosition(FVector(X1, Y0, Z10)), Transform.TransformPosition(FVector(X1, Y1, Z11)), ComponentBorderColour, SDPG_Foreground);
-						}
-
-						if (ComponentY == 0)
-						{
-							PDI->SetHitProxy(new HNewLandscapeGrabHandleProxy(ELandscapeEdge::Y_Negative));
-							PDI->DrawLine(Transform.TransformPosition(FVector(X0, Y0, Z00)), Transform.TransformPosition(FVector(X1, Y0, Z10)), ComponentBorderColour, SDPG_Foreground);
-							PDI->SetHitProxy(NULL);
-						}
-
-						if (ComponentY == ComponentCountY - 1)
-						{
-							PDI->SetHitProxy(new HNewLandscapeGrabHandleProxy(ELandscapeEdge::Y_Positive));
-							PDI->DrawLine(Transform.TransformPosition(FVector(X0, Y1, Z01)), Transform.TransformPosition(FVector(X1, Y1, Z11)), ComponentBorderColour, SDPG_Foreground);
-							PDI->SetHitProxy(NULL);
-						}
-						else
-						{
-							PDI->DrawLine(Transform.TransformPosition(FVector(X0, Y1, Z01)), Transform.TransformPosition(FVector(X1, Y1, Z11)), ComponentBorderColour, SDPG_Foreground);
-						}
-
-						// intra-component lines - too slow for big landscapes
-						/*
-						for (int32 x=1;x<QuadsPerComponent;x++)
-						{
-						PDI->DrawLine(Transform.TransformPosition(FVector(X0+x, Y0, FMath::Lerp(Z00,Z10,(float)x*InvQuadsPerComponent))), Transform.TransformPosition(FVector(X0+x, Y1, FMath::Lerp(Z01,Z11,(float)x*InvQuadsPerComponent))), ComponentBorderColour, SDPG_World);
-						}
-						for (int32 y=1;y<QuadsPerComponent;y++)
-						{
-						PDI->DrawLine(Transform.TransformPosition(FVector(X0, Y0+y, FMath::Lerp(Z00,Z01,(float)y*InvQuadsPerComponent))), Transform.TransformPosition(FVector(X1, Y0+y, FMath::Lerp(Z10,Z11,(float)y*InvQuadsPerComponent))), ComponentBorderColour, SDPG_World);
-						}
-						*/
-					}
-				}
-			}
-		}
-		else //if (NewLandscapePreviewMode == ENewLandscapePreviewMode::NewLandscape)
-		{
-			if (ViewportType == LVT_Perspective || ViewportType == LVT_OrthoXY || ViewportType == LVT_OrthoNegativeXY)
-			{
-				for (int32 x = 0; x <= ComponentCountX * QuadsPerComponent; x++)
-				{
-					if (x == 0)
-					{
-						PDI->SetHitProxy(new HNewLandscapeGrabHandleProxy(ELandscapeEdge::X_Negative_Y_Negative));
-						PDI->DrawLine(Transform.TransformPosition(FVector(x, 0, 0)), Transform.TransformPosition(FVector(x, CornerSize * ComponentSize, 0)), CornerColour, SDPG_Foreground);
-						PDI->SetHitProxy(new HNewLandscapeGrabHandleProxy(ELandscapeEdge::X_Negative));
-						PDI->DrawLine(Transform.TransformPosition(FVector(x, CornerSize * ComponentSize, 0)), Transform.TransformPosition(FVector(x, (ComponentCountY - CornerSize) * ComponentSize, 0)), EdgeColour, SDPG_Foreground);
-						PDI->SetHitProxy(new HNewLandscapeGrabHandleProxy(ELandscapeEdge::X_Negative_Y_Positive));
-						PDI->DrawLine(Transform.TransformPosition(FVector(x, (ComponentCountY - CornerSize) * ComponentSize, 0)), Transform.TransformPosition(FVector(x, ComponentCountY * ComponentSize, 0)), CornerColour, SDPG_Foreground);
-						PDI->SetHitProxy(NULL);
-					}
-					else if (x == ComponentCountX * QuadsPerComponent)
-					{
-						PDI->SetHitProxy(new HNewLandscapeGrabHandleProxy(ELandscapeEdge::X_Positive_Y_Negative));
-						PDI->DrawLine(Transform.TransformPosition(FVector(x, 0, 0)), Transform.TransformPosition(FVector(x, CornerSize * ComponentSize, 0)), CornerColour, SDPG_Foreground);
-						PDI->SetHitProxy(new HNewLandscapeGrabHandleProxy(ELandscapeEdge::X_Positive));
-						PDI->DrawLine(Transform.TransformPosition(FVector(x, CornerSize * ComponentSize, 0)), Transform.TransformPosition(FVector(x, (ComponentCountY - CornerSize) * ComponentSize, 0)), EdgeColour, SDPG_Foreground);
-						PDI->SetHitProxy(new HNewLandscapeGrabHandleProxy(ELandscapeEdge::X_Positive_Y_Positive));
-						PDI->DrawLine(Transform.TransformPosition(FVector(x, (ComponentCountY - CornerSize) * ComponentSize, 0)), Transform.TransformPosition(FVector(x, ComponentCountY * ComponentSize, 0)), CornerColour, SDPG_Foreground);
-						PDI->SetHitProxy(NULL);
-					}
-					else if (x % QuadsPerComponent == 0)
-					{
-						PDI->DrawLine(Transform.TransformPosition(FVector(x, 0, 0)), Transform.TransformPosition(FVector(x, ComponentCountY * ComponentSize, 0)), ComponentBorderColour, SDPG_Foreground);
-					}
-					else if (x % UISettings->NewLandscape_QuadsPerSection == 0)
-					{
-						PDI->DrawLine(Transform.TransformPosition(FVector(x, 0, 0)), Transform.TransformPosition(FVector(x, ComponentCountY * ComponentSize, 0)), SectionBorderColour, SDPG_Foreground);
-					}
-					else
-					{
-						PDI->DrawLine(Transform.TransformPosition(FVector(x, 0, 0)), Transform.TransformPosition(FVector(x, ComponentCountY * ComponentSize, 0)), InnerColour, SDPG_World);
-					}
-				}
-			}
-			else
-			{
-				// Don't allow dragging to resize in side-view
-				// and there's no point drawing the inner lines as only the outer is visible
-				PDI->DrawLine(Transform.TransformPosition(FVector(0, 0, 0)), Transform.TransformPosition(FVector(0, ComponentCountY * ComponentSize, 0)), EdgeColour, SDPG_World);
-				PDI->DrawLine(Transform.TransformPosition(FVector(ComponentCountX * QuadsPerComponent, 0, 0)), Transform.TransformPosition(FVector(ComponentCountX * QuadsPerComponent, ComponentCountY * ComponentSize, 0)), EdgeColour, SDPG_World);
-			}
-
-			if (ViewportType == LVT_Perspective || ViewportType == LVT_OrthoXY || ViewportType == LVT_OrthoNegativeXY)
-			{
-				for (int32 y = 0; y <= ComponentCountY * QuadsPerComponent; y++)
-				{
-					if (y == 0)
-					{
-						PDI->SetHitProxy(new HNewLandscapeGrabHandleProxy(ELandscapeEdge::X_Negative_Y_Negative));
-						PDI->DrawLine(Transform.TransformPosition(FVector(0, y, 0)), Transform.TransformPosition(FVector(CornerSize * ComponentSize, y, 0)), CornerColour, SDPG_Foreground);
-						PDI->SetHitProxy(new HNewLandscapeGrabHandleProxy(ELandscapeEdge::Y_Negative));
-						PDI->DrawLine(Transform.TransformPosition(FVector(CornerSize * ComponentSize, y, 0)), Transform.TransformPosition(FVector((ComponentCountX - CornerSize) * ComponentSize, y, 0)), EdgeColour, SDPG_Foreground);
-						PDI->SetHitProxy(new HNewLandscapeGrabHandleProxy(ELandscapeEdge::X_Positive_Y_Negative));
-						PDI->DrawLine(Transform.TransformPosition(FVector((ComponentCountX - CornerSize) * ComponentSize, y, 0)), Transform.TransformPosition(FVector(ComponentCountX * ComponentSize, y, 0)), CornerColour, SDPG_Foreground);
-						PDI->SetHitProxy(NULL);
-					}
-					else if (y == ComponentCountY * QuadsPerComponent)
-					{
-						PDI->SetHitProxy(new HNewLandscapeGrabHandleProxy(ELandscapeEdge::X_Negative_Y_Positive));
-						PDI->DrawLine(Transform.TransformPosition(FVector(0, y, 0)), Transform.TransformPosition(FVector(CornerSize * ComponentSize, y, 0)), CornerColour, SDPG_Foreground);
-						PDI->SetHitProxy(new HNewLandscapeGrabHandleProxy(ELandscapeEdge::Y_Positive));
-						PDI->DrawLine(Transform.TransformPosition(FVector(CornerSize * ComponentSize, y, 0)), Transform.TransformPosition(FVector((ComponentCountX - CornerSize) * ComponentSize, y, 0)), EdgeColour, SDPG_Foreground);
-						PDI->SetHitProxy(new HNewLandscapeGrabHandleProxy(ELandscapeEdge::X_Positive_Y_Positive));
-						PDI->DrawLine(Transform.TransformPosition(FVector((ComponentCountX - CornerSize) * ComponentSize, y, 0)), Transform.TransformPosition(FVector(ComponentCountX * ComponentSize, y, 0)), CornerColour, SDPG_Foreground);
-						PDI->SetHitProxy(NULL);
-					}
-					else if (y % QuadsPerComponent == 0)
-					{
-						PDI->DrawLine(Transform.TransformPosition(FVector(0, y, 0)), Transform.TransformPosition(FVector(ComponentCountX * ComponentSize, y, 0)), ComponentBorderColour, SDPG_Foreground);
-					}
-					else if (y % UISettings->NewLandscape_QuadsPerSection == 0)
-					{
-						PDI->DrawLine(Transform.TransformPosition(FVector(0, y, 0)), Transform.TransformPosition(FVector(ComponentCountX * ComponentSize, y, 0)), SectionBorderColour, SDPG_Foreground);
-					}
-					else
-					{
-						PDI->DrawLine(Transform.TransformPosition(FVector(0, y, 0)), Transform.TransformPosition(FVector(ComponentCountX * ComponentSize, y, 0)), InnerColour, SDPG_World);
-					}
-				}
-			}
-			else
-			{
-				// Don't allow dragging to resize in side-view
-				// and there's no point drawing the inner lines as only the outer is visible
-				PDI->DrawLine(Transform.TransformPosition(FVector(0, 0, 0)), Transform.TransformPosition(FVector(ComponentCountX * ComponentSize, 0, 0)), EdgeColour, SDPG_World);
-				PDI->DrawLine(Transform.TransformPosition(FVector(0, ComponentCountY * QuadsPerComponent, 0)), Transform.TransformPosition(FVector(ComponentCountX * ComponentSize, ComponentCountY * QuadsPerComponent, 0)), EdgeColour, SDPG_World);
-			}
-		}
-
-		return;
-	}
-
 	if (LandscapeRenderAddCollision)
 	{
 		PDI->DrawLine(LandscapeRenderAddCollision->Corners[0], LandscapeRenderAddCollision->Corners[3], FColor(0, 255, 128), SDPG_Foreground);
@@ -3290,17 +2932,15 @@ void FEdModeLandscape::ReimportData(const FLandscapeTargetListInfo& TargetInfo)
 }
 
 template<class T>
-void ImportDataInternal(const FLandscapeTargetListInfo& TargetInfo, const FString& Filename, bool bSingleFile, TFunctionRef<void(int32, int32, int32, int32, const TArray<T>&)> SetDataFunc)
+void ImportDataInternal(ULandscapeInfo* LandscapeInfo, const FString& Filename, FName LayerName, bool bSingleFile, const FIntRect& ImportRegionVerts, ELandscapeImportTransformType TransformType, FIntPoint Offset, TFunctionRef<void(int32, int32, int32, int32, const TArray<T>&)> SetDataFunc)
 {
-	ULandscapeInfo* LandscapeInfo = TargetInfo.LandscapeInfo.Get();
-	int32 MinX, MinY, MaxX, MaxY;
-	if (LandscapeInfo && LandscapeInfo->GetLandscapeExtent(MinX, MinY, MaxX, MaxY))
+	if (LandscapeInfo)
 	{
-		const FLandscapeImportResolution LandscapeResolution = { (uint32)(1 + MaxX - MinX), (uint32)(1 + MaxY - MinY) };
+		const FLandscapeImportResolution LandscapeResolution = { (uint32)(ImportRegionVerts.Width()), (uint32)(ImportRegionVerts.Height()) };
 		FLandscapeImportDescriptor OutImportDescriptor;
 		FText OutMessage;
 
-		ELandscapeImportResult ImportResult = FLandscapeImportHelper::GetImportDescriptor<T>(Filename, bSingleFile, TargetInfo.LayerName, OutImportDescriptor, OutMessage);
+		ELandscapeImportResult ImportResult = FLandscapeImportHelper::GetImportDescriptor<T>(Filename, bSingleFile, LayerName, OutImportDescriptor, OutMessage);
 		if (ImportResult == ELandscapeImportResult::Error)
 		{
 			FMessageDialog::Open(EAppMsgType::Ok, OutMessage);
@@ -3363,7 +3003,7 @@ void ImportDataInternal(const FLandscapeTargetListInfo& TargetInfo, const FStrin
 		}
 
 		TArray<T> ImportData;
-		ImportResult = FLandscapeImportHelper::GetImportData<T>(OutImportDescriptor, DescriptorIndex, TargetInfo.LayerName, ImportData, OutMessage);
+		ImportResult = FLandscapeImportHelper::GetImportData<T>(OutImportDescriptor, DescriptorIndex, LayerName, ImportData, OutMessage);
 		if (ImportResult == ELandscapeImportResult::Error)
 		{
 			FMessageDialog::Open(EAppMsgType::Ok, OutMessage);
@@ -3375,50 +3015,71 @@ void ImportDataInternal(const FLandscapeTargetListInfo& TargetInfo, const FStrin
 			TArray<T> FinalData;
 			if (bResolutionMismatch)
 			{
-				FLandscapeImportHelper::ExpandImportData<T>(ImportData, FinalData, ImportResolution, LandscapeResolution);
+				FLandscapeImportHelper::TransformImportData<T>(ImportData, FinalData, ImportResolution, LandscapeResolution, TransformType, Offset);
 			}
 			else
 			{
 				FinalData = MoveTemp(ImportData);
 			}
 									
-			SetDataFunc(MinX, MinY, MaxX, MaxY, FinalData);
+			// Set Data is in Quads (so remove 1)
+			SetDataFunc(ImportRegionVerts.Min.X, ImportRegionVerts.Min.Y, ImportRegionVerts.Max.X-1, ImportRegionVerts.Max.Y-1, FinalData);
 		}
 	}			
 }
 
-void FEdModeLandscape::ImportData(const FLandscapeTargetListInfo& TargetInfo, const FString& Filename)
+void FEdModeLandscape::ImportHeightData(ULandscapeInfo* LandscapeInfo, const FGuid& LayerGuid, const FString& Filename, const FIntRect& ImportRegionVerts, ELandscapeImportTransformType TransformType, FIntPoint Offset, const ELandscapeLayerPaintingRestriction& PaintRestriction)
 {
-	if (TargetInfo.TargetType == ELandscapeToolTargetType::Heightmap)
+	ImportDataInternal<uint16>(LandscapeInfo, Filename, NAME_None, UseSingleFileImport(), ImportRegionVerts, TransformType, Offset, [LandscapeInfo, LayerGuid, PaintRestriction](int32 MinX, int32 MinY, int32 MaxX, int32 MaxY, const TArray<uint16>& Data)
 	{
-		ImportDataInternal<uint16>(TargetInfo, Filename, UseSingleFileImport(), [this, &TargetInfo](int32 MinX, int32 MinY, int32 MaxX, int32 MaxY, const TArray<uint16>& Data)
-		{
-			ULandscapeInfo* LandscapeInfo = TargetInfo.LandscapeInfo.Get();
-			ALandscape* Landscape = LandscapeInfo->LandscapeActor.Get();
-			FScopedSetLandscapeEditingLayer Scope(Landscape, GetCurrentLayerGuid(), [&] { check(Landscape); Landscape->RequestLayersContentUpdate(ELandscapeLayerUpdateMode::Update_Heightmap_All); });
+		ALandscape* Landscape = LandscapeInfo->LandscapeActor.Get();
+		FScopedSetLandscapeEditingLayer Scope(Landscape, LayerGuid, [&] { check(Landscape); Landscape->RequestLayersContentUpdate(ELandscapeLayerUpdateMode::Update_Heightmap_All); });
 
-			FScopedTransaction Transaction(LOCTEXT("Undo_ImportHeightmap", "Importing Landscape Heightmap"));
+		FScopedTransaction Transaction(LOCTEXT("Undo_ImportHeightmap", "Importing Landscape Heightmap"));
 
-			FHeightmapAccessor<false> HeightmapAccessor(LandscapeInfo);
-			HeightmapAccessor.SetData(MinX, MinY, MaxX, MaxY, Data.GetData());
-		});
-	}
-	else
+		FHeightmapAccessor<false> HeightmapAccessor(LandscapeInfo);
+		HeightmapAccessor.SetData(MinX, MinY, MaxX, MaxY, Data.GetData(), PaintRestriction);
+	});
+}
+
+void FEdModeLandscape::ImportWeightData(ULandscapeInfo* LandscapeInfo, const FGuid& LayerGuid, ULandscapeLayerInfoObject* LayerInfo, const FString& Filename, const FIntRect& ImportRegionVerts, ELandscapeImportTransformType TransformType, FIntPoint Offset, const ELandscapeLayerPaintingRestriction& PaintRestriction)
+{
+	if (LayerInfo)
 	{
-		ImportDataInternal<uint8>(TargetInfo, Filename, UseSingleFileImport(), [this, &TargetInfo](int32 MinX, int32 MinY, int32 MaxX, int32 MaxY, const TArray<uint8>& Data)
+		ImportDataInternal<uint8>(LandscapeInfo, Filename, LayerInfo->LayerName, UseSingleFileImport(), ImportRegionVerts, TransformType, Offset, [LandscapeInfo, LayerGuid, LayerInfo, PaintRestriction](int32 MinX, int32 MinY, int32 MaxX, int32 MaxY, const TArray<uint8>& Data)
 		{
-			ULandscapeInfo* LandscapeInfo = TargetInfo.LandscapeInfo.Get();
 			ALandscape* Landscape = LandscapeInfo->LandscapeActor.Get();
-			FScopedSetLandscapeEditingLayer Scope(Landscape, GetCurrentLayerGuid(), [&] { check(Landscape); Landscape->RequestLayersContentUpdate(ELandscapeLayerUpdateMode::Update_Weightmap_All); });
+			FScopedSetLandscapeEditingLayer Scope(Landscape, LayerGuid, [&] { check(Landscape); Landscape->RequestLayersContentUpdate(ELandscapeLayerUpdateMode::Update_Weightmap_All); });
 
 			FScopedTransaction Transaction(LOCTEXT("Undo_ImportWeightmap", "Importing Landscape Layer"));
 
-			FAlphamapAccessor<false, false> AlphamapAccessor(LandscapeInfo, TargetInfo.LayerInfoObj.Get());
-			AlphamapAccessor.SetData(MinX, MinY, MaxX, MaxY, Data.GetData(), ELandscapeLayerPaintingRestriction::None);
+			FAlphamapAccessor<false, false> AlphamapAccessor(LandscapeInfo, LayerInfo);
+			AlphamapAccessor.SetData(MinX, MinY, MaxX, MaxY, Data.GetData(), PaintRestriction);
 		});
 	}
 }
 
+void FEdModeLandscape::ImportData(const FLandscapeTargetListInfo& TargetInfo, const FString& Filename)
+{
+	ULandscapeInfo* LandscapeInfo = TargetInfo.LandscapeInfo.Get();
+	FIntRect ImportExtent;
+	if (LandscapeInfo && LandscapeInfo->GetLandscapeExtent(ImportExtent))
+	{
+		// Import is in Verts (Extent is in Quads)
+		ImportExtent.Max.X += 1;
+		ImportExtent.Max.Y += 1;
+
+		if (TargetInfo.TargetType == ELandscapeToolTargetType::Heightmap)
+		{
+			ImportHeightData(LandscapeInfo, GetCurrentLayerGuid(), Filename, ImportExtent, ELandscapeImportTransformType::ExpandCentered);
+		}
+		else
+		{
+			ImportWeightData(LandscapeInfo, GetCurrentLayerGuid(), TargetInfo.LayerInfoObj.Get(), Filename, ImportExtent, ELandscapeImportTransformType::ExpandCentered);
+		}
+	}
+}
+
 void FEdModeLandscape::DeleteLandscapeComponents(ULandscapeInfo* LandscapeInfo, TSet<ULandscapeComponent*> ComponentsToDelete)
 {
 	LandscapeInfo->Modify();
@@ -3650,8 +3311,14 @@ ALandscape* FEdModeLandscape::ChangeComponentSetting(int32 NumComponentsX, int32
 						int32 TMinX = OldMinX, TMinY = OldMinY, TMaxX = OldMaxX, TMaxY = OldMaxY;
 						LandscapeEdit.GetHeightData(TMinX, TMinY, TMaxX, TMaxY, OutHeightData.GetData(), 0);
 
-						OutHeightData = LandscapeEditorUtils::ResampleData(OutHeightData, OldVertsX, OldVertsY, NewVertsX, NewVertsY);
+						TArray<uint16> ResampledHeightData;
+
+						FIntRect SrcRegion(0, 0, OldVertsX, OldVertsY);
+						FIntRect DestRegion(0, 0, NewVertsX, NewVertsY);
+						FLandscapeConfigHelper::ResampleData(OutHeightData, ResampledHeightData, SrcRegion, DestRegion);
+						OutHeightData = MoveTemp(ResampledHeightData);
 
+						TArray<uint8> ResampledWeightData;
 						for (const FLandscapeInfoLayerSettings& LayerSettings : LandscapeInfo->Layers)
 						{
 							if (LayerSettings.LayerInfoObj != NULL)
@@ -3662,7 +3329,8 @@ ALandscape* FEdModeLandscape::ChangeComponentSetting(int32 NumComponentsX, int32
 								TMinX = OldMinX; TMinY = OldMinY; TMaxX = OldMaxX; TMaxY = OldMaxY;
 								LandscapeEdit.GetWeightData(LayerSettings.LayerInfoObj, TMinX, TMinY, TMaxX, TMaxY, ImportLayerInfo->LayerData.GetData(), 0);
 
-								ImportLayerInfo->LayerData = LandscapeEditorUtils::ResampleData(ImportLayerInfo->LayerData, OldVertsX, OldVertsY, NewVertsX, NewVertsY);
+								FLandscapeConfigHelper::ResampleData(ImportLayerInfo->LayerData, ResampledWeightData, SrcRegion, DestRegion);
+								ImportLayerInfo->LayerData = MoveTemp(ResampledWeightData);
 							}
 						}
 					}
@@ -3681,9 +3349,14 @@ ALandscape* FEdModeLandscape::ChangeComponentSetting(int32 NumComponentsX, int32
 						// GetHeightData alters its args, so make temp copies to avoid screwing things up
 						int32 TMinX = RequestedMinX, TMinY = RequestedMinY, TMaxX = RequestedMaxX, TMaxY = RequestedMaxY;
 						LandscapeEdit.GetHeightData(TMinX, TMinY, TMaxX, OldMaxY, OutHeightData.GetData(), 0);
+						
+						FIntRect SrcRegion(RequestedMinX, RequestedMinY, RequestedMaxX, RequestedMaxY);
+						FIntRect DestRegion(NewMinX, NewMinY, NewMaxX, NewMaxY);
+						TArray<uint16> ExpandedHeightData;
+						FLandscapeConfigHelper::ExpandData(OutHeightData, ExpandedHeightData, SrcRegion, DestRegion, true);
+						OutHeightData = MoveTemp(ExpandedHeightData);
 
-						OutHeightData = LandscapeEditorUtils::ExpandData(OutHeightData, RequestedMinX, RequestedMinY, RequestedMaxX, RequestedMaxY, NewMinX, NewMinY, NewMaxX, NewMaxY);
-
+						TArray<uint8> ExpandedWeightData;
 						for (const FLandscapeInfoLayerSettings& LayerSettings : LandscapeInfo->Layers)
 						{
 							if (LayerSettings.LayerInfoObj != NULL)
@@ -3694,7 +3367,8 @@ ALandscape* FEdModeLandscape::ChangeComponentSetting(int32 NumComponentsX, int32
 								TMinX = RequestedMinX; TMinY = RequestedMinY; TMaxX = RequestedMaxX; TMaxY = RequestedMaxY;
 								LandscapeEdit.GetWeightData(LayerSettings.LayerInfoObj, TMinX, TMinY, TMaxX, TMaxY, ImportLayerInfo->LayerData.GetData(), 0);
 
-								ImportLayerInfo->LayerData = LandscapeEditorUtils::ExpandData(ImportLayerInfo->LayerData, RequestedMinX, RequestedMinY, RequestedMaxX, RequestedMaxY, NewMinX, NewMinY, NewMaxX, NewMaxY);
+								FLandscapeConfigHelper::ExpandData(ImportLayerInfo->LayerData, ExpandedWeightData, SrcRegion, DestRegion, true);
+								ImportLayerInfo->LayerData = MoveTemp(ExpandedWeightData);
 							}
 						}
 					}
diff --git a/Engine/Source/Editor/LandscapeEditor/Private/LandscapeEdMode.h b/Engine/Source/Editor/LandscapeEditor/Private/LandscapeEdMode.h
--- a/Engine/Source/Editor/LandscapeEditor/Private/LandscapeEdMode.h
+++ b/Engine/Source/Editor/LandscapeEditor/Private/LandscapeEdMode.h
@@ -14,6 +14,7 @@
 #include "LandscapeGizmoActiveActor.h"
 #include "LandscapeEdit.h"
 #include "Containers/Set.h"
+#include "LandscapeImportHelper.h"
 
 class ALandscape;
 class FCanvas;
@@ -241,34 +242,22 @@ struct FGizmoHistory
 	}
 };
 
-
-namespace ELandscapeEdge
+namespace ENewLandscapePreviewMode
 {
 	enum Type
 	{
 		None,
-
-		// Edges
-		X_Negative,
-		X_Positive,
-		Y_Negative,
-		Y_Positive,
-
-		// Corners
-		X_Negative_Y_Negative,
-		X_Positive_Y_Negative,
-		X_Negative_Y_Positive,
-		X_Positive_Y_Positive,
+		NewLandscape,
+		ImportLandscape,
 	};
 }
 
-namespace ENewLandscapePreviewMode
+namespace EImportExportMode
 {
 	enum Type
 	{
-		None,
-		NewLandscape,
-		ImportLandscape,
+		Import,
+		Export,
 	};
 }
 
@@ -306,8 +295,7 @@ public:
 	int32 CurrentBrushSetIndex;
 
 	ENewLandscapePreviewMode::Type NewLandscapePreviewMode;
-	ELandscapeEdge::Type DraggingEdge;
-	float DraggingEdge_Remainder;
+	EImportExportMode::Type ImportExportMode;
 
 	TWeakObjectPtr<ALandscapeGizmoActiveActor> CurrentGizmoActor;
 	// UI callbacks for copy/paste tool
@@ -351,6 +339,7 @@ public:
 	void InitializeTool_Retopologize();
 	void InitializeTool_NewLandscape();
 	void InitializeTool_ResizeLandscape();
+	void InitializeTool_ImportExport();
 	void InitializeTool_Select();
 	void InitializeTool_AddComponent();
 	void InitializeTool_DeleteComponent();
@@ -521,6 +510,7 @@ public:
 	const TArray<ALandscapeBlueprintBrushBase*>& GetBrushList() const;
 
 	const TArray<TSharedRef<FLandscapeTargetListInfo>>& GetTargetList() const;
+	UMaterialInterface* GetTargetLandscapeMaterial() const { return CachedLandscapeMaterial; }
 	const TArray<FName>* GetTargetDisplayOrderList() const;
 	const TArray<FName>& GetTargetShownList() const;
 	int32 GetTargetLayerStartingIndex() const;
@@ -603,6 +593,8 @@ public:
 
 	void ReimportData(const FLandscapeTargetListInfo& TargetInfo);
 	void ImportData(const FLandscapeTargetListInfo& TargetInfo, const FString& Filename);
+	void ImportHeightData(ULandscapeInfo* LandscapeInfo, const FGuid& LayerGuid, const FString& Filename, const FIntRect& ImportRegionVerts, ELandscapeImportTransformType TransformType = ELandscapeImportTransformType::ExpandCentered, FIntPoint Offset = FIntPoint(0,0), const ELandscapeLayerPaintingRestriction& PaintRestriction = ELandscapeLayerPaintingRestriction::None);
+	void ImportWeightData(ULandscapeInfo* LandscapeInfo, const FGuid& LayerGuid, ULandscapeLayerInfoObject* LayerInfo, const FString& Filename, const FIntRect& ImportRegionVerts, ELandscapeImportTransformType TransformType = ELandscapeImportTransformType::ExpandCentered, FIntPoint Offset = FIntPoint(0, 0), const ELandscapeLayerPaintingRestriction& PaintRestriction = ELandscapeLayerPaintingRestriction::None);
 	bool UseSingleFileImport() const { return !IsGridBased(); }
 
 	/** Resample landscape to a different resolution or change the component size */
diff --git a/Engine/Source/Editor/LandscapeEditor/Private/LandscapeEdModeComponentTools.cpp b/Engine/Source/Editor/LandscapeEditor/Private/LandscapeEdModeComponentTools.cpp
--- a/Engine/Source/Editor/LandscapeEditor/Private/LandscapeEdModeComponentTools.cpp
+++ b/Engine/Source/Editor/LandscapeEditor/Private/LandscapeEdModeComponentTools.cpp
@@ -1650,6 +1650,62 @@ void FEdModeLandscape::PasteDataFromGizmo()
 	}
 }
 
+namespace ELandscapeEdge
+{
+	enum Type
+	{
+		None,
+
+		// Edges
+		X_Negative,
+		X_Positive,
+		Y_Negative,
+		Y_Positive,
+
+		// Corners
+		X_Negative_Y_Negative,
+		X_Positive_Y_Negative,
+		X_Negative_Y_Positive,
+		X_Positive_Y_Positive,
+	};
+}
+
+struct HNewLandscapeGrabHandleProxy : public HHitProxy
+{
+	DECLARE_HIT_PROXY();
+
+	ELandscapeEdge::Type Edge;
+
+	HNewLandscapeGrabHandleProxy(ELandscapeEdge::Type InEdge) :
+		HHitProxy(HPP_Wireframe),
+		Edge(InEdge)
+	{
+	}
+
+	virtual EMouseCursor::Type GetMouseCursor() override
+	{
+		switch (Edge)
+		{
+		case ELandscapeEdge::X_Negative:
+		case ELandscapeEdge::X_Positive:
+			return EMouseCursor::ResizeLeftRight;
+		case ELandscapeEdge::Y_Negative:
+		case ELandscapeEdge::Y_Positive:
+			return EMouseCursor::ResizeUpDown;
+		case ELandscapeEdge::X_Negative_Y_Negative:
+		case ELandscapeEdge::X_Positive_Y_Positive:
+			return EMouseCursor::ResizeSouthEast;
+		case ELandscapeEdge::X_Negative_Y_Positive:
+		case ELandscapeEdge::X_Positive_Y_Negative:
+			return EMouseCursor::ResizeSouthWest;
+		}
+
+		return EMouseCursor::SlashedCircle;
+	}
+};
+
+IMPLEMENT_HIT_PROXY(HNewLandscapeGrabHandleProxy, HHitProxy)
+
 //
 // FLandscapeToolNewLandscape
 //
@@ -1658,11 +1714,15 @@ class FLandscapeToolNewLandscape : public FLandscapeTool
 public:
 	FEdModeLandscape* EdMode;
 	ENewLandscapePreviewMode::Type NewLandscapePreviewMode;
+	ELandscapeEdge::Type DraggingEdge;
+	float DraggingEdge_Remainder;
 
 	FLandscapeToolNewLandscape(FEdModeLandscape* InEdMode)
 		: FLandscapeTool()
 		, EdMode(InEdMode)
 		, NewLandscapePreviewMode(ENewLandscapePreviewMode::NewLandscape)
+		, DraggingEdge(ELandscapeEdge::None)
+		, DraggingEdge_Remainder(0.0f)
 	{
 	}
 
@@ -1675,6 +1735,8 @@ public:
 
 	virtual void EnterTool() override
 	{
+		DraggingEdge = ELandscapeEdge::None;
+		DraggingEdge_Remainder = 0.0f;
 		EdMode->NewLandscapePreviewMode = NewLandscapePreviewMode;
 		EdMode->UISettings->ImportLandscapeData();
 	}
@@ -1702,6 +1764,305 @@ public:
 		// does nothing
 		return false;
 	}
+
+	virtual bool InputKey(FEditorViewportClient* ViewportClient, FViewport* Viewport, FKey Key, EInputEvent Event) override 
+	{ 
+		if (Key == EKeys::LeftMouseButton)
+		{
+			// Press mouse button
+			if (Event == IE_Pressed && !IsAltDown(Viewport))
+			{
+				// See if we clicked on a new landscape handle..
+				int32 HitX = Viewport->GetMouseX();
+				int32 HitY = Viewport->GetMouseY();
+				HHitProxy* HitProxy = Viewport->GetHitProxy(HitX, HitY);
+				if (HitProxy)
+				{
+					if (HitProxy->IsA(HNewLandscapeGrabHandleProxy::StaticGetType()))
+					{
+						HNewLandscapeGrabHandleProxy* EdgeProxy = (HNewLandscapeGrabHandleProxy*)HitProxy;
+						DraggingEdge = EdgeProxy->Edge;
+						DraggingEdge_Remainder = 0;
+
+						return false; // false to let FEditorViewportClient.InputKey start mouse tracking and enable InputDelta() so we can use it
+					}
+				}
+			}
+			else if (Event == IE_Released)
+			{
+				if (DraggingEdge)
+				{
+					DraggingEdge = ELandscapeEdge::None;
+					DraggingEdge_Remainder = 0;
+
+					return false; // false to let FEditorViewportClient.InputKey end mouse tracking
+				}
+			}
+		}
+		
+		return false; 
+	}
+
+	virtual bool InputDelta(FEditorViewportClient* InViewportClient, FViewport* InViewport, FVector& InDrag, FRotator& InRot, FVector& InScale) override
+	{
+		ULandscapeEditorObject* UISettings = EdMode->UISettings;
+		if (InViewportClient->GetCurrentWidgetAxis() != EAxisList::None)
+		{
+			FVector DeltaScale = InScale;
+			DeltaScale.X = DeltaScale.Y = (FMath::Abs(InScale.X) > FMath::Abs(InScale.Y)) ? InScale.X : InScale.Y;
+
+			UISettings->Modify();
+			UISettings->NewLandscape_Location += InDrag;
+			UISettings->NewLandscape_Rotation += InRot;
+			UISettings->NewLandscape_Scale += DeltaScale;
+
+			return true;
+		}
+		else if (DraggingEdge != ELandscapeEdge::None)
+		{
+			FVector HitLocation;
+			EdMode->LandscapePlaneTrace(InViewportClient, FPlane(UISettings->NewLandscape_Location, FVector(0, 0, 1)), HitLocation);
+
+			FTransform Transform(UISettings->NewLandscape_Rotation, UISettings->NewLandscape_Location, UISettings->NewLandscape_Scale * UISettings->NewLandscape_QuadsPerSection * UISettings->NewLandscape_SectionsPerComponent);
+			HitLocation = Transform.InverseTransformPosition(HitLocation);
+
+			UISettings->Modify();
+			switch (DraggingEdge)
+			{
+			case ELandscapeEdge::X_Negative:
+			case ELandscapeEdge::X_Negative_Y_Negative:
+			case ELandscapeEdge::X_Negative_Y_Positive:
+			{
+				const int32 InitialComponentCountX = UISettings->NewLandscape_ComponentCount.X;
+				const int32 Delta = FMath::RoundToInt(HitLocation.X + (float)InitialComponentCountX / 2);
+				UISettings->NewLandscape_ComponentCount.X = InitialComponentCountX - Delta;
+				UISettings->NewLandscape_ClampSize();
+				const int32 ActualDelta = UISettings->NewLandscape_ComponentCount.X - InitialComponentCountX;
+				UISettings->NewLandscape_Location -= Transform.TransformVector(FVector(((float)ActualDelta / 2), 0, 0));
+			}
+			break;
+			case ELandscapeEdge::X_Positive:
+			case ELandscapeEdge::X_Positive_Y_Negative:
+			case ELandscapeEdge::X_Positive_Y_Positive:
+			{
+				const int32 InitialComponentCountX = UISettings->NewLandscape_ComponentCount.X;
+				int32 Delta = FMath::RoundToInt(HitLocation.X - (float)InitialComponentCountX / 2);
+				UISettings->NewLandscape_ComponentCount.X = InitialComponentCountX + Delta;
+				UISettings->NewLandscape_ClampSize();
+				const int32 ActualDelta = UISettings->NewLandscape_ComponentCount.X - InitialComponentCountX;
+				UISettings->NewLandscape_Location += Transform.TransformVector(FVector(((float)ActualDelta / 2), 0, 0));
+			}
+			break;
+			case  ELandscapeEdge::Y_Negative:
+			case  ELandscapeEdge::Y_Positive:
+				break;
+			}
+						
+			return true;
+		}
+		
+		return false;
+	}
+
+	virtual void Render(const FSceneView* View, FViewport* Viewport, FPrimitiveDrawInterface* PDI) override
+	{
+		if (NewLandscapePreviewMode != ENewLandscapePreviewMode::None)
+		{
+			static const float        CornerSize = 0.33f;
+			static const FLinearColor CornerColour(1.0f, 1.0f, 0.5f);
+			static const FLinearColor EdgeColour(1.0f, 1.0f, 0.0f);
+			static const FLinearColor ComponentBorderColour(0.0f, 0.85f, 0.0f);
+			static const FLinearColor SectionBorderColour(0.0f, 0.4f, 0.0f);
+			static const FLinearColor InnerColour(0.0f, 0.25f, 0.0f);
+
+			const ELevelViewportType ViewportType = ((FEditorViewportClient*)Viewport->GetClient())->ViewportType;
+
+			const int32 ComponentCountX = EdMode->UISettings->NewLandscape_ComponentCount.X;
+			const int32 ComponentCountY = EdMode->UISettings->NewLandscape_ComponentCount.Y;
+			const int32 QuadsPerComponent = EdMode->UISettings->NewLandscape_SectionsPerComponent * EdMode->UISettings->NewLandscape_QuadsPerSection;
+			const float ComponentSize = QuadsPerComponent;
+			const FVector Offset = EdMode->UISettings->NewLandscape_Location + FTransform(EdMode->UISettings->NewLandscape_Rotation, FVector::ZeroVector, EdMode->UISettings->NewLandscape_Scale).TransformVector(FVector(-ComponentCountX * ComponentSize / 2, -ComponentCountY * ComponentSize / 2, 0));
+			const FTransform Transform = FTransform(EdMode->UISettings->NewLandscape_Rotation, Offset, EdMode->UISettings->NewLandscape_Scale);
+
+			if (NewLandscapePreviewMode == ENewLandscapePreviewMode::ImportLandscape)
+			{
+				const TArray<uint16>& ImportHeights = EdMode->UISettings->GetImportLandscapeData();
+				if (ImportHeights.Num() != 0)
+				{
+					const float InvQuadsPerComponent = 1.0f / (float)QuadsPerComponent;
+					const int32 SizeX = ComponentCountX * QuadsPerComponent + 1;
+					const int32 SizeY = ComponentCountY * QuadsPerComponent + 1;
+					const int32 ImportSizeX = EdMode->UISettings->ImportLandscape_Width;
+					const int32 ImportSizeY = EdMode->UISettings->ImportLandscape_Height;
+					const int32 OffsetX = (SizeX - ImportSizeX) / 2;
+					const int32 OffsetY = (SizeY - ImportSizeY) / 2;
+
+					for (int32 ComponentY = 0; ComponentY < ComponentCountY; ComponentY++)
+					{
+						const int32 Y0 = ComponentY * QuadsPerComponent;
+						const int32 Y1 = (ComponentY + 1) * QuadsPerComponent;
+
+						const int32 ImportY0 = FMath::Clamp<int32>(Y0 - OffsetY, 0, ImportSizeY - 1);
+						const int32 ImportY1 = FMath::Clamp<int32>(Y1 - OffsetY, 0, ImportSizeY - 1);
+
+						for (int32 ComponentX = 0; ComponentX < ComponentCountX; ComponentX++)
+						{
+							const int32 X0 = ComponentX * QuadsPerComponent;
+							const int32 X1 = (ComponentX + 1) * QuadsPerComponent;
+							const int32 ImportX0 = FMath::Clamp<int32>(X0 - OffsetX, 0, ImportSizeX - 1);
+							const int32 ImportX1 = FMath::Clamp<int32>(X1 - OffsetX, 0, ImportSizeX - 1);
+							const float Z00 = ((float)ImportHeights[ImportX0 + ImportY0 * ImportSizeX] - 32768.0f) * LANDSCAPE_ZSCALE;
+							const float Z01 = ((float)ImportHeights[ImportX0 + ImportY1 * ImportSizeX] - 32768.0f) * LANDSCAPE_ZSCALE;
+							const float Z10 = ((float)ImportHeights[ImportX1 + ImportY0 * ImportSizeX] - 32768.0f) * LANDSCAPE_ZSCALE;
+							const float Z11 = ((float)ImportHeights[ImportX1 + ImportY1 * ImportSizeX] - 32768.0f) * LANDSCAPE_ZSCALE;
+
+							if (ComponentX == 0)
+							{
+								PDI->SetHitProxy(new HNewLandscapeGrabHandleProxy(ELandscapeEdge::X_Negative));
+								PDI->DrawLine(Transform.TransformPosition(FVector(X0, Y0, Z00)), Transform.TransformPosition(FVector(X0, Y1, Z01)), ComponentBorderColour, SDPG_Foreground);
+								PDI->SetHitProxy(NULL);
+							}
+
+							if (ComponentX == ComponentCountX - 1)
+							{
+								PDI->SetHitProxy(new HNewLandscapeGrabHandleProxy(ELandscapeEdge::X_Positive));
+								PDI->DrawLine(Transform.TransformPosition(FVector(X1, Y0, Z10)), Transform.TransformPosition(FVector(X1, Y1, Z11)), ComponentBorderColour, SDPG_Foreground);
+								PDI->SetHitProxy(NULL);
+							}
+							else
+							{
+								PDI->DrawLine(Transform.TransformPosition(FVector(X1, Y0, Z10)), Transform.TransformPosition(FVector(X1, Y1, Z11)), ComponentBorderColour, SDPG_Foreground);
+							}
+
+							if (ComponentY == 0)
+							{
+								PDI->SetHitProxy(new HNewLandscapeGrabHandleProxy(ELandscapeEdge::Y_Negative));
+								PDI->DrawLine(Transform.TransformPosition(FVector(X0, Y0, Z00)), Transform.TransformPosition(FVector(X1, Y0, Z10)), ComponentBorderColour, SDPG_Foreground);
+								PDI->SetHitProxy(NULL);
+							}
+
+							if (ComponentY == ComponentCountY - 1)
+							{
+								PDI->SetHitProxy(new HNewLandscapeGrabHandleProxy(ELandscapeEdge::Y_Positive));
+								PDI->DrawLine(Transform.TransformPosition(FVector(X0, Y1, Z01)), Transform.TransformPosition(FVector(X1, Y1, Z11)), ComponentBorderColour, SDPG_Foreground);
+								PDI->SetHitProxy(NULL);
+							}
+							else
+							{
+								PDI->DrawLine(Transform.TransformPosition(FVector(X0, Y1, Z01)), Transform.TransformPosition(FVector(X1, Y1, Z11)), ComponentBorderColour, SDPG_Foreground);
+							}
+
+							// intra-component lines - too slow for big landscapes
+							/*
+							for (int32 x=1;x<QuadsPerComponent;x++)
+							{
+							PDI->DrawLine(Transform.TransformPosition(FVector(X0+x, Y0, FMath::Lerp(Z00,Z10,(float)x*InvQuadsPerComponent))), Transform.TransformPosition(FVector(X0+x, Y1, FMath::Lerp(Z01,Z11,(float)x*InvQuadsPerComponent))), ComponentBorderColour, SDPG_World);
+							}
+							for (int32 y=1;y<QuadsPerComponent;y++)
+							{
+							PDI->DrawLine(Transform.TransformPosition(FVector(X0, Y0+y, FMath::Lerp(Z00,Z01,(float)y*InvQuadsPerComponent))), Transform.TransformPosition(FVector(X1, Y0+y, FMath::Lerp(Z10,Z11,(float)y*InvQuadsPerComponent))), ComponentBorderColour, SDPG_World);
+							}
+							*/
+						}
+					}
+				}
+			}
+			else //if (NewLandscapePreviewMode == ENewLandscapePreviewMode::NewLandscape)
+			{
+				if (ViewportType == LVT_Perspective || ViewportType == LVT_OrthoXY || ViewportType == LVT_OrthoNegativeXY)
+				{
+					for (int32 x = 0; x <= ComponentCountX * QuadsPerComponent; x++)
+					{
+						if (x == 0)
+						{
+							PDI->SetHitProxy(new HNewLandscapeGrabHandleProxy(ELandscapeEdge::X_Negative_Y_Negative));
+							PDI->DrawLine(Transform.TransformPosition(FVector(x, 0, 0)), Transform.TransformPosition(FVector(x, CornerSize * ComponentSize, 0)), CornerColour, SDPG_Foreground);
+							PDI->SetHitProxy(new HNewLandscapeGrabHandleProxy(ELandscapeEdge::X_Negative));
+							PDI->DrawLine(Transform.TransformPosition(FVector(x, CornerSize * ComponentSize, 0)), Transform.TransformPosition(FVector(x, (ComponentCountY - CornerSize) * ComponentSize, 0)), EdgeColour, SDPG_Foreground);
+							PDI->SetHitProxy(new HNewLandscapeGrabHandleProxy(ELandscapeEdge::X_Negative_Y_Positive));
+							PDI->DrawLine(Transform.TransformPosition(FVector(x, (ComponentCountY - CornerSize) * ComponentSize, 0)), Transform.TransformPosition(FVector(x, ComponentCountY * ComponentSize, 0)), CornerColour, SDPG_Foreground);
+							PDI->SetHitProxy(NULL);
+						}
+						else if (x == ComponentCountX * QuadsPerComponent)
+						{
+							PDI->SetHitProxy(new HNewLandscapeGrabHandleProxy(ELandscapeEdge::X_Positive_Y_Negative));
+							PDI->DrawLine(Transform.TransformPosition(FVector(x, 0, 0)), Transform.TransformPosition(FVector(x, CornerSize * ComponentSize, 0)), CornerColour, SDPG_Foreground);
+							PDI->SetHitProxy(new HNewLandscapeGrabHandleProxy(ELandscapeEdge::X_Positive));
+							PDI->DrawLine(Transform.TransformPosition(FVector(x, CornerSize * ComponentSize, 0)), Transform.TransformPosition(FVector(x, (ComponentCountY - CornerSize) * ComponentSize, 0)), EdgeColour, SDPG_Foreground);
+							PDI->SetHitProxy(new HNewLandscapeGrabHandleProxy(ELandscapeEdge::X_Positive_Y_Positive));
+							PDI->DrawLine(Transform.TransformPosition(FVector(x, (ComponentCountY - CornerSize) * ComponentSize, 0)), Transform.TransformPosition(FVector(x, ComponentCountY * ComponentSize, 0)), CornerColour, SDPG_Foreground);
+							PDI->SetHitProxy(NULL);
+						}
+						else if (x % QuadsPerComponent == 0)
+						{
+							PDI->DrawLine(Transform.TransformPosition(FVector(x, 0, 0)), Transform.TransformPosition(FVector(x, ComponentCountY * ComponentSize, 0)), ComponentBorderColour, SDPG_Foreground);
+						}
+						else if (x % EdMode->UISettings->NewLandscape_QuadsPerSection == 0)
+						{
+							PDI->DrawLine(Transform.TransformPosition(FVector(x, 0, 0)), Transform.TransformPosition(FVector(x, ComponentCountY * ComponentSize, 0)), SectionBorderColour, SDPG_Foreground);
+						}
+						else
+						{
+							PDI->DrawLine(Transform.TransformPosition(FVector(x, 0, 0)), Transform.TransformPosition(FVector(x, ComponentCountY * ComponentSize, 0)), InnerColour, SDPG_World);
+						}
+					}
+				}
+				else
+				{
+					// Don't allow dragging to resize in side-view
+					// and there's no point drawing the inner lines as only the outer is visible
+					PDI->DrawLine(Transform.TransformPosition(FVector(0, 0, 0)), Transform.TransformPosition(FVector(0, ComponentCountY * ComponentSize, 0)), EdgeColour, SDPG_World);
+					PDI->DrawLine(Transform.TransformPosition(FVector(ComponentCountX * QuadsPerComponent, 0, 0)), Transform.TransformPosition(FVector(ComponentCountX * QuadsPerComponent, ComponentCountY * ComponentSize, 0)), EdgeColour, SDPG_World);
+				}
+
+				if (ViewportType == LVT_Perspective || ViewportType == LVT_OrthoXY || ViewportType == LVT_OrthoNegativeXY)
+				{
+					for (int32 y = 0; y <= ComponentCountY * QuadsPerComponent; y++)
+					{
+						if (y == 0)
+						{
+							PDI->SetHitProxy(new HNewLandscapeGrabHandleProxy(ELandscapeEdge::X_Negative_Y_Negative));
+							PDI->DrawLine(Transform.TransformPosition(FVector(0, y, 0)), Transform.TransformPosition(FVector(CornerSize * ComponentSize, y, 0)), CornerColour, SDPG_Foreground);
+							PDI->SetHitProxy(new HNewLandscapeGrabHandleProxy(ELandscapeEdge::Y_Negative));
+							PDI->DrawLine(Transform.TransformPosition(FVector(CornerSize * ComponentSize, y, 0)), Transform.TransformPosition(FVector((ComponentCountX - CornerSize) * ComponentSize, y, 0)), EdgeColour, SDPG_Foreground);
+							PDI->SetHitProxy(new HNewLandscapeGrabHandleProxy(ELandscapeEdge::X_Positive_Y_Negative));
+							PDI->DrawLine(Transform.TransformPosition(FVector((ComponentCountX - CornerSize) * ComponentSize, y, 0)), Transform.TransformPosition(FVector(ComponentCountX * ComponentSize, y, 0)), CornerColour, SDPG_Foreground);
+							PDI->SetHitProxy(NULL);
+						}
+						else if (y == ComponentCountY * QuadsPerComponent)
+						{
+							PDI->SetHitProxy(new HNewLandscapeGrabHandleProxy(ELandscapeEdge::X_Negative_Y_Positive));
+							PDI->DrawLine(Transform.TransformPosition(FVector(0, y, 0)), Transform.TransformPosition(FVector(CornerSize * ComponentSize, y, 0)), CornerColour, SDPG_Foreground);
+							PDI->SetHitProxy(new HNewLandscapeGrabHandleProxy(ELandscapeEdge::Y_Positive));
+							PDI->DrawLine(Transform.TransformPosition(FVector(CornerSize * ComponentSize, y, 0)), Transform.TransformPosition(FVector((ComponentCountX - CornerSize) * ComponentSize, y, 0)), EdgeColour, SDPG_Foreground);
+							PDI->SetHitProxy(new HNewLandscapeGrabHandleProxy(ELandscapeEdge::X_Positive_Y_Positive));
+							PDI->DrawLine(Transform.TransformPosition(FVector((ComponentCountX - CornerSize) * ComponentSize, y, 0)), Transform.TransformPosition(FVector(ComponentCountX * ComponentSize, y, 0)), CornerColour, SDPG_Foreground);
+							PDI->SetHitProxy(NULL);
+						}
+						else if (y % QuadsPerComponent == 0)
+						{
+							PDI->DrawLine(Transform.TransformPosition(FVector(0, y, 0)), Transform.TransformPosition(FVector(ComponentCountX * ComponentSize, y, 0)), ComponentBorderColour, SDPG_Foreground);
+						}
+						else if (y % EdMode->UISettings->NewLandscape_QuadsPerSection == 0)
+						{
+							PDI->DrawLine(Transform.TransformPosition(FVector(0, y, 0)), Transform.TransformPosition(FVector(ComponentCountX * ComponentSize, y, 0)), SectionBorderColour, SDPG_Foreground);
+						}
+						else
+						{
+							PDI->DrawLine(Transform.TransformPosition(FVector(0, y, 0)), Transform.TransformPosition(FVector(ComponentCountX * ComponentSize, y, 0)), InnerColour, SDPG_World);
+						}
+					}
+				}
+				else
+				{
+					// Don't allow dragging to resize in side-view
+					// and there's no point drawing the inner lines as only the outer is visible
+					PDI->DrawLine(Transform.TransformPosition(FVector(0, 0, 0)), Transform.TransformPosition(FVector(ComponentCountX * ComponentSize, 0, 0)), EdgeColour, SDPG_World);
+					PDI->DrawLine(Transform.TransformPosition(FVector(0, ComponentCountY * QuadsPerComponent, 0)), Transform.TransformPosition(FVector(ComponentCountX * ComponentSize, ComponentCountY * QuadsPerComponent, 0)), EdgeColour, SDPG_World);
+				}
+			}
+		}
+	}
 };
 
 
@@ -1772,6 +2133,244 @@ public:
 	}
 };
 
+class FLandscapeToolImportExport : public FLandscapeTool
+{
+public:
+	FEdModeLandscape* EdMode;
+	FVector MinLocation;
+	FIntRect LandscapeExtent;
+	
+	FLandscapeToolImportExport(FEdModeLandscape* InEdMode)
+		: FLandscapeTool()
+		, EdMode(InEdMode)
+	{
+	}
+
+	virtual const TCHAR* GetToolName() override { return TEXT("ImportExport"); }
+	virtual FText GetDisplayName() override { return LOCTEXT("LandscapeMode_ImportExport", "Import Export"); };
+	virtual FText GetDisplayMessage() override { return LOCTEXT("LandscapeMode_ImportExport_Message", "Import/Export Landscape Data"); };
+
+	virtual bool UsesTransformWidget() const { return true; }
+	virtual void SetEditRenderType() override { GLandscapeEditRenderMode = ELandscapeEditRenderMode::SelectComponent | (GLandscapeEditRenderMode & ELandscapeEditRenderMode::BitMaskForMask); }
+	virtual bool SupportsMask() override { return false; }
+
+	virtual FVector GetWidgetLocation() const override
+	{
+		
+		if (EdMode->UISettings->ImportType != ELandscapeImportTransformType::Resample)
+		{
+			if (ALandscapeGizmoActiveActor* GizmoActor = EdMode->CurrentGizmoActor.Get())
+			{
+				return EdMode->CurrentGizmoActor->GetActorLocation();
+			}	
+		}
+		return MinLocation;
+	}
+
+	virtual FMatrix GetWidgetRotation() const override
+	{ 
+		return FMatrix::Identity; 
+	}
+
+	virtual void EnterTool() override
+	{
+		if (EdMode)
+		{
+			EdMode->UISettings->RefreshImportLayersList(/* bRefreshFromTarget */true);
+			EdMode->UISettings->RefreshImports();
+
+			if (ULandscapeInfo* LandscapeInfo = EdMode->CurrentToolTarget.LandscapeInfo.Get())
+			{
+				MinLocation = FVector(0, 0, 0);
+				if (LandscapeInfo->GetLandscapeExtent(LandscapeExtent))
+				{
+					MinLocation = LandscapeInfo->GetLandscapeProxy()->LandscapeActorToWorld().TransformPosition(FVector(LandscapeExtent.Min.X, LandscapeExtent.Min.Y, 0.0f));
+				}
+				else
+				{
+					LandscapeExtent = FIntRect(0, 0, 0, 0);
+				}
+
+				FVector LocalPosition(EdMode->UISettings->ImportLandscape_GizmoLocalPosition.X, EdMode->UISettings->ImportLandscape_GizmoLocalPosition.Y, 0.0f);
+				FVector GizmoPosition = LandscapeInfo->GetLandscapeProxy()->LandscapeActorToWorld().TransformPosition(LocalPosition);
+				
+				if (ALandscapeGizmoActiveActor* GizmoActor = EdMode->CurrentGizmoActor.Get())
+				{
+					GizmoActor->SetActorLocation(GizmoPosition);
+				}
+			}
+		}
+	}
+
+	virtual void ExitTool() override
+	{
+		
+	}
+
+	virtual bool BeginTool(FEditorViewportClient* ViewportClient, const FLandscapeToolTarget& Target, const FVector& InHitLocation) override
+	{
+		// does nothing
+		return false;
+	}
+
+	virtual void EndTool(FEditorViewportClient* ViewportClient) override
+	{
+		// does nothing
+	}
+
+	virtual bool MouseMove(FEditorViewportClient* ViewportClient, FViewport* Viewport, int32 x, int32 y) override
+	{
+		// does nothing
+		return false;
+	}
+		
+	virtual bool InputDelta(FEditorViewportClient* InViewportClient, FViewport* InViewport, FVector& InDrag, FRotator& InRot, FVector& InScale)
+	{
+		if (InViewportClient->GetCurrentWidgetAxis() != EAxisList::None)
+		{
+			// Resample Gizmo can't move
+			if (EdMode && EdMode->UISettings->ImportType != ELandscapeImportTransformType::Resample)
+			{
+				if (ALandscapeGizmoActiveActor* GizmoActor = EdMode->CurrentGizmoActor.Get())
+				{
+					GEditor->ApplyDeltaToActor(
+						GizmoActor,
+						true,
+						&InDrag,
+						nullptr,
+						nullptr,
+						InViewportClient->IsAltPressed(),
+						InViewportClient->IsShiftPressed(),
+						InViewportClient->IsCtrlPressed());
+
+					return true;
+				}
+			}
+		}
+	
+		return false;
+	}
+
+	virtual void Render(const FSceneView* View, FViewport* Viewport, FPrimitiveDrawInterface* PDI) override
+	{
+		ULandscapeInfo* LandscapeInfo = EdMode->CurrentToolTarget.LandscapeInfo.Get();
+		if (!LandscapeInfo)
+		{
+			return;
+		}
+			
+		static const float        CornerSize = 0.33f;
+		static const FLinearColor CornerColour(1.0f, 1.0f, 0.5f);
+		static const FLinearColor EdgeColour(1.0f, 1.0f, 0.0f);
+		static const FLinearColor ComponentBorderColour(0.0f, 0.85f, 0.0f);
+		static const FLinearColor SectionBorderColour(0.0f, 0.4f, 0.0f);
+		static const FLinearColor ComponentOutsideColor(0.0f, 0.0f, 0.85f);
+		static const FLinearColor InnerColour(0.0f, 0.25f, 0.0f);
+
+		const ELevelViewportType ViewportType = ((FEditorViewportClient*)Viewport->GetClient())->ViewportType;
+
+		const int32 ComponentSizeInt = EdMode->UISettings->NewLandscape_SectionsPerComponent * EdMode->UISettings->NewLandscape_QuadsPerSection;
+		const float ComponentSize = ComponentSizeInt;
+		const FTransform GizmoTransform = FTransform(FRotator(0,0,0), GetWidgetLocation(), LandscapeInfo->DrawScale);
+		const int32 Height = EdMode->UISettings->ImportType == ELandscapeImportTransformType::None ? EdMode->UISettings->ImportLandscape_Height : (LandscapeExtent.Height()+1);
+		const int32 Width = EdMode->UISettings->ImportType == ELandscapeImportTransformType::None ? EdMode->UISettings->ImportLandscape_Width: (LandscapeExtent.Width()+1);
+		const int32 ImportHeight = EdMode->UISettings->ImportType != ELandscapeImportTransformType::Resample ? EdMode->UISettings->ImportLandscape_Height : (LandscapeExtent.Height() + 1);
+		const int32 ImportWidth = EdMode->UISettings->ImportType != ELandscapeImportTransformType::Resample ? EdMode->UISettings->ImportLandscape_Height : (LandscapeExtent.Height() + 1);
+
+		const FTransform LandscapeTransfo = FTransform(FRotator(0, 0, 0), MinLocation, LandscapeInfo->DrawScale);
+
+		if (EdMode->ImportExportMode == EImportExportMode::Import)
+		{
+			if (ViewportType == LVT_Perspective || ViewportType == LVT_OrthoXY || ViewportType == LVT_OrthoNegativeXY)
+			{
+				for (int32 x = 0; x < ImportWidth; x++)
+				{
+					if (x == 0)
+					{
+						PDI->DrawLine(GizmoTransform.TransformPosition(FVector(x, 0, 0)), GizmoTransform.TransformPosition(FVector(x, CornerSize * ComponentSize, 0)), CornerColour, SDPG_Foreground);
+						PDI->DrawLine(GizmoTransform.TransformPosition(FVector(x, CornerSize * ComponentSize, 0)), GizmoTransform.TransformPosition(FVector(x, ImportHeight - (CornerSize * ComponentSize), 0)), EdgeColour, SDPG_Foreground);
+						PDI->DrawLine(GizmoTransform.TransformPosition(FVector(x, ImportHeight - (CornerSize * ComponentSize), 0)), GizmoTransform.TransformPosition(FVector(x, ImportHeight, 0)), CornerColour, SDPG_Foreground);
+					}
+					else if (x == (ImportWidth - 1))
+					{
+						PDI->DrawLine(GizmoTransform.TransformPosition(FVector(x, 0, 0)), GizmoTransform.TransformPosition(FVector(x, CornerSize * ComponentSize, 0)), CornerColour, SDPG_Foreground);
+						PDI->DrawLine(GizmoTransform.TransformPosition(FVector(x, CornerSize * ComponentSize, 0)), GizmoTransform.TransformPosition(FVector(x, ImportHeight - (CornerSize * ComponentSize), 0)), EdgeColour, SDPG_Foreground);
+						PDI->DrawLine(GizmoTransform.TransformPosition(FVector(x, ImportHeight - (CornerSize * ComponentSize), 0)), GizmoTransform.TransformPosition(FVector(x, ImportHeight, 0)), CornerColour, SDPG_Foreground);
+					}
+					else if ((x % ComponentSizeInt) == 0)
+					{
+						PDI->DrawLine(GizmoTransform.TransformPosition(FVector(x, 0, 0)), GizmoTransform.TransformPosition(FVector(x, ImportHeight, 0)), ComponentBorderColour, SDPG_Foreground);
+					}
+					else if ((x % EdMode->UISettings->NewLandscape_QuadsPerSection) == 0)
+					{
+						PDI->DrawLine(GizmoTransform.TransformPosition(FVector(x, 0, 0)), GizmoTransform.TransformPosition(FVector(x, ImportHeight, 0)), SectionBorderColour, SDPG_Foreground);
+					}
+				}
+
+				if (ImportWidth != Width)
+				{
+					for (int32 x = 0; x < Width; x++)
+					{
+						if ((x % ComponentSizeInt) == 0)
+						{
+							PDI->DrawLine(LandscapeTransfo.TransformPosition(FVector(x, 0, 0)), LandscapeTransfo.TransformPosition(FVector(x, Height, 0)), ComponentOutsideColor, SDPG_Foreground);
+						}
+					}
+				}
+			}
+			else
+			{
+				PDI->DrawLine(GizmoTransform.TransformPosition(FVector(0, 0, 0)), GizmoTransform.TransformPosition(FVector(0, ImportHeight, 0)), EdgeColour, SDPG_Foreground);
+				PDI->DrawLine(GizmoTransform.TransformPosition(FVector(ImportWidth, 0, 0)), GizmoTransform.TransformPosition(FVector(ImportWidth, ImportHeight, 0)), EdgeColour, SDPG_Foreground);
+			}
+
+			if (ViewportType == LVT_Perspective || ViewportType == LVT_OrthoXY || ViewportType == LVT_OrthoNegativeXY)
+			{
+				for (int32 y = 0; y < ImportHeight; y++)
+				{
+					if (y == 0)
+					{
+						PDI->DrawLine(GizmoTransform.TransformPosition(FVector(0, y, 0)), GizmoTransform.TransformPosition(FVector(CornerSize * ComponentSize, y, 0)), CornerColour, SDPG_Foreground);
+						PDI->DrawLine(GizmoTransform.TransformPosition(FVector(CornerSize * ComponentSize, y, 0)), GizmoTransform.TransformPosition(FVector(ImportWidth - (CornerSize * ComponentSize), y, 0)), EdgeColour, SDPG_Foreground);
+						PDI->DrawLine(GizmoTransform.TransformPosition(FVector(ImportWidth - (CornerSize * ComponentSize), y, 0)), GizmoTransform.TransformPosition(FVector(ImportWidth, y, 0)), CornerColour, SDPG_Foreground);
+					}
+					else if (y == (ImportHeight - 1))
+					{
+						PDI->DrawLine(GizmoTransform.TransformPosition(FVector(0, y, 0)), GizmoTransform.TransformPosition(FVector(CornerSize * ComponentSize, y, 0)), CornerColour, SDPG_Foreground);
+						PDI->DrawLine(GizmoTransform.TransformPosition(FVector(CornerSize * ComponentSize, y, 0)), GizmoTransform.TransformPosition(FVector(ImportWidth - (CornerSize * ComponentSize), y, 0)), EdgeColour, SDPG_Foreground);
+						PDI->DrawLine(GizmoTransform.TransformPosition(FVector(ImportWidth - (CornerSize * ComponentSize), y, 0)), GizmoTransform.TransformPosition(FVector(ImportWidth, y, 0)), CornerColour, SDPG_Foreground);
+					}
+					else if ((y % ComponentSizeInt) == 0)
+					{
+						PDI->DrawLine(GizmoTransform.TransformPosition(FVector(0, y, 0)), GizmoTransform.TransformPosition(FVector(ImportWidth, y, 0)), ComponentBorderColour, SDPG_Foreground);
+					}
+					else if ((y % EdMode->UISettings->NewLandscape_QuadsPerSection) == 0)
+					{
+						PDI->DrawLine(GizmoTransform.TransformPosition(FVector(0, y, 0)), GizmoTransform.TransformPosition(FVector(ImportWidth, y, 0)), SectionBorderColour, SDPG_Foreground);
+					}
+				}
+
+				if (ImportHeight != Height)
+				{
+					for (int32 y = 0; y < Height; y++)
+					{
+						if ((y % ComponentSizeInt) == 0)
+						{
+							PDI->DrawLine(LandscapeTransfo.TransformPosition(FVector(0, y, 0)), LandscapeTransfo.TransformPosition(FVector(Width, y, 0)), ComponentOutsideColor, SDPG_Foreground);
+						}
+					}
+				}
+			}
+			else
+			{
+				// and there's no point drawing the inner lines as only the outer is visible
+				PDI->DrawLine(GizmoTransform.TransformPosition(FVector(0, 0, 0)), GizmoTransform.TransformPosition(FVector(ImportWidth, 0, 0)), EdgeColour, SDPG_Foreground);
+				PDI->DrawLine(GizmoTransform.TransformPosition(FVector(0, ImportHeight, 0)), GizmoTransform.TransformPosition(FVector(ImportWidth, ImportHeight, 0)), EdgeColour, SDPG_Foreground);
+			}
+		}
+	}
+};
+
 //////////////////////////////////////////////////////////////////////////
 
 void FEdModeLandscape::InitializeTool_NewLandscape()
@@ -1788,6 +2387,13 @@ void FEdModeLandscape::InitializeTool_ResizeLandscape()
 	LandscapeTools.Add(MoveTemp(Tool_ResizeLandscape));
 }
 
+void FEdModeLandscape::InitializeTool_ImportExport()
+{
+	auto Tool_ImportExportLandscape = MakeUnique<FLandscapeToolImportExport>(this);
+	Tool_ImportExportLandscape->ValidBrushes.Add("BrushSet_Dummy");
+	LandscapeTools.Add(MoveTemp(Tool_ImportExportLandscape));
+}
+
 void FEdModeLandscape::InitializeTool_Select()
 {
 	auto Tool_Select = MakeUnique<FLandscapeToolSelect>(this);
diff --git a/Engine/Source/Editor/LandscapeEditor/Private/LandscapeEditorCommands.cpp b/Engine/Source/Editor/LandscapeEditor/Private/LandscapeEditorCommands.cpp
--- a/Engine/Source/Editor/LandscapeEditor/Private/LandscapeEditorCommands.cpp
+++ b/Engine/Source/Editor/LandscapeEditor/Private/LandscapeEditorCommands.cpp
@@ -33,6 +33,9 @@ void FLandscapeEditorCommands::RegisterCommands()
 	UI_COMMAND(ResizeLandscape, "Resize", "Change Component Size", EUserInterfaceActionType::RadioButton, FInputChord());
 	NameToCommandMap.Add("Tool_ResizeLandscape", ResizeLandscape);
 
+	UI_COMMAND(ImportExportTool, "Import", "Import or Export landscape data", EUserInterfaceActionType::RadioButton, FInputChord());
+	NameToCommandMap.Add("Tool_ImportExport", ImportExportTool);
+
 	UI_COMMAND(SculptTool, "Sculpt", "Sculpt height data.\n\nCtrl+Click to Raise, Ctrl+Shift_Click to lower.", EUserInterfaceActionType::RadioButton, FInputChord());
 	NameToCommandMap.Add("Tool_Sculpt", SculptTool);
 
diff --git a/Engine/Source/Editor/LandscapeEditor/Private/LandscapeEditorCommands.h b/Engine/Source/Editor/LandscapeEditor/Private/LandscapeEditorCommands.h
--- a/Engine/Source/Editor/LandscapeEditor/Private/LandscapeEditorCommands.h
+++ b/Engine/Source/Editor/LandscapeEditor/Private/LandscapeEditorCommands.h
@@ -32,6 +32,7 @@ public:
 	// Tools
 	TSharedPtr<FUICommandInfo> NewLandscape;
 	TSharedPtr<FUICommandInfo> ResizeLandscape;
+	TSharedPtr<FUICommandInfo> ImportExportTool;
 
 	TSharedPtr<FUICommandInfo> SculptTool;
 	TSharedPtr<FUICommandInfo> EraseTool;
diff --git a/Engine/Source/Editor/LandscapeEditor/Private/LandscapeEditorDetailCustomization_Base.cpp b/Engine/Source/Editor/LandscapeEditor/Private/LandscapeEditorDetailCustomization_Base.cpp
--- a/Engine/Source/Editor/LandscapeEditor/Private/LandscapeEditorDetailCustomization_Base.cpp
+++ b/Engine/Source/Editor/LandscapeEditor/Private/LandscapeEditorDetailCustomization_Base.cpp
@@ -5,40 +5,63 @@
 #include "EditorModes.h"
 #include "LandscapeEdMode.h"
 
+namespace LandscapeCustomization
+{
+	FEdModeLandscape* GetEditorMode()
+	{
+		return (FEdModeLandscape*)GLevelEditorModeTools().GetActiveMode(FBuiltinEditorModes::EM_Landscape);
+	}
 
+	bool IsToolActive(FName ToolName)
+	{
+		FEdModeLandscape* LandscapeEdMode = GetEditorMode();
+		if (LandscapeEdMode != NULL && LandscapeEdMode->CurrentTool != NULL)
+		{
+			const FName CurrentToolName = LandscapeEdMode->CurrentTool->GetToolName();
+			return CurrentToolName == ToolName;
+		}
+
+		return false;
+	}
+
+	bool IsBrushSetActive(FName BrushSetName)
+	{
+		FEdModeLandscape* LandscapeEdMode = GetEditorMode();
+		if (LandscapeEdMode && LandscapeEdMode->CurrentBrushSetIndex >= 0)
+		{
+			const FName CurrentBrushSetName = LandscapeEdMode->LandscapeBrushSets[LandscapeEdMode->CurrentBrushSetIndex].BrushSetName;
+			return CurrentBrushSetName == BrushSetName;
+		}
+
+		return false;
+	}
+}
 
 FEdModeLandscape* FLandscapeEditorDetailCustomization_Base::GetEditorMode()
 {
-	return (FEdModeLandscape*)GLevelEditorModeTools().GetActiveMode(FBuiltinEditorModes::EM_Landscape);
+	return LandscapeCustomization::GetEditorMode();
 }
 
 bool FLandscapeEditorDetailCustomization_Base::IsToolActive(FName ToolName)
 {
-	FEdModeLandscape* LandscapeEdMode = GetEditorMode();
-	if (LandscapeEdMode != NULL && LandscapeEdMode->CurrentTool != NULL)
-	{
-		const FName CurrentToolName = LandscapeEdMode->CurrentTool->GetToolName();
-		return CurrentToolName == ToolName;
-	}
-
-	return false;
+	return LandscapeCustomization::IsToolActive(ToolName);
 }
 
 bool FLandscapeEditorDetailCustomization_Base::IsBrushSetActive(FName BrushSetName)
 {
-	FEdModeLandscape* LandscapeEdMode = GetEditorMode();
-	if (LandscapeEdMode && LandscapeEdMode->CurrentBrushSetIndex >= 0)
-	{
-		const FName CurrentBrushSetName = LandscapeEdMode->LandscapeBrushSets[LandscapeEdMode->CurrentBrushSetIndex].BrushSetName;
-		return CurrentBrushSetName == BrushSetName;
-	}
-
-	return false;
+	return LandscapeCustomization::IsBrushSetActive(BrushSetName);
 }
 
 //////////////////////////////////////////////////////////////////////////
 
 FEdModeLandscape* FLandscapeEditorStructCustomization_Base::GetEditorMode()
 {
-	return (FEdModeLandscape*)GLevelEditorModeTools().GetActiveMode(FBuiltinEditorModes::EM_Landscape);
+	return LandscapeCustomization::GetEditorMode();
 }
+
+bool FLandscapeEditorStructCustomization_Base::IsToolActive(FName ToolName)
+{
+	return LandscapeCustomization::IsToolActive(ToolName);
+}
+
+
diff --git a/Engine/Source/Editor/LandscapeEditor/Private/LandscapeEditorDetailCustomization_Base.h b/Engine/Source/Editor/LandscapeEditor/Private/LandscapeEditorDetailCustomization_Base.h
--- a/Engine/Source/Editor/LandscapeEditor/Private/LandscapeEditorDetailCustomization_Base.h
+++ b/Engine/Source/Editor/LandscapeEditor/Private/LandscapeEditorDetailCustomization_Base.h
@@ -135,4 +135,5 @@ class FLandscapeEditorStructCustomization_Base : public IPropertyTypeCustomizati
 {
 protected:
 	static FEdModeLandscape* GetEditorMode();
+	static bool IsToolActive(FName ToolName);
 };
diff --git a/Engine/Source/Editor/LandscapeEditor/Private/LandscapeEditorDetailCustomization_ImportExport.cpp b/Engine/Source/Editor/LandscapeEditor/Private/LandscapeEditorDetailCustomization_ImportExport.cpp
new file mode 100644
--- /dev/null
+++ b/Engine/Source/Editor/LandscapeEditor/Private/LandscapeEditorDetailCustomization_ImportExport.cpp
@@ -0,0 +1,625 @@
+// Copyright Epic Games, Inc. All Rights Reserved.
+
+#include "LandscapeEditorDetailCustomization_ImportExport.h"
+#include "SlateOptMacros.h"
+#include "IDetailCustomization.h"
+#include "DetailLayoutBuilder.h"
+#include "DetailCategoryBuilder.h"
+#include "DetailWidgetRow.h"
+#include "Widgets/Layout/SUniformGridPanel.h"
+#include "Widgets/Input/SCheckBox.h"
+#include "Widgets/Input/SButton.h"
+#include "Widgets/Input/SEditableTextBox.h"
+#include "Widgets/Text/STextBlock.h"
+#include "Widgets/Notifications/SErrorText.h"
+#include "Widgets/SToolTip.h"
+#include "Framework/MultiBox/MultiBoxBuilder.h"
+
+#include "DesktopPlatformModule.h"
+
+#include "LandscapeEditorObject.h"
+#include "LandscapeEditorModule.h"
+#include "LandscapeImportHelper.h"
+
+#define LOCTEXT_NAMESPACE "LandscapeEditor.ImportExport"
+
+TSharedRef<IDetailCustomization> FLandscapeEditorDetailCustomization_ImportExport::MakeInstance()
+{
+	return MakeShareable(new FLandscapeEditorDetailCustomization_ImportExport);
+}
+
+BEGIN_SLATE_FUNCTION_BUILD_OPTIMIZATION
+void FLandscapeEditorDetailCustomization_ImportExport::CustomizeDetails(IDetailLayoutBuilder& DetailBuilder)
+{
+	if (!IsToolActive("ImportExport"))
+	{
+		return;
+	}
+
+	DetailBuilder.HideCategory("New Landscape");
+	IDetailCategoryBuilder& ImportExportCategory = DetailBuilder.EditCategory("Import / Export");
+
+	ImportExportCategory.AddCustomRow(FText::GetEmpty())
+	[
+		SNew(SUniformGridPanel)
+		.SlotPadding(FMargin(10, 2))
+		+ SUniformGridPanel::Slot(0, 0)
+		[
+			SNew(SCheckBox)
+			.Style(FEditorStyle::Get(), "RadioButton")
+			.IsChecked_Static(&FLandscapeEditorDetailCustomization_ImportExport::ModeIsChecked, EImportExportMode::Import)
+			.OnCheckStateChanged_Static(&FLandscapeEditorDetailCustomization_ImportExport::OnModeChanged, EImportExportMode::Import)
+			[
+				SNew(STextBlock)
+				.Text(LOCTEXT("Import", "Import"))
+			]
+		]
+		+ SUniformGridPanel::Slot(1, 0)
+		[
+			SNew(SCheckBox)
+			.Style(FEditorStyle::Get(), "RadioButton")
+			.IsChecked_Static(&FLandscapeEditorDetailCustomization_ImportExport::ModeIsChecked, EImportExportMode::Export)
+			.OnCheckStateChanged_Static(&FLandscapeEditorDetailCustomization_ImportExport::OnModeChanged, EImportExportMode::Export)
+			[
+				SNew(STextBlock)
+				.Text(LOCTEXT("Export", "Export"))
+			]
+		]
+	];
+
+	TSharedRef<IPropertyHandle> PropertyHandle_ImportHeightmapFilename = DetailBuilder.GetProperty(GET_MEMBER_NAME_CHECKED(ULandscapeEditorObject, ImportLandscape_HeightmapFilename));
+	TSharedRef<IPropertyHandle> PropertyHandle_ExportHeightmapFilename = DetailBuilder.GetProperty(GET_MEMBER_NAME_CHECKED(ULandscapeEditorObject, HeightmapExportFilename));
+	
+	PropertyHandle_ImportHeightmapFilename->SetOnPropertyValueChanged(FSimpleDelegate::CreateLambda([PropertyHandle_ImportHeightmapFilename]()
+	{
+		FLandscapeEditorDetailCustomization_ImportExport::FormatFilename(PropertyHandle_ImportHeightmapFilename);
+		FLandscapeEditorDetailCustomization_ImportExport::OnImportHeightmapFilenameChanged();
+	}));
+
+	PropertyHandle_ExportHeightmapFilename->SetOnPropertyValueChanged(FSimpleDelegate::CreateLambda([PropertyHandle_ExportHeightmapFilename]()
+	{
+		FLandscapeEditorDetailCustomization_ImportExport::FormatFilename(PropertyHandle_ExportHeightmapFilename);
+	}));
+
+	TSharedRef<IPropertyHandle> PropertyHandle_HeightmapImportResult = DetailBuilder.GetProperty(GET_MEMBER_NAME_CHECKED(ULandscapeEditorObject, ImportLandscape_HeightmapImportResult));
+	TSharedRef<IPropertyHandle> PropertyHandle_HeightmapErrorMessage = DetailBuilder.GetProperty(GET_MEMBER_NAME_CHECKED(ULandscapeEditorObject, ImportLandscape_HeightmapErrorMessage));
+	DetailBuilder.HideProperty(PropertyHandle_HeightmapImportResult);
+	DetailBuilder.HideProperty(PropertyHandle_HeightmapErrorMessage);
+		
+	TSharedRef<IPropertyHandle> PropertyHandle_ExportEditLayer = DetailBuilder.GetProperty(GET_MEMBER_NAME_CHECKED(ULandscapeEditorObject, bExportEditLayer));
+
+	auto AddHeightmapFileName = [&](TSharedRef<IPropertyHandle> PropertyHandle_HeightmapFilename, TAttribute<EVisibility> PropertyVisibility)
+	{
+		return ImportExportCategory.AddProperty(PropertyHandle_HeightmapFilename)
+			.Visibility(PropertyVisibility)
+			.CustomWidget()
+			.NameContent()
+			[
+				SNew(SHorizontalBox)
+				+ SHorizontalBox::Slot()
+				.AutoWidth()
+				.Padding(0, 0, 2, 0)
+				[
+					SNew(SCheckBox)
+					.IsChecked_Static(&FLandscapeEditorDetailCustomization_ImportExport::GetHeightmapSelectedCheckState)
+					.OnCheckStateChanged_Static(&FLandscapeEditorDetailCustomization_ImportExport::OnHeightmapSelectedCheckStateChanged)
+				]
+				+ SHorizontalBox::Slot()
+				.FillWidth(1)
+				[
+					PropertyHandle_HeightmapFilename->CreatePropertyNameWidget()
+				]
+			]
+			.ValueContent()
+			.MinDesiredWidth(250.0f)
+			.MaxDesiredWidth(0)
+			[
+				SNew(SHorizontalBox)
+				.IsEnabled_Static(&FLandscapeEditorDetailCustomization_ImportExport::IsHeightmapEnabled)
+				+ SHorizontalBox::Slot()
+				.AutoWidth()
+				.Padding(0, 0, 2, 0)
+				[
+					SNew(SErrorText)
+					.Visibility_Static(&GetImportResultErrorVisibility, PropertyHandle_HeightmapImportResult)
+					.BackgroundColor_Static(&GetImportResultErrorColor, PropertyHandle_HeightmapImportResult)
+					.ErrorText(NSLOCTEXT("UnrealEd", "Error", "!"))
+					.ToolTip(
+						SNew(SToolTip)
+						.Text_Static(&GetPropertyValue<FText>, PropertyHandle_HeightmapErrorMessage)
+					)
+				]
+				+ SHorizontalBox::Slot()
+				.FillWidth(1)
+				[
+					SNew(SEditableTextBox)
+					.Font(DetailBuilder.GetDetailFont())
+					.Text_Static(&GetPropertyValueText, PropertyHandle_HeightmapFilename)
+					.OnTextCommitted_Static(&SetFilename, PropertyHandle_HeightmapFilename)
+				]
+				+ SHorizontalBox::Slot()
+				.AutoWidth()
+				.Padding(1, 0, 0, 0)
+				[
+					SNew(SButton)
+					.ContentPadding(FMargin(4, 0))
+					.Text(NSLOCTEXT("UnrealEd", "GenericOpenDialog", "..."))
+					.OnClicked_Static(&OnBrowseFilenameButtonClicked, PropertyHandle_HeightmapFilename)
+				]
+			];
+	};
+
+	AddHeightmapFileName(PropertyHandle_ExportHeightmapFilename, MakeAttributeLambda([]() { return FLandscapeEditorDetailCustomization_ImportExport::GetImportExportVisibility(false); }));
+	AddHeightmapFileName(PropertyHandle_ImportHeightmapFilename, MakeAttributeLambda([]() { return FLandscapeEditorDetailCustomization_ImportExport::GetImportExportVisibility(true); }));
+
+	TSharedRef<IPropertyHandle> PropertyHandle_ImportType = DetailBuilder.GetProperty(GET_MEMBER_NAME_CHECKED(ULandscapeEditorObject, ImportType));
+	ImportExportCategory.AddProperty(PropertyHandle_ImportType).Visibility(MakeAttributeLambda([]() { return FLandscapeEditorDetailCustomization_ImportExport::GetImportExportVisibility(true); }));
+
+	ImportExportCategory.AddProperty(PropertyHandle_ExportEditLayer).Visibility(MakeAttributeLambda([this]()
+	{
+		if (ULandscapeInfo* LandscapeInfo = GetEditorMode()->CurrentToolTarget.LandscapeInfo.Get())
+		{
+			if (!LandscapeInfo->CanHaveLayersContent())
+			{
+				return EVisibility::Collapsed;
+			}
+		}
+
+		return FLandscapeEditorDetailCustomization_ImportExport::GetImportExportVisibility(false);
+	}));
+
+	ImportExportCategory.AddCustomRow(LOCTEXT("ImportResolution", "Import Resolution"))
+		.Visibility(MakeAttributeLambda([]() { return FLandscapeEditorDetailCustomization_ImportExport::GetImportExportVisibility(true); }))
+		.NameContent()
+		[
+			SNew(SBox)
+			.VAlign(VAlign_Center)
+			.Padding(FMargin(2))
+			[
+				SNew(STextBlock)
+				.Font(DetailBuilder.GetDetailFont())
+				.Text(LOCTEXT("ImportResolution", "Import Resolution"))
+			]
+		]
+		.ValueContent()
+		[
+			SNew(SBox)
+			.Padding(FMargin(0, 0, 12, 0)) // Line up with the other properties due to having no reset to default button
+			[
+				SNew(SComboButton)
+				.OnGetMenuContent_Static(&FLandscapeEditorDetailCustomization_ImportExport::GetImportLandscapeResolutionMenu)
+				.ContentPadding(2)
+				.ButtonContent()
+				[
+					SNew(STextBlock)
+					.Font(DetailBuilder.GetDetailFont())
+					.Text_Static(&FLandscapeEditorDetailCustomization_ImportExport::GetImportLandscapeResolution)
+				]
+			]
+		];
+
+	TSharedRef<IPropertyHandle> PropertyHandle_Layers = DetailBuilder.GetProperty(GET_MEMBER_NAME_CHECKED(ULandscapeEditorObject, ImportLandscape_Layers));
+	ImportExportCategory.AddProperty(PropertyHandle_Layers);
+
+	ImportExportCategory.AddCustomRow(FText::GetEmpty())
+	[
+		SNew(SHorizontalBox)
+		+ SHorizontalBox::Slot()
+		.FillWidth(1)
+		+ SHorizontalBox::Slot()
+		.AutoWidth()
+		[
+			SNew(SButton)
+			.Text_Static(&GetImportExportButtonText)
+			.OnClicked_Static(&OnImportExportButtonClicked)
+			.IsEnabled_Static(&GetImportExportButtonIsEnabled)
+		]
+	];
+}
+END_SLATE_FUNCTION_BUILD_OPTIMIZATION
+
+EVisibility FLandscapeEditorDetailCustomization_ImportExport::GetImportExportVisibility(bool bImport)
+{
+	return IsImporting() == bImport ? EVisibility::Visible : EVisibility::Collapsed;
+}
+
+bool FLandscapeEditorDetailCustomization_ImportExport::IsHeightmapEnabled()
+{
+	FEdModeLandscape* LandscapeEdMode = GetEditorMode();
+	check(LandscapeEdMode != nullptr);
+	if (LandscapeEdMode->UISettings)
+	{
+		return LandscapeEdMode->UISettings->bHeightmapSelected;
+	}
+
+	return false;
+}
+
+ECheckBoxState FLandscapeEditorDetailCustomization_ImportExport::GetHeightmapSelectedCheckState()
+{
+	return IsHeightmapEnabled() ? ECheckBoxState::Checked : ECheckBoxState::Unchecked;
+}
+
+void FLandscapeEditorDetailCustomization_ImportExport::OnHeightmapSelectedCheckStateChanged(ECheckBoxState CheckState)
+{
+	FEdModeLandscape* LandscapeEdMode = GetEditorMode();
+	check(LandscapeEdMode != nullptr);
+	if (LandscapeEdMode->UISettings)
+	{
+		if (CheckState != ECheckBoxState::Undetermined)
+		{
+			LandscapeEdMode->UISettings->bHeightmapSelected = CheckState == ECheckBoxState::Checked;
+		}
+	}
+}
+
+ECheckBoxState FLandscapeEditorDetailCustomization_ImportExport::ModeIsChecked(EImportExportMode::Type Value)
+{
+	FEdModeLandscape* LandscapeEdMode = GetEditorMode();
+	check(LandscapeEdMode != nullptr);
+	return LandscapeEdMode->ImportExportMode == Value ? ECheckBoxState::Checked : ECheckBoxState::Unchecked;
+}
+
+void FLandscapeEditorDetailCustomization_ImportExport::OnModeChanged(ECheckBoxState NewCheckedState, EImportExportMode::Type Value)
+{
+	if (NewCheckedState == ECheckBoxState::Checked)
+	{
+		FEdModeLandscape* LandscapeEdMode = GetEditorMode();
+		check(LandscapeEdMode != nullptr);
+		LandscapeEdMode->ImportExportMode = Value;
+	}
+}
+
+EVisibility FLandscapeEditorDetailCustomization_ImportExport::GetImportResultErrorVisibility(TSharedRef<IPropertyHandle> PropertyHandle_Result)
+{
+	ELandscapeImportResult ImportResult;
+	FPropertyAccess::Result Result = PropertyHandle_Result->GetValue((uint8&)ImportResult);
+
+	if (Result == FPropertyAccess::Fail)
+	{
+		return EVisibility::Collapsed;
+	}
+
+	if (Result == FPropertyAccess::MultipleValues)
+	{
+		return EVisibility::Visible;
+	}
+
+	if (ImportResult != ELandscapeImportResult::Success)
+	{
+		return EVisibility::Visible;
+	}
+
+	return EVisibility::Collapsed;
+}
+
+FSlateColor FLandscapeEditorDetailCustomization_ImportExport::GetImportResultErrorColor(TSharedRef<IPropertyHandle> PropertyHandle_Result)
+{
+	ELandscapeImportResult ImportResult;
+	FPropertyAccess::Result Result = PropertyHandle_Result->GetValue((uint8&)ImportResult);
+
+	if (Result == FPropertyAccess::Fail ||
+		Result == FPropertyAccess::MultipleValues)
+	{
+		return FCoreStyle::Get().GetColor("ErrorReporting.BackgroundColor");
+	}
+
+	switch (ImportResult)
+	{
+	case ELandscapeImportResult::Success:
+		return FCoreStyle::Get().GetColor("InfoReporting.BackgroundColor");
+	case ELandscapeImportResult::Warning:
+		return FCoreStyle::Get().GetColor("ErrorReporting.WarningBackgroundColor");
+	case ELandscapeImportResult::Error:
+		return FCoreStyle::Get().GetColor("ErrorReporting.BackgroundColor");
+	default:
+		check(0);
+		return FSlateColor();
+	}
+}
+
+void FLandscapeEditorDetailCustomization_ImportExport::SetFilename(const FText& NewValue, ETextCommit::Type CommitInfo, TSharedRef<IPropertyHandle> PropertyHandle_Filename)
+{
+	FString Filename = NewValue.ToString();
+	ensure(PropertyHandle_Filename->SetValue(Filename) == FPropertyAccess::Success);
+}
+
+FReply FLandscapeEditorDetailCustomization_ImportExport::OnBrowseFilenameButtonClicked(TSharedRef<IPropertyHandle> PropertyHandle_Filename)
+{
+	FEdModeLandscape* LandscapeEdMode = GetEditorMode();
+	check(LandscapeEdMode != nullptr);
+
+	// Prompt the user for the Filenames
+	IDesktopPlatform* DesktopPlatform = FDesktopPlatformModule::Get();
+	if (DesktopPlatform != nullptr)
+	{
+		bool bSuccess = false;
+		TArray<FString> Filenames;
+		ILandscapeEditorModule& LandscapeEditorModule = FModuleManager::GetModuleChecked<ILandscapeEditorModule>("LandscapeEditor");
+		if (IsImporting())
+		{
+			bSuccess = DesktopPlatform->OpenFileDialog(
+				FSlateApplication::Get().FindBestParentWindowHandleForDialogs(nullptr),
+				LOCTEXT("ImportHeightmap", "Import Heightmap").ToString(),
+				LandscapeEdMode->UISettings->LastImportPath,
+				TEXT(""),
+				LandscapeEditorModule.GetHeightmapImportDialogTypeString(),
+				EFileDialogFlags::None,
+				Filenames);
+		}
+		else
+		{
+			bSuccess = DesktopPlatform->SaveFileDialog(
+				FSlateApplication::Get().FindBestParentWindowHandleForDialogs(nullptr),
+				LOCTEXT("ExportHeightmap", "Export Heightmap").ToString(),
+				LandscapeEdMode->UISettings->LastImportPath,
+				TEXT(""),
+				LandscapeEditorModule.GetHeightmapExportDialogTypeString(),
+				EFileDialogFlags::None,
+				Filenames);
+		}
+
+		if (bSuccess)
+		{
+			ensure(PropertyHandle_Filename->SetValue(Filenames[0]) == FPropertyAccess::Success);
+			LandscapeEdMode->UISettings->LastImportPath = FPaths::GetPath(Filenames[0]);
+		}
+	}
+
+	return FReply::Handled();
+}
+
+void FLandscapeEditorDetailCustomization_ImportExport::OnImportHeightmapFilenameChanged()
+{
+	FEdModeLandscape* LandscapeEdMode = GetEditorMode();
+	check(LandscapeEdMode != nullptr);
+	LandscapeEdMode->UISettings->OnImportHeightmapFilenameChanged();
+}
+
+void FLandscapeEditorDetailCustomization_ImportExport::FormatFilename(TSharedRef<IPropertyHandle> PropertyHandle_Filename)
+{
+	FEdModeLandscape* LandscapeEdMode = GetEditorMode();
+	check(LandscapeEdMode != nullptr);
+	if (!LandscapeEdMode->UseSingleFileImport())
+	{
+		// If selected file as the _xN_yM suffix remove it
+		FString	FilePath;
+		if (PropertyHandle_Filename->GetValue(FilePath) == FPropertyAccess::Success)
+		{
+			FIntPoint OutCoord;
+			FString OutBaseFilePattern;
+			if (FLandscapeImportHelper::ExtractCoordinates(FPaths::GetBaseFilename(FilePath), OutCoord, OutBaseFilePattern))
+			{
+				PropertyHandle_Filename->SetValue(FString::Format(TEXT("{0}/{1}{2}"), { FPaths::GetPath(FilePath), OutBaseFilePattern, FPaths::GetExtension(FilePath, true) }));
+			}
+		}
+	}
+}
+
+bool FLandscapeEditorDetailCustomization_ImportExport::IsImporting()
+{
+	FEdModeLandscape* LandscapeEdMode = GetEditorMode();
+	check(LandscapeEdMode != nullptr);
+	return LandscapeEdMode->ImportExportMode == EImportExportMode::Import;
+}
+
+FText FLandscapeEditorDetailCustomization_ImportExport::GetImportExportButtonText()
+{
+	if (IsImporting())
+	{
+		return LOCTEXT("Import", "Import");
+	}
+
+	return LOCTEXT("Export", "Export");
+}
+
+FReply FLandscapeEditorDetailCustomization_ImportExport::OnImportExportButtonClicked()
+{
+	FEdModeLandscape* LandscapeEdMode = GetEditorMode();
+	check(LandscapeEdMode);
+	if (ULandscapeInfo* LandscapeInfo = LandscapeEdMode->CurrentToolTarget.LandscapeInfo.Get())
+	{
+		FIntRect LandscapeExtent;
+		if (!LandscapeInfo->GetLandscapeExtent(LandscapeExtent))
+		{
+			return FReply::Handled();
+		}
+	
+		if (IsImporting())
+		{
+			FGuid CurrentLayerGuid = LandscapeEdMode->GetCurrentLayerGuid();
+						
+			const ELandscapeLayerPaintingRestriction PaintRestriction = ELandscapeLayerPaintingRestriction::None;
+						
+			ELandscapeImportTransformType TransformType = LandscapeEdMode->UISettings->ImportType;
+									
+			FVector LocalGizmoPosition = LandscapeInfo->GetLandscapeProxy()->LandscapeActorToWorld().InverseTransformPosition(LandscapeEdMode->CurrentGizmoActor->GetActorLocation());
+			FIntPoint LocalGizmoPoint = FIntPoint(FMath::FloorToInt(LocalGizmoPosition.X), FMath::FloorToInt(LocalGizmoPosition.Y));
+			
+			// Update Gizmo Position if we exit and comeback into tool
+			LandscapeEdMode->UISettings->ImportLandscape_GizmoLocalPosition = LocalGizmoPoint;
+
+			FIntRect ImportRegion = LandscapeExtent;
+			FIntPoint ImportOffset(0, 0);
+			if (TransformType == ELandscapeImportTransformType::ExpandOffset)
+			{
+				ImportOffset = LocalGizmoPoint - FIntPoint(LandscapeExtent.Min.X, LandscapeExtent.Min.Y);
+			}
+			else if (TransformType == ELandscapeImportTransformType::None)
+			{
+				ImportRegion = FIntRect(LocalGizmoPoint.X, LocalGizmoPoint.Y, LocalGizmoPoint.X + LandscapeEdMode->UISettings->ImportLandscape_Width, LocalGizmoPoint.Y + LandscapeEdMode->UISettings->ImportLandscape_Height);
+			}
+			
+			// Import Heightmap
+			if (LandscapeEdMode->UISettings->bHeightmapSelected)
+			{
+				check(LandscapeEdMode->UISettings->ImportLandscape_HeightmapImportResult != ELandscapeImportResult::Error);
+				LandscapeEdMode->ImportHeightData(LandscapeInfo, CurrentLayerGuid, LandscapeEdMode->UISettings->ImportLandscape_HeightmapFilename, ImportRegion, TransformType, ImportOffset, PaintRestriction);
+			}
+
+			for (const FLandscapeImportLayer& ImportLayer : LandscapeEdMode->UISettings->ImportLandscape_Layers)
+			{
+				if (ImportLayer.bSelected)
+				{
+					check(ImportLayer.ImportResult != ELandscapeImportResult::Error);
+					LandscapeEdMode->ImportWeightData(LandscapeInfo, CurrentLayerGuid, ImportLayer.LayerInfo, ImportLayer.SourceFilePath, ImportRegion, TransformType, ImportOffset, PaintRestriction);
+				}
+			}
+		}
+		else
+		{
+			auto BuildExportFileName = [](const FString& Filename, FIntPoint FileOffset, bool bUseOffset)
+			{
+				if (bUseOffset)
+				{
+					FString Extension = FPaths::GetExtension(Filename, true);
+					FString BaseFilename = FPaths::GetBaseFilename(Filename, false);
+					return FString::Format(TEXT("{0}_x{1}_y{2}{3}"), { BaseFilename, FString::FromInt(FileOffset.X), FString::FromInt(FileOffset.Y), Extension });
+				}
+				
+				return Filename;
+			};
+
+			FScopedSetLandscapeEditingLayer Scope(LandscapeInfo->LandscapeActor.Get(), LandscapeEdMode->UISettings->bExportEditLayer ? LandscapeEdMode->GetCurrentLayerGuid() : FGuid());
+
+			LandscapeInfo->ForAllLandscapeProxies([LandscapeInfo, LandscapeEdMode, LandscapeExtent, BuildExportFileName](ALandscapeProxy* LandscapeProxy)
+			{
+				FIntRect ExportRegion;
+				if (LandscapeInfo->GetLandscapeExtent(LandscapeProxy, ExportRegion))
+				{
+					FIntPoint FileOffset;
+					if (!LandscapeEdMode->UseSingleFileImport())
+					{
+						ALandscape* Landscape = LandscapeInfo->LandscapeActor.Get();
+						check(Landscape);
+						FileOffset = FIntPoint((ExportRegion.Min.X - LandscapeExtent.Min.X) / Landscape->GridSize,
+								(ExportRegion.Min.Y - LandscapeExtent.Min.Y) / Landscape->GridSize);
+
+						// Remove the shared line/column that this proxy has with its neighbors because it 
+						// will be included by the neighbor or lost if there is none (that could become an option to avoid that loss)
+						ExportRegion.Max.X -= 1;
+						ExportRegion.Max.Y -= 1;
+					}
+											
+
+					if (LandscapeEdMode->UISettings->bHeightmapSelected)
+					{
+						FString ExportFilename = BuildExportFileName(LandscapeEdMode->UISettings->HeightmapExportFilename, FileOffset, !LandscapeEdMode->UseSingleFileImport());
+						LandscapeInfo->ExportHeightmap(ExportFilename, ExportRegion);
+					}
+					
+					for (const FLandscapeImportLayer& ImportLayer : LandscapeEdMode->UISettings->ImportLandscape_Layers)
+					{
+						if (ImportLayer.bSelected)
+						{
+							FString ExportFilename = BuildExportFileName(ImportLayer.ExportFilePath, FileOffset, !LandscapeEdMode->UseSingleFileImport());
+							LandscapeInfo->ExportLayer(ImportLayer.LayerInfo, ExportFilename, ExportRegion);
+						}
+					}
+				}
+			});
+			
+		}
+	}
+
+	return FReply::Handled();
+}
+
+bool FLandscapeEditorDetailCustomization_ImportExport::GetImportExportButtonIsEnabled()
+{
+	FEdModeLandscape* LandscapeEdMode = GetEditorMode();
+	check(LandscapeEdMode);
+
+	bool bHasOneSelection = false;
+	if (LandscapeEdMode->UISettings->bHeightmapSelected)
+	{
+		bHasOneSelection = true;
+		if (IsImporting())
+		{
+			if (LandscapeEdMode->UISettings->ImportLandscape_HeightmapImportResult == ELandscapeImportResult::Error || LandscapeEdMode->UISettings->ImportLandscape_HeightmapFilename.IsEmpty())
+			{
+				return false;
+			}
+		}
+		else
+		{
+			if (LandscapeEdMode->UISettings->HeightmapExportFilename.IsEmpty())
+			{
+				return false;
+			}
+		}
+	}
+
+	for (const FLandscapeImportLayer& ImportLayer : LandscapeEdMode->UISettings->ImportLandscape_Layers)
+	{
+		if (ImportLayer.bSelected)
+		{
+			bHasOneSelection = true;
+			
+			if (IsImporting())
+			{
+				if (ImportLayer.ImportResult == ELandscapeImportResult::Error || ImportLayer.SourceFilePath.IsEmpty())
+				{
+					return false;
+				}
+			}
+			else
+			{
+				if (ImportLayer.ExportFilePath.IsEmpty())
+				{
+					return false;
+				}
+			}
+		}
+	}
+
+	return bHasOneSelection;
+}
+
+TSharedRef<SWidget> FLandscapeEditorDetailCustomization_ImportExport::GetImportLandscapeResolutionMenu()
+{
+	FMenuBuilder MenuBuilder(true, nullptr);
+
+	FEdModeLandscape* LandscapeEdMode = GetEditorMode();
+	check(LandscapeEdMode != nullptr);
+	for (int32 i = 0; i < LandscapeEdMode->UISettings->HeightmapImportDescriptor.ImportResolutions.Num(); i++)
+	{
+		FFormatNamedArguments Args;
+		Args.Add(TEXT("Width"), LandscapeEdMode->UISettings->HeightmapImportDescriptor.ImportResolutions[i].Width);
+		Args.Add(TEXT("Height"), LandscapeEdMode->UISettings->HeightmapImportDescriptor.ImportResolutions[i].Height);
+		MenuBuilder.AddMenuEntry(FText::Format(LOCTEXT("ImportResolution_Format", "{Width}\u00D7{Height}"), Args), FText(), FSlateIcon(), FExecuteAction::CreateStatic(&FLandscapeEditorDetailCustomization_ImportExport::OnChangeImportLandscapeResolution, i));
+	}
+	
+	return MenuBuilder.MakeWidget();
+}
+
+void FLandscapeEditorDetailCustomization_ImportExport::OnChangeImportLandscapeResolution(int32 Index)
+{
+	FEdModeLandscape* LandscapeEdMode = GetEditorMode();
+	check(LandscapeEdMode != nullptr);
+	LandscapeEdMode->UISettings->OnChangeImportLandscapeResolution(Index);
+}
+
+FText FLandscapeEditorDetailCustomization_ImportExport::GetImportLandscapeResolution()
+{
+	FEdModeLandscape* LandscapeEdMode = GetEditorMode();
+	check(LandscapeEdMode != nullptr);
+	const int32	Width = LandscapeEdMode->UISettings->ImportLandscape_Width;
+	const int32	Height = LandscapeEdMode->UISettings->ImportLandscape_Height;
+		
+	if (Width != 0 && Height != 0)
+	{
+		FFormatNamedArguments Args;
+		Args.Add(TEXT("Width"), Width);
+		Args.Add(TEXT("Height"), Height);
+		return FText::Format(LOCTEXT("ImportResolution_Format", "{Width}\u00D7{Height}"), Args);
+	}
+	else
+	{
+		return LOCTEXT("ImportResolution_Invalid", "(invalid)");
+	}
+	
+	return FText::GetEmpty();
+}
+
+#undef LOCTEXT_NAMESPACE
diff --git a/Engine/Source/Editor/LandscapeEditor/Private/LandscapeEditorDetailCustomization_ImportExport.h b/Engine/Source/Editor/LandscapeEditor/Private/LandscapeEditorDetailCustomization_ImportExport.h
new file mode 100644
--- /dev/null
+++ b/Engine/Source/Editor/LandscapeEditor/Private/LandscapeEditorDetailCustomization_ImportExport.h
@@ -0,0 +1,58 @@
+// Copyright Epic Games, Inc. All Rights Reserved.
+
+#pragma once
+
+#include "CoreMinimal.h"
+#include "Layout/Visibility.h"
+#include "Input/Reply.h"
+#include "Widgets/SWidget.h"
+#include "Editor/LandscapeEditor/Private/LandscapeEdMode.h"
+#include "IDetailCustomization.h"
+#include "IPropertyTypeCustomization.h"
+#include "PropertyHandle.h"
+#include "LandscapeEditorDetailCustomization_Base.h"
+
+class IDetailLayoutBuilder;
+
+/**
+ * Slate widgets customizer for the "Import Export" tool
+ */
+
+class FLandscapeEditorDetailCustomization_ImportExport : public FLandscapeEditorDetailCustomization_Base
+{
+public:
+	FLandscapeEditorDetailCustomization_ImportExport()
+	{}
+
+	/** Makes a new instance of this detail layout class for a specific detail view requesting it */
+	static TSharedRef<IDetailCustomization> MakeInstance();
+
+	/** IDetailCustomization interface */
+	virtual void CustomizeDetails(IDetailLayoutBuilder& DetailBuilder) override;
+
+	static void FormatFilename(TSharedRef<IPropertyHandle> PropertyHandle_Filename);
+private:
+	static EVisibility GetImportExportVisibility(bool bImport);
+	static bool IsHeightmapEnabled();
+	static ECheckBoxState GetHeightmapSelectedCheckState();
+	static void OnHeightmapSelectedCheckStateChanged(ECheckBoxState CheckState);
+	static ECheckBoxState ModeIsChecked(EImportExportMode::Type Value);
+	static void OnModeChanged(ECheckBoxState NewCheckedState, EImportExportMode::Type Value);
+
+	static EVisibility GetImportResultErrorVisibility(TSharedRef<IPropertyHandle> PropertyHandle_ImportResult);
+	static FSlateColor GetImportResultErrorColor(TSharedRef<IPropertyHandle> PropertyHandle_ImportResult);
+	static void SetFilename(const FText& NewValue, ETextCommit::Type CommitInfo, TSharedRef<IPropertyHandle> PropertyHandle_Filename);
+	static FReply OnBrowseFilenameButtonClicked(TSharedRef<IPropertyHandle> PropertyHandle_Filename);
+
+	static FText GetImportExportButtonText();
+	static FReply OnImportExportButtonClicked();
+	static bool GetImportExportButtonIsEnabled();
+
+	static TSharedRef<SWidget> GetImportLandscapeResolutionMenu();
+	static void OnChangeImportLandscapeResolution(int32 DescriptorIndex);
+	static FText GetImportLandscapeResolution();
+
+	static void OnImportHeightmapFilenameChanged();
+	
+	static bool IsImporting();
+};
\ No newline at end of file
diff --git a/Engine/Source/Editor/LandscapeEditor/Private/LandscapeEditorDetailCustomization_ImportLayers.cpp b/Engine/Source/Editor/LandscapeEditor/Private/LandscapeEditorDetailCustomization_ImportLayers.cpp
new file mode 100644
--- /dev/null
+++ b/Engine/Source/Editor/LandscapeEditor/Private/LandscapeEditorDetailCustomization_ImportLayers.cpp
@@ -0,0 +1,495 @@
+// Copyright Epic Games, Inc. All Rights Reserved.
+
+#include "LandscapeEditorDetailCustomization_ImportLayers.h"
+#include "LandscapeEditorDetailCustomization_ImportExport.h"
+#include "AssetRegistry/AssetRegistryModule.h"
+#include "DesktopPlatformModule.h"
+#include "Editor.h"
+#include "LandscapeEditorModule.h"
+#include "LandscapeEditorObject.h"
+#include "SLandscapeEditor.h"
+
+#include "IDetailChildrenBuilder.h"
+#include "DetailWidgetRow.h"
+
+#include "PropertyCustomizationHelpers.h"
+
+#include "Dialogs/DlgPickAssetPath.h"
+#include "Widgets/Input/SCheckBox.h"
+#include "Widgets/Input/SButton.h"
+#include "Widgets/Images/SImage.h"
+#include "Widgets/Notifications/SErrorText.h"
+#include "Widgets/SToolTip.h"
+#include "Framework/MultiBox/MultiBoxBuilder.h"
+
+#define LOCTEXT_NAMESPACE "LandscapeEditor.ImportLayers"
+
+TSharedRef<IPropertyTypeCustomization> FLandscapeEditorStructCustomization_FLandscapeImportLayer::MakeInstance()
+{
+	return MakeShareable(new FLandscapeEditorStructCustomization_FLandscapeImportLayer);
+}
+
+bool FLandscapeEditorStructCustomization_FLandscapeImportLayer::IsImporting()
+{
+	if (IsToolActive("ImportExport"))
+	{
+		FEdModeLandscape* LandscapeEdMode = GetEditorMode();
+		check(LandscapeEdMode != nullptr);
+		return LandscapeEdMode->ImportExportMode == EImportExportMode::Import;
+	}
+
+	return true;
+}
+
+void FLandscapeEditorStructCustomization_FLandscapeImportLayer::CustomizeHeader(TSharedRef<IPropertyHandle> StructPropertyHandle, FDetailWidgetRow& HeaderRow, IPropertyTypeCustomizationUtils& StructCustomizationUtils)
+{
+}
+
+void FLandscapeEditorStructCustomization_FLandscapeImportLayer::CustomizeChildren(TSharedRef<IPropertyHandle> StructPropertyHandle, IDetailChildrenBuilder& ChildBuilder, IPropertyTypeCustomizationUtils& StructCustomizationUtils)
+{
+	TSharedRef<IPropertyHandle> PropertyHandle_LayerName = StructPropertyHandle->GetChildHandle(GET_MEMBER_NAME_CHECKED(FLandscapeImportLayer, LayerName)).ToSharedRef();
+	TSharedRef<IPropertyHandle> PropertyHandle_LayerInfo = StructPropertyHandle->GetChildHandle(GET_MEMBER_NAME_CHECKED(FLandscapeImportLayer, LayerInfo)).ToSharedRef();
+	TSharedRef<IPropertyHandle> PropertyHandle_SourceFilePath = StructPropertyHandle->GetChildHandle(GET_MEMBER_NAME_CHECKED(FLandscapeImportLayer, SourceFilePath)).ToSharedRef();
+	TSharedRef<IPropertyHandle> PropertyHandle_ExportFilePath = StructPropertyHandle->GetChildHandle(GET_MEMBER_NAME_CHECKED(FLandscapeImportLayer, ExportFilePath)).ToSharedRef();
+	TSharedRef<IPropertyHandle> PropertyHandle_ThumbnailMIC = StructPropertyHandle->GetChildHandle(GET_MEMBER_NAME_CHECKED(FLandscapeImportLayer, ThumbnailMIC)).ToSharedRef();
+	TSharedRef<IPropertyHandle> PropertyHandle_ImportResult = StructPropertyHandle->GetChildHandle(GET_MEMBER_NAME_CHECKED(FLandscapeImportLayer, ImportResult)).ToSharedRef();
+	TSharedRef<IPropertyHandle> PropertyHandle_ErrorMessage = StructPropertyHandle->GetChildHandle(GET_MEMBER_NAME_CHECKED(FLandscapeImportLayer, ErrorMessage)).ToSharedRef();
+	TSharedRef<IPropertyHandle> PropertyHandle_Selected = StructPropertyHandle->GetChildHandle(GET_MEMBER_NAME_CHECKED(FLandscapeImportLayer, bSelected)).ToSharedRef();
+	
+	PropertyHandle_SourceFilePath->SetOnPropertyValueChanged(FSimpleDelegate::CreateLambda([PropertyHandle_SourceFilePath]()
+	{
+		FLandscapeEditorDetailCustomization_ImportExport::FormatFilename(PropertyHandle_SourceFilePath);
+		FLandscapeEditorStructCustomization_FLandscapeImportLayer::OnImportWeightmapFilenameChanged();
+	}));
+
+	PropertyHandle_ExportFilePath->SetOnPropertyValueChanged(FSimpleDelegate::CreateLambda([PropertyHandle_ExportFilePath]()
+	{
+		FLandscapeEditorDetailCustomization_ImportExport::FormatFilename(PropertyHandle_ExportFilePath);
+	}));
+
+	FName LayerName;
+	FText LayerNameText;
+	FPropertyAccess::Result Result = PropertyHandle_LayerName->GetValue(LayerName);
+	checkSlow(Result == FPropertyAccess::Success);
+	LayerNameText = FText::FromName(LayerName);
+	if (Result == FPropertyAccess::MultipleValues)
+	{
+		LayerName = NAME_None;
+		LayerNameText = NSLOCTEXT("PropertyEditor", "MultipleValues", "Multiple Values");
+	}
+
+	UObject* ThumbnailMIC = nullptr;
+	Result = PropertyHandle_ThumbnailMIC->GetValue(ThumbnailMIC);
+	checkSlow(Result == FPropertyAccess::Success);
+
+	auto CreateFilenameWidget = [&](TSharedRef<IPropertyHandle> PropertyHandle_Filepath, bool bImport)
+	{
+		return SNew(SHorizontalBox)
+			.Visibility_Static(&GetImportExportVisibility, bImport)
+			+ SHorizontalBox::Slot()
+			.AutoWidth()
+			.Padding(0, 0, 2, 0)
+			[
+				SNew(SErrorText)
+				.Visibility_Static(&GetErrorVisibility, PropertyHandle_ImportResult)
+				.BackgroundColor_Static(&GetErrorColor, PropertyHandle_ImportResult)
+				.ErrorText(NSLOCTEXT("UnrealEd", "Error", "!"))
+				.ToolTip(
+					SNew(SToolTip)
+					.Text_Static(&GetErrorText, PropertyHandle_ErrorMessage)
+				)
+			]
+			+ SHorizontalBox::Slot()
+			[
+				PropertyHandle_Filepath->CreatePropertyValueWidget()
+			]
+			+ SHorizontalBox::Slot()
+			.AutoWidth()
+			.Padding(1, 0, 0, 0)
+			[
+				SNew(SButton)
+				.ContentPadding(FMargin(4, 0))
+				.Text(NSLOCTEXT("UnrealEd", "GenericOpenDialog", "..."))
+				.OnClicked_Static(&FLandscapeEditorStructCustomization_FLandscapeImportLayer::OnLayerFilenameButtonClicked, PropertyHandle_Filepath)
+			];
+	};
+
+	ChildBuilder.AddCustomRow(LayerNameText)
+	.NameContent()
+	[
+		SNew(SHorizontalBox)
+		+ SHorizontalBox::Slot()
+		.AutoWidth()
+		.VAlign(VAlign_Center)
+		.HAlign(HAlign_Left)
+		.Padding(FMargin(2))
+		[
+			SNew(SCheckBox)
+			.Visibility_Static(GetImportLayerSelectionVisibility)
+			.IsEnabled_Static(&FLandscapeEditorStructCustomization_FLandscapeImportLayer::IsValidLayerInfo, PropertyHandle_LayerInfo)
+			.IsChecked_Static(&FLandscapeEditorStructCustomization_FLandscapeImportLayer::GetImportLayerSelectedCheckState, PropertyHandle_Selected, PropertyHandle_LayerInfo)
+			.OnCheckStateChanged_Static(&FLandscapeEditorStructCustomization_FLandscapeImportLayer::OnImportLayerSelectedCheckStateChanged, PropertyHandle_Selected)
+			.ToolTipText_Static(&FLandscapeEditorStructCustomization_FLandscapeImportLayer::GetImportLayerSelectedToolTip, PropertyHandle_Selected, PropertyHandle_LayerInfo)
+		]
+		+ SHorizontalBox::Slot()
+		.FillWidth(1)
+		.VAlign(VAlign_Center)
+		.Padding(FMargin(2))
+		[
+			SNew(STextBlock)
+			.Font(StructCustomizationUtils.GetRegularFont())
+			.Text(LayerNameText)
+		]
+		+ SHorizontalBox::Slot()
+		.AutoWidth()
+		.VAlign(VAlign_Center)
+		.Padding(FMargin(2))
+		[
+			SNew(SLandscapeAssetThumbnail, ThumbnailMIC, StructCustomizationUtils.GetThumbnailPool().ToSharedRef())
+			.ThumbnailSize(FIntPoint(48, 48))
+		]
+	]
+	.ValueContent()
+	.MinDesiredWidth(250.0f) // copied from SPropertyEditorAsset::GetDesiredWidth
+	.MaxDesiredWidth(0)
+	[
+		SNew(SBox)
+		.VAlign(VAlign_Center)
+		.Padding(FMargin(0, 0, 12, 0)) // Line up with the other properties due to having no reset to default button
+		[
+			SNew(SVerticalBox)
+			+ SVerticalBox::Slot()
+			.AutoHeight()
+			[
+				SNew(SHorizontalBox)
+				.Visibility_Static(&FLandscapeEditorStructCustomization_FLandscapeImportLayer::GetLayerInfoAssignVisibility)
+				+ SHorizontalBox::Slot()
+				[
+					SNew(SObjectPropertyEntryBox)
+					.AllowedClass(ULandscapeLayerInfoObject::StaticClass())
+					.PropertyHandle(PropertyHandle_LayerInfo)
+					.OnShouldFilterAsset_Static(&ShouldFilterLayerInfo, LayerName)
+				]
+				+ SHorizontalBox::Slot()
+				.AutoWidth()
+				.VAlign(VAlign_Center)
+				[
+					SNew(SComboButton)
+					.ButtonStyle(FEditorStyle::Get(), "HoverHintOnly")
+					.HasDownArrow(false)
+					.ContentPadding(4.0f)
+					.ForegroundColor(FSlateColor::UseForeground())
+					.IsFocusable(false)
+					.ToolTipText(LOCTEXT("Target_Create", "Create Layer Info"))
+					.Visibility_Static(&GetImportLayerCreateVisibility, PropertyHandle_LayerInfo)
+					.OnGetMenuContent_Static(&OnGetImportLayerCreateMenu, PropertyHandle_LayerInfo, LayerName)
+					.ButtonContent()
+					[
+						SNew(SImage)
+						.Image(FEditorStyle::GetBrush("LandscapeEditor.Target_Create"))
+						.ColorAndOpacity(FSlateColor::UseForeground())
+					]
+				]
+			]
+			+ SVerticalBox::Slot()
+			.AutoHeight()
+			[
+				SNew(SVerticalBox)
+				.Visibility_Static(&FLandscapeEditorStructCustomization_FLandscapeImportLayer::GetImportLayerVisibility)
+				.IsEnabled_Static(&FLandscapeEditorStructCustomization_FLandscapeImportLayer::IsImportLayerSelected, PropertyHandle_Selected, PropertyHandle_LayerInfo)
+				+ SVerticalBox::Slot()
+				.AutoHeight()
+				[
+					CreateFilenameWidget(PropertyHandle_SourceFilePath, true)
+				]
+				+ SVerticalBox::Slot()
+				.AutoHeight()
+				[
+					CreateFilenameWidget(PropertyHandle_ExportFilePath, false)
+				]
+			]
+		]
+	];
+}
+
+void FLandscapeEditorStructCustomization_FLandscapeImportLayer::OnImportWeightmapFilenameChanged()
+{
+	FEdModeLandscape* LandscapeEdMode = GetEditorMode();
+	check(LandscapeEdMode != nullptr);
+	LandscapeEdMode->UISettings->OnImportWeightmapFilenameChanged();
+}
+
+FReply FLandscapeEditorStructCustomization_FLandscapeImportLayer::OnLayerFilenameButtonClicked(TSharedRef<IPropertyHandle> PropertyHandle_LayerFilename)
+{
+	FEdModeLandscape* LandscapeEdMode = GetEditorMode();
+	check(LandscapeEdMode != nullptr);
+
+	// Prompt the user for the Filenames
+	IDesktopPlatform* DesktopPlatform = FDesktopPlatformModule::Get();
+	if (DesktopPlatform != nullptr)
+	{
+		bool bSuccess = false;
+		TArray<FString> Filenames;
+		ILandscapeEditorModule& LandscapeEditorModule = FModuleManager::GetModuleChecked<ILandscapeEditorModule>("LandscapeEditor");
+		if (IsImporting())
+		{
+			bSuccess = DesktopPlatform->OpenFileDialog(
+				FSlateApplication::Get().FindBestParentWindowHandleForDialogs(nullptr),
+				LOCTEXT("ImportLayer", "Import Layer").ToString(),
+				LandscapeEdMode->UISettings->LastImportPath,
+				TEXT(""),
+				LandscapeEditorModule.GetWeightmapImportDialogTypeString(),
+				EFileDialogFlags::None,
+				Filenames);
+		}
+		else
+		{
+			bSuccess = DesktopPlatform->SaveFileDialog(
+				FSlateApplication::Get().FindBestParentWindowHandleForDialogs(nullptr),
+				LOCTEXT("ExportLayer", "Export Layer").ToString(),
+				LandscapeEdMode->UISettings->LastImportPath,
+				TEXT(""),
+				LandscapeEditorModule.GetWeightmapExportDialogTypeString(),
+				EFileDialogFlags::None,
+				Filenames);
+		}
+
+		if (bSuccess)
+		{
+			ensure(PropertyHandle_LayerFilename->SetValue(Filenames[0]) == FPropertyAccess::Success);
+			LandscapeEdMode->UISettings->LastImportPath = FPaths::GetPath(Filenames[0]);
+		}
+	}
+
+	return FReply::Handled();
+}
+
+bool FLandscapeEditorStructCustomization_FLandscapeImportLayer::ShouldFilterLayerInfo(const FAssetData& AssetData, FName LayerName)
+{
+	const FName LayerNameMetaData = AssetData.GetTagValueRef<FName>("LayerName");
+	if (!LayerNameMetaData.IsNone())
+	{
+		return LayerNameMetaData != LayerName;
+	}
+
+	ULandscapeLayerInfoObject* LayerInfo = CastChecked<ULandscapeLayerInfoObject>(AssetData.GetAsset());
+	return LayerInfo->LayerName != LayerName;
+}
+
+EVisibility FLandscapeEditorStructCustomization_FLandscapeImportLayer::GetImportLayerSelectionVisibility()
+{
+	return IsToolActive("ImportExport") ? EVisibility::Visible : EVisibility::Collapsed;
+}
+
+EVisibility FLandscapeEditorStructCustomization_FLandscapeImportLayer::GetImportLayerCreateVisibility(TSharedRef<IPropertyHandle> PropertyHandle_LayerInfo)
+{
+	if (IsToolActive("ImportExport"))
+	{
+		return EVisibility::Collapsed;
+	}
+
+	
+	return IsValidLayerInfo(PropertyHandle_LayerInfo) ? EVisibility::Collapsed : EVisibility::Visible;
+	
+
+	return EVisibility::Collapsed;
+}
+
+bool FLandscapeEditorStructCustomization_FLandscapeImportLayer::IsValidLayerInfo(TSharedRef<IPropertyHandle> PropertyHandle_LayerInfo)
+{
+	UObject* LayerInfoAsUObject = nullptr;
+	return (PropertyHandle_LayerInfo->GetValue(LayerInfoAsUObject) != FPropertyAccess::Fail && LayerInfoAsUObject != nullptr);
+}
+
+EVisibility FLandscapeEditorStructCustomization_FLandscapeImportLayer::GetImportExportVisibility(bool bImport)
+{
+	return IsImporting() == bImport ? EVisibility::Visible : EVisibility::Collapsed;
+}
+
+TSharedRef<SWidget> FLandscapeEditorStructCustomization_FLandscapeImportLayer::OnGetImportLayerCreateMenu(TSharedRef<IPropertyHandle> PropertyHandle_LayerInfo, FName LayerName)
+{
+	FMenuBuilder MenuBuilder(true, nullptr);
+
+	MenuBuilder.AddMenuEntry(LOCTEXT("Target_Create_Blended", "Weight-Blended Layer (normal)"), FText(), FSlateIcon(),
+		FUIAction(FExecuteAction::CreateStatic(&OnImportLayerCreateClicked, PropertyHandle_LayerInfo, LayerName, false)));
+
+	MenuBuilder.AddMenuEntry(LOCTEXT("Target_Create_NoWeightBlend", "Non Weight-Blended Layer"), FText(), FSlateIcon(),
+		FUIAction(FExecuteAction::CreateStatic(&OnImportLayerCreateClicked, PropertyHandle_LayerInfo, LayerName, true)));
+
+	return MenuBuilder.MakeWidget();
+}
+
+void FLandscapeEditorStructCustomization_FLandscapeImportLayer::OnImportLayerCreateClicked(TSharedRef<IPropertyHandle> PropertyHandle_LayerInfo, FName LayerName, bool bNoWeightBlend)
+{
+	FEdModeLandscape* LandscapeEdMode = GetEditorMode();
+	check(LandscapeEdMode != nullptr);
+	
+	// Hack as we don't have a direct world pointer in the EdMode...
+	ULevel* Level = LandscapeEdMode->CurrentGizmoActor->GetWorld()->GetCurrentLevel();
+
+	// Build default layer object name and package name
+	FName LayerObjectName = FName(*FString::Printf(TEXT("%s_LayerInfo"), *LayerName.ToString()));
+	FString Path = Level->GetOutermost()->GetName() + TEXT("_sharedassets/");
+	if (Path.StartsWith("/Temp/"))
+	{
+		Path = FString("/Game/") + Path.RightChop(FString("/Temp/").Len());
+	}
+	FString PackageName = Path + LayerObjectName.ToString();
+
+	TSharedRef<SDlgPickAssetPath> NewLayerDlg =
+		SNew(SDlgPickAssetPath)
+		.Title(LOCTEXT("CreateNewLayerInfo", "Create New Landscape Layer Info Object"))
+		.DefaultAssetPath(FText::FromString(PackageName));
+
+	if (NewLayerDlg->ShowModal() != EAppReturnType::Cancel)
+	{
+		PackageName = NewLayerDlg->GetFullAssetPath().ToString();
+		LayerObjectName = FName(*NewLayerDlg->GetAssetName().ToString());
+
+		UPackage* Package = CreatePackage(*PackageName);
+		ULandscapeLayerInfoObject* LayerInfo = NewObject<ULandscapeLayerInfoObject>(Package, LayerObjectName, RF_Public | RF_Standalone | RF_Transactional);
+		LayerInfo->LayerName = LayerName;
+		LayerInfo->bNoWeightBlend = bNoWeightBlend;
+
+		const UObject* LayerInfoAsUObject = LayerInfo; // HACK: If SetValue took a reference to a const ptr (T* const &) or a non-reference (T*) then this cast wouldn't be necessary
+		ensure(PropertyHandle_LayerInfo->SetValue(LayerInfoAsUObject) == FPropertyAccess::Success);
+
+		// Notify the asset registry
+		FAssetRegistryModule::AssetCreated(LayerInfo);
+
+		// Mark the package dirty...
+		Package->MarkPackageDirty();
+
+		// Show in the content browser
+		TArray<UObject*> Objects;
+		Objects.Add(LayerInfo);
+		GEditor->SyncBrowserToObjects(Objects);
+	}
+	
+}
+
+EVisibility FLandscapeEditorStructCustomization_FLandscapeImportLayer::GetLayerInfoAssignVisibility()
+{
+	return IsToolActive("ImportExport") ? EVisibility::Collapsed : EVisibility::Visible;
+}
+
+EVisibility FLandscapeEditorStructCustomization_FLandscapeImportLayer::GetImportLayerVisibility()
+{
+	if (IsToolActive("ImportExport"))
+	{
+		return EVisibility::Visible;
+	}
+
+	if (IsToolActive("NewLandscape"))
+	{
+		FEdModeLandscape* LandscapeEdMode = GetEditorMode();
+		if (LandscapeEdMode != nullptr)
+		{
+			if (LandscapeEdMode->NewLandscapePreviewMode == ENewLandscapePreviewMode::ImportLandscape)
+			{
+				return EVisibility::Visible;
+			}
+		}
+	}
+
+	return EVisibility::Collapsed;
+}
+
+EVisibility FLandscapeEditorStructCustomization_FLandscapeImportLayer::GetErrorVisibility(TSharedRef<IPropertyHandle> PropertyHandle_ImportResult)
+{
+	ELandscapeImportResult WeightmapImportResult;
+	FPropertyAccess::Result Result = PropertyHandle_ImportResult->GetValue((uint8&)WeightmapImportResult);
+
+	if (Result == FPropertyAccess::Fail ||
+		Result == FPropertyAccess::MultipleValues)
+	{
+		return EVisibility::Visible;
+	}
+
+	if (WeightmapImportResult != ELandscapeImportResult::Success)
+	{
+		return EVisibility::Visible;
+	}
+	return EVisibility::Collapsed;
+}
+
+FSlateColor FLandscapeEditorStructCustomization_FLandscapeImportLayer::GetErrorColor(TSharedRef<IPropertyHandle> PropertyHandle_ImportResult)
+{
+	ELandscapeImportResult WeightmapImportResult;
+	FPropertyAccess::Result Result = PropertyHandle_ImportResult->GetValue((uint8&)WeightmapImportResult);
+	check(Result == FPropertyAccess::Success);
+
+	if (Result == FPropertyAccess::MultipleValues)
+	{
+		return FCoreStyle::Get().GetColor("ErrorReporting.BackgroundColor");
+	}
+
+	switch (WeightmapImportResult)
+	{
+	case ELandscapeImportResult::Success:
+		return FCoreStyle::Get().GetColor("InfoReporting.BackgroundColor");
+	case ELandscapeImportResult::Warning:
+		return FCoreStyle::Get().GetColor("ErrorReporting.WarningBackgroundColor");
+	case ELandscapeImportResult::Error:
+		return FCoreStyle::Get().GetColor("ErrorReporting.BackgroundColor");
+	default:
+		check(0);
+		return FSlateColor();
+	}
+}
+
+FText FLandscapeEditorStructCustomization_FLandscapeImportLayer::GetErrorText(TSharedRef<IPropertyHandle> PropertyHandle_ErrorMessage)
+{
+	FText ErrorMessage;
+	FPropertyAccess::Result Result = PropertyHandle_ErrorMessage->GetValue(ErrorMessage);
+	if (Result == FPropertyAccess::Fail)
+	{
+		return LOCTEXT("Import_LayerUnknownError", "Unknown Error");
+	}
+	else if (Result == FPropertyAccess::MultipleValues)
+	{
+		return NSLOCTEXT("PropertyEditor", "MultipleValues", "Multiple Values");
+	}
+
+	return ErrorMessage;
+}
+
+bool FLandscapeEditorStructCustomization_FLandscapeImportLayer::IsImportLayerSelected(TSharedRef<IPropertyHandle> PropertyHandle_Selected, TSharedRef<IPropertyHandle> PropertyHandle_LayerInfo)
+{
+	if (IsToolActive("ImportExport"))
+	{
+		// Need a valid layer info to import/export
+		if (!IsValidLayerInfo(PropertyHandle_LayerInfo))
+		{
+			return false;
+		}
+
+		bool bSelected;
+		FPropertyAccess::Result Result = PropertyHandle_Selected->GetValue(bSelected);
+		return Result == FPropertyAccess::Success ? bSelected : false;
+	}
+
+	return true;
+}
+
+ECheckBoxState FLandscapeEditorStructCustomization_FLandscapeImportLayer::GetImportLayerSelectedCheckState(TSharedRef<IPropertyHandle> PropertyHandle_Selected, TSharedRef<IPropertyHandle> PropertyHandle_LayerInfo)
+{
+	return IsImportLayerSelected(PropertyHandle_Selected, PropertyHandle_LayerInfo) ? ECheckBoxState::Checked : ECheckBoxState::Unchecked;
+}
+
+FText FLandscapeEditorStructCustomization_FLandscapeImportLayer::GetImportLayerSelectedToolTip(TSharedRef<IPropertyHandle> PropertyHandle_Selected, TSharedRef<IPropertyHandle> PropertyHandle_LayerInfo)
+{
+	if (!IsValidLayerInfo(PropertyHandle_LayerInfo))
+	{
+		return LOCTEXT("InvalidLayerInfo", "This layer doesn't have a valid LayerInfo object assigned.");
+	}
+
+	return FText();
+}
+
+void FLandscapeEditorStructCustomization_FLandscapeImportLayer::OnImportLayerSelectedCheckStateChanged(ECheckBoxState CheckState, TSharedRef<IPropertyHandle> PropertyHandle_Selected)
+{
+	ensure(PropertyHandle_Selected->SetValue(CheckState == ECheckBoxState::Checked, EPropertyValueSetFlags::NotTransactable) == FPropertyAccess::Success);
+}
+
+#undef LOCTEXT_NAMESPACE
diff --git a/Engine/Source/Editor/LandscapeEditor/Private/LandscapeEditorDetailCustomization_ImportLayers.h b/Engine/Source/Editor/LandscapeEditor/Private/LandscapeEditorDetailCustomization_ImportLayers.h
new file mode 100644
--- /dev/null
+++ b/Engine/Source/Editor/LandscapeEditor/Private/LandscapeEditorDetailCustomization_ImportLayers.h
@@ -0,0 +1,45 @@
+// Copyright Epic Games, Inc. All Rights Reserved.
+
+#pragma once
+
+#include "CoreMinimal.h"
+#include "Layout/Visibility.h"
+#include "Styling/SlateColor.h"
+#include "Input/Reply.h"
+#include "Widgets/SWidget.h"
+#include "Editor/LandscapeEditor/Private/LandscapeEdMode.h"
+#include "PropertyHandle.h"
+#include "Editor/LandscapeEditor/Private/LandscapeEditorDetailCustomization_Base.h"
+
+class FLandscapeEditorStructCustomization_FLandscapeImportLayer : public FLandscapeEditorStructCustomization_Base
+{
+public:
+	static TSharedRef<IPropertyTypeCustomization> MakeInstance();
+
+	virtual void CustomizeHeader(TSharedRef<IPropertyHandle> StructPropertyHandle, class FDetailWidgetRow& HeaderRow, IPropertyTypeCustomizationUtils& StructCustomizationUtils) override;
+	virtual void CustomizeChildren(TSharedRef<IPropertyHandle> StructPropertyHandle, class IDetailChildrenBuilder& ChildBuilder, IPropertyTypeCustomizationUtils& StructCustomizationUtils) override;
+
+public:
+	static void OnImportWeightmapFilenameChanged();
+	static FReply OnLayerFilenameButtonClicked(TSharedRef<IPropertyHandle> PropertyHandle_LayerFilename);
+	static bool ShouldFilterLayerInfo(const struct FAssetData& AssetData, FName LayerName);
+
+	static EVisibility GetImportLayerCreateVisibility(TSharedRef<IPropertyHandle> PropertyHandle_LayerInfo);
+	static TSharedRef<SWidget> OnGetImportLayerCreateMenu(TSharedRef<IPropertyHandle> PropertyHandle_LayerInfo, FName LayerName);
+	static void OnImportLayerCreateClicked(TSharedRef<IPropertyHandle> PropertyHandle_LayerInfo, FName LayerName, bool bNoWeightBlend);
+	static EVisibility GetImportLayerSelectionVisibility();
+	static EVisibility GetImportExportVisibility(bool bImport);
+
+	static EVisibility GetLayerInfoAssignVisibility();
+	static EVisibility GetImportLayerVisibility();
+	static EVisibility GetErrorVisibility(TSharedRef<IPropertyHandle> PropertyHandle_ImportResult);
+	static FSlateColor GetErrorColor(TSharedRef<IPropertyHandle> PropertyHandle_ImportResult);
+	static FText GetErrorText(TSharedRef<IPropertyHandle> PropertyHandle_ErrorMessage);
+
+	static bool IsImportLayerSelected(TSharedRef<IPropertyHandle> PropertyHandle_Selected, TSharedRef<IPropertyHandle> PropertyHandle_LayerInfo);
+	static ECheckBoxState GetImportLayerSelectedCheckState(TSharedRef<IPropertyHandle> PropertyHandle_Selected, TSharedRef<IPropertyHandle> PropertyHandle_LayerInfo);
+	static void OnImportLayerSelectedCheckStateChanged(ECheckBoxState CheckState, TSharedRef<IPropertyHandle> PropertyHandle_Selected);
+	static FText GetImportLayerSelectedToolTip(TSharedRef<IPropertyHandle> PropertyHandle_Selected, TSharedRef<IPropertyHandle> PropertyHandle_LayerInfo);
+	static bool IsValidLayerInfo(TSharedRef<IPropertyHandle> PropertyHandle_LayerInfo);
+	static bool IsImporting();
+};
diff --git a/Engine/Source/Editor/LandscapeEditor/Private/LandscapeEditorDetailCustomization_NewLandscape.cpp b/Engine/Source/Editor/LandscapeEditor/Private/LandscapeEditorDetailCustomization_NewLandscape.cpp
--- a/Engine/Source/Editor/LandscapeEditor/Private/LandscapeEditorDetailCustomization_NewLandscape.cpp
+++ b/Engine/Source/Editor/LandscapeEditor/Private/LandscapeEditorDetailCustomization_NewLandscape.cpp
@@ -1,6 +1,7 @@
 // Copyright Epic Games, Inc. All Rights Reserved.
 
 #include "LandscapeEditorDetailCustomization_NewLandscape.h"
+#include "LandscapeEditorDetailCustomization_ImportExport.h"
 #include "Framework/Commands/UIAction.h"
 #include "Widgets/Text/STextBlock.h"
 #include "Framework/MultiBox/MultiBoxBuilder.h"
@@ -19,7 +20,6 @@
 #include "LandscapeEditorModule.h"
 #include "LandscapeEditorObject.h"
 #include "Landscape.h"
-#include "LandscapeEditorUtils.h"
 #include "LandscapeConfigHelper.h"
 #include "LandscapeImportHelper.h"
 
@@ -109,7 +109,11 @@ void FLandscapeEditorDetailCustomization_NewLandscape::CustomizeDetails(IDetailL
 	TSharedRef<IPropertyHandle> PropertyHandle_HeightmapErrorMessage = DetailBuilder.GetProperty(GET_MEMBER_NAME_CHECKED(ULandscapeEditorObject, ImportLandscape_HeightmapErrorMessage));
 	DetailBuilder.HideProperty(PropertyHandle_HeightmapImportResult);
 	DetailBuilder.HideProperty(PropertyHandle_HeightmapErrorMessage);
-	PropertyHandle_HeightmapFilename->SetOnPropertyValueChanged(FSimpleDelegate::CreateSP(this, &FLandscapeEditorDetailCustomization_NewLandscape::OnImportHeightmapFilenameChanged));
+	PropertyHandle_HeightmapFilename->SetOnPropertyValueChanged(FSimpleDelegate::CreateLambda([this, PropertyHandle_HeightmapFilename]()
+	{
+		FLandscapeEditorDetailCustomization_ImportExport::FormatFilename(PropertyHandle_HeightmapFilename);
+		OnImportHeightmapFilenameChanged();
+	}));
 	NewLandscapeCategory.AddProperty(PropertyHandle_HeightmapFilename)
 	.Visibility(TAttribute<EVisibility>::Create(TAttribute<EVisibility>::FGetter::CreateStatic(&GetVisibilityOnlyInNewLandscapeMode, ENewLandscapePreviewMode::ImportLandscape)))
 	.CustomWidget()
@@ -1135,317 +1139,4 @@ EVisibility FLandscapeEditorDetailCustomization_NewLandscape::GetMaterialTipVisi
 	return EVisibility::Collapsed;
 }
 
-//////////////////////////////////////////////////////////////////////////
-
-TSharedRef<IPropertyTypeCustomization> FLandscapeEditorStructCustomization_FLandscapeImportLayer::MakeInstance()
-{
-	return MakeShareable(new FLandscapeEditorStructCustomization_FLandscapeImportLayer);
-}
-
-void FLandscapeEditorStructCustomization_FLandscapeImportLayer::CustomizeHeader(TSharedRef<IPropertyHandle> StructPropertyHandle, FDetailWidgetRow& HeaderRow, IPropertyTypeCustomizationUtils& StructCustomizationUtils)
-{
-}
-
-void FLandscapeEditorStructCustomization_FLandscapeImportLayer::CustomizeChildren(TSharedRef<IPropertyHandle> StructPropertyHandle, IDetailChildrenBuilder& ChildBuilder, IPropertyTypeCustomizationUtils& StructCustomizationUtils)
-{
-	TSharedRef<IPropertyHandle> PropertyHandle_LayerName = StructPropertyHandle->GetChildHandle(GET_MEMBER_NAME_CHECKED(FLandscapeImportLayer, LayerName)).ToSharedRef();
-	TSharedRef<IPropertyHandle> PropertyHandle_LayerInfo = StructPropertyHandle->GetChildHandle(GET_MEMBER_NAME_CHECKED(FLandscapeImportLayer, LayerInfo)).ToSharedRef();
-	TSharedRef<IPropertyHandle> PropertyHandle_SourceFilePath = StructPropertyHandle->GetChildHandle(GET_MEMBER_NAME_CHECKED(FLandscapeImportLayer, SourceFilePath)).ToSharedRef();
-	TSharedRef<IPropertyHandle> PropertyHandle_ThumbnailMIC = StructPropertyHandle->GetChildHandle(GET_MEMBER_NAME_CHECKED(FLandscapeImportLayer, ThumbnailMIC)).ToSharedRef();
-	TSharedRef<IPropertyHandle> PropertyHandle_ImportResult = StructPropertyHandle->GetChildHandle(GET_MEMBER_NAME_CHECKED(FLandscapeImportLayer, ImportResult)).ToSharedRef();
-	TSharedRef<IPropertyHandle> PropertyHandle_ErrorMessage = StructPropertyHandle->GetChildHandle(GET_MEMBER_NAME_CHECKED(FLandscapeImportLayer, ErrorMessage)).ToSharedRef();
-
-	FName LayerName;
-	FText LayerNameText;
-	FPropertyAccess::Result Result = PropertyHandle_LayerName->GetValue(LayerName);
-	checkSlow(Result == FPropertyAccess::Success);
-	LayerNameText = FText::FromName(LayerName);
-	if (Result == FPropertyAccess::MultipleValues)
-	{
-		LayerName = NAME_None;
-		LayerNameText = NSLOCTEXT("PropertyEditor", "MultipleValues", "Multiple Values");
-	}
-
-	UObject* ThumbnailMIC = nullptr;
-	Result = PropertyHandle_ThumbnailMIC->GetValue(ThumbnailMIC);
-	checkSlow(Result == FPropertyAccess::Success);
-
-	ChildBuilder.AddCustomRow(LayerNameText)
-	.NameContent()
-	[
-		SNew(SHorizontalBox)
-		+ SHorizontalBox::Slot()
-		.FillWidth(1)
-		.VAlign(VAlign_Center)
-		.Padding(FMargin(2))
-		[
-			SNew(STextBlock)
-			.Font(StructCustomizationUtils.GetRegularFont())
-			.Text(LayerNameText)
-		]
-		+ SHorizontalBox::Slot()
-		.AutoWidth()
-		.VAlign(VAlign_Center)
-		.Padding(FMargin(2))
-		[
-			SNew(SLandscapeAssetThumbnail, ThumbnailMIC, StructCustomizationUtils.GetThumbnailPool().ToSharedRef())
-			.ThumbnailSize(FIntPoint(48, 48))
-		]
-	]
-	.ValueContent()
-	.MinDesiredWidth(250.0f) // copied from SPropertyEditorAsset::GetDesiredWidth
-	.MaxDesiredWidth(0)
-	[
-		SNew(SBox)
-		.VAlign(VAlign_Center)
-		.Padding(FMargin(0,0,12,0)) // Line up with the other properties due to having no reset to default button
-		[
-			SNew(SVerticalBox)
-			+ SVerticalBox::Slot()
-			.AutoHeight()
-			[
-				SNew(SHorizontalBox)
-				+ SHorizontalBox::Slot()
-				[
-					SNew(SObjectPropertyEntryBox)
-					.AllowedClass(ULandscapeLayerInfoObject::StaticClass())
-					.PropertyHandle(PropertyHandle_LayerInfo)
-					.OnShouldFilterAsset_Static(&ShouldFilterLayerInfo, LayerName)
-				]
-				+ SHorizontalBox::Slot()
-				.AutoWidth()
-				.VAlign(VAlign_Center)
-				[
-					SNew(SComboButton)
-					//.ButtonStyle( FEditorStyle::Get(), "NoBorder" )
-					.ButtonStyle( FEditorStyle::Get(), "HoverHintOnly" )
-					.HasDownArrow(false)
-					//.ContentPadding(0)
-					.ContentPadding(4.0f)
-					.ForegroundColor(FSlateColor::UseForeground())
-					.IsFocusable(false)
-					.ToolTipText(LOCTEXT("Target_Create", "Create Layer Info"))
-					.Visibility_Static(&GetImportLayerCreateVisibility, PropertyHandle_LayerInfo)
-					.OnGetMenuContent_Static(&OnGetImportLayerCreateMenu, PropertyHandle_LayerInfo, LayerName)
-					.ButtonContent()
-					[
-						SNew(SImage)
-						.Image(FEditorStyle::GetBrush("LandscapeEditor.Target_Create"))
-						.ColorAndOpacity(FSlateColor::UseForeground())
-					]
-				]
-			]
-			+ SVerticalBox::Slot()
-			.AutoHeight()
-			[
-				SNew(SHorizontalBox)
-				.Visibility_Static(&FLandscapeEditorDetailCustomization_NewLandscape::GetVisibilityOnlyInNewLandscapeMode, ENewLandscapePreviewMode::ImportLandscape)
-				+ SHorizontalBox::Slot()
-				.AutoWidth()
-				.Padding(0, 0, 2, 0)
-				[
-					SNew(SErrorText)
-					.Visibility_Static(&GetErrorVisibility, PropertyHandle_ImportResult)
-					.BackgroundColor_Static(&GetErrorColor, PropertyHandle_ImportResult)
-					.ErrorText(NSLOCTEXT("UnrealEd", "Error", "!"))
-					.ToolTip(
-						SNew(SToolTip)
-						.Text_Static(&GetErrorText, PropertyHandle_ErrorMessage)
-					)
-				]
-				+ SHorizontalBox::Slot()
-				[
-					PropertyHandle_SourceFilePath->CreatePropertyValueWidget()
-				]
-				+ SHorizontalBox::Slot()
-				.AutoWidth()
-				.Padding(1, 0, 0, 0)
-				[
-					SNew(SButton)
-					.ContentPadding(FMargin(4, 0))
-					.Text(NSLOCTEXT("UnrealEd", "GenericOpenDialog", "..."))
-					.OnClicked_Static(&FLandscapeEditorStructCustomization_FLandscapeImportLayer::OnLayerFilenameButtonClicked, PropertyHandle_SourceFilePath)
-				]
-			]
-		]
-	];
-}
-
-FReply FLandscapeEditorStructCustomization_FLandscapeImportLayer::OnLayerFilenameButtonClicked(TSharedRef<IPropertyHandle> PropertyHandle_LayerFilename)
-{
-	FEdModeLandscape* LandscapeEdMode = GetEditorMode();
-	check(LandscapeEdMode != nullptr);
-
-	// Prompt the user for the Filenames
-	IDesktopPlatform* DesktopPlatform = FDesktopPlatformModule::Get();
-	if (DesktopPlatform != nullptr)
-	{
-		ILandscapeEditorModule& LandscapeEditorModule = FModuleManager::GetModuleChecked<ILandscapeEditorModule>("LandscapeEditor");
-		const TCHAR* FileTypes = LandscapeEditorModule.GetWeightmapImportDialogTypeString();
-
-		TArray<FString> OpenFilenames;
-		bool bOpened = DesktopPlatform->OpenFileDialog(
-			FSlateApplication::Get().FindBestParentWindowHandleForDialogs(nullptr),
-			NSLOCTEXT("UnrealEd", "Import", "Import").ToString(),
-			LandscapeEdMode->UISettings->LastImportPath,
-			TEXT(""),
-			FileTypes,
-			EFileDialogFlags::None,
-			OpenFilenames);
-
-		if (bOpened)
-		{
-			ensure(PropertyHandle_LayerFilename->SetValue(OpenFilenames[0]) == FPropertyAccess::Success);
-			LandscapeEdMode->UISettings->LastImportPath = FPaths::GetPath(OpenFilenames[0]);
-		}
-	}
-
-	return FReply::Handled();
-}
-
-bool FLandscapeEditorStructCustomization_FLandscapeImportLayer::ShouldFilterLayerInfo(const FAssetData& AssetData, FName LayerName)
-{
-	const FName LayerNameMetaData = AssetData.GetTagValueRef<FName>("LayerName");
-	if (!LayerNameMetaData.IsNone())
-	{
-		return LayerNameMetaData != LayerName;
-	}
-
-	ULandscapeLayerInfoObject* LayerInfo = CastChecked<ULandscapeLayerInfoObject>(AssetData.GetAsset());
-	return LayerInfo->LayerName != LayerName;
-}
-
-EVisibility FLandscapeEditorStructCustomization_FLandscapeImportLayer::GetImportLayerCreateVisibility(TSharedRef<IPropertyHandle> PropertyHandle_LayerInfo)
-{
-	UObject* LayerInfoAsUObject = nullptr;
-	if (PropertyHandle_LayerInfo->GetValue(LayerInfoAsUObject) != FPropertyAccess::Fail &&
-		LayerInfoAsUObject == nullptr)
-	{
-		return EVisibility::Visible;
-	}
-
-	return EVisibility::Collapsed;
-}
-
-TSharedRef<SWidget> FLandscapeEditorStructCustomization_FLandscapeImportLayer::OnGetImportLayerCreateMenu(TSharedRef<IPropertyHandle> PropertyHandle_LayerInfo, FName LayerName)
-{
-	FMenuBuilder MenuBuilder(true, nullptr);
-
-	MenuBuilder.AddMenuEntry(LOCTEXT("Target_Create_Blended", "Weight-Blended Layer (normal)"), FText(), FSlateIcon(),
-		FUIAction(FExecuteAction::CreateStatic(&OnImportLayerCreateClicked, PropertyHandle_LayerInfo, LayerName, false)));
-
-	MenuBuilder.AddMenuEntry(LOCTEXT("Target_Create_NoWeightBlend", "Non Weight-Blended Layer"), FText(), FSlateIcon(),
-		FUIAction(FExecuteAction::CreateStatic(&OnImportLayerCreateClicked, PropertyHandle_LayerInfo, LayerName, true)));
-
-	return MenuBuilder.MakeWidget();
-}
-
-void FLandscapeEditorStructCustomization_FLandscapeImportLayer::OnImportLayerCreateClicked(TSharedRef<IPropertyHandle> PropertyHandle_LayerInfo, FName LayerName, bool bNoWeightBlend)
-{
-	FEdModeLandscape* LandscapeEdMode = GetEditorMode();
-	if (LandscapeEdMode != nullptr)
-	{
-		// Hack as we don't have a direct world pointer in the EdMode...
-		ULevel* Level = LandscapeEdMode->CurrentGizmoActor->GetWorld()->GetCurrentLevel();
-
-		// Build default layer object name and package name
-		FName LayerObjectName = FName(*FString::Printf(TEXT("%s_LayerInfo"), *LayerName.ToString()));
-		FString Path = Level->GetOutermost()->GetName() + TEXT("_sharedassets/");
-		if (Path.StartsWith("/Temp/"))
-		{
-			Path = FString("/Game/") + Path.RightChop(FString("/Temp/").Len());
-		}
-		FString PackageName = Path + LayerObjectName.ToString();
-
-		TSharedRef<SDlgPickAssetPath> NewLayerDlg =
-			SNew(SDlgPickAssetPath)
-			.Title(LOCTEXT("CreateNewLayerInfo", "Create New Landscape Layer Info Object"))
-			.DefaultAssetPath(FText::FromString(PackageName));
-
-		if (NewLayerDlg->ShowModal() != EAppReturnType::Cancel)
-		{
-			PackageName = NewLayerDlg->GetFullAssetPath().ToString();
-			LayerObjectName = FName(*NewLayerDlg->GetAssetName().ToString());
-
-			UPackage* Package = CreatePackage( *PackageName);
-			ULandscapeLayerInfoObject* LayerInfo = NewObject<ULandscapeLayerInfoObject>(Package, LayerObjectName, RF_Public | RF_Standalone | RF_Transactional);
-			LayerInfo->LayerName = LayerName;
-			LayerInfo->bNoWeightBlend = bNoWeightBlend;
-
-			const UObject* LayerInfoAsUObject = LayerInfo; // HACK: If SetValue took a reference to a const ptr (T* const &) or a non-reference (T*) then this cast wouldn't be necessary
-			ensure(PropertyHandle_LayerInfo->SetValue(LayerInfoAsUObject) == FPropertyAccess::Success);
-
-			// Notify the asset registry
-			FAssetRegistryModule::AssetCreated(LayerInfo);
-
-			// Mark the package dirty...
-			Package->MarkPackageDirty();
-
-			// Show in the content browser
-			TArray<UObject*> Objects;
-			Objects.Add(LayerInfo);
-			GEditor->SyncBrowserToObjects(Objects);
-		}
-	}
-}
-
-EVisibility FLandscapeEditorStructCustomization_FLandscapeImportLayer::GetErrorVisibility(TSharedRef<IPropertyHandle> PropertyHandle_ImportResult)
-{
-	ELandscapeImportResult WeightmapImportResult;
-	FPropertyAccess::Result Result = PropertyHandle_ImportResult->GetValue((uint8&)WeightmapImportResult);
-
-	if (Result == FPropertyAccess::Fail ||
-		Result == FPropertyAccess::MultipleValues)
-	{
-		return EVisibility::Visible;
-	}
-
-	if (WeightmapImportResult != ELandscapeImportResult::Success)
-	{
-		return EVisibility::Visible;
-	}
-	return EVisibility::Collapsed;
-}
-
-FSlateColor FLandscapeEditorStructCustomization_FLandscapeImportLayer::GetErrorColor(TSharedRef<IPropertyHandle> PropertyHandle_ImportResult)
-{
-	ELandscapeImportResult WeightmapImportResult;
-	FPropertyAccess::Result Result = PropertyHandle_ImportResult->GetValue((uint8&)WeightmapImportResult);
-	check(Result == FPropertyAccess::Success);
-
-	if (Result == FPropertyAccess::MultipleValues)
-	{
-		return FCoreStyle::Get().GetColor("ErrorReporting.BackgroundColor");
-	}
-
-	switch (WeightmapImportResult)
-	{
-	case ELandscapeImportResult::Success:
-		return FCoreStyle::Get().GetColor("InfoReporting.BackgroundColor");
-	case ELandscapeImportResult::Warning:
-		return FCoreStyle::Get().GetColor("ErrorReporting.WarningBackgroundColor");
-	case ELandscapeImportResult::Error:
-		return FCoreStyle::Get().GetColor("ErrorReporting.BackgroundColor");
-	default:
-		check(0);
-		return FSlateColor();
-	}
-}
-
-FText FLandscapeEditorStructCustomization_FLandscapeImportLayer::GetErrorText(TSharedRef<IPropertyHandle> PropertyHandle_ErrorMessage)
-{
-	FText ErrorMessage;
-	FPropertyAccess::Result Result = PropertyHandle_ErrorMessage->GetValue(ErrorMessage);
-	if (Result == FPropertyAccess::Fail)
-	{
-		return LOCTEXT("Import_LayerUnknownError", "Unknown Error");
-	}
-	else if (Result == FPropertyAccess::MultipleValues)
-	{
-		return NSLOCTEXT("PropertyEditor", "MultipleValues", "Multiple Values");
-	}
-
-	return ErrorMessage;
-}
-
 #undef LOCTEXT_NAMESPACE
diff --git a/Engine/Source/Editor/LandscapeEditor/Private/LandscapeEditorDetailCustomization_NewLandscape.h b/Engine/Source/Editor/LandscapeEditor/Private/LandscapeEditorDetailCustomization_NewLandscape.h
--- a/Engine/Source/Editor/LandscapeEditor/Private/LandscapeEditorDetailCustomization_NewLandscape.h
+++ b/Engine/Source/Editor/LandscapeEditor/Private/LandscapeEditorDetailCustomization_NewLandscape.h
@@ -84,24 +84,3 @@ public:
 protected:
 	bool bUsingSlider;
 };
-
-class FLandscapeEditorStructCustomization_FLandscapeImportLayer : public FLandscapeEditorStructCustomization_Base
-{
-public:
-	static TSharedRef<IPropertyTypeCustomization> MakeInstance();
-
-	virtual void CustomizeHeader(TSharedRef<IPropertyHandle> StructPropertyHandle, class FDetailWidgetRow& HeaderRow, IPropertyTypeCustomizationUtils& StructCustomizationUtils) override;
-	virtual void CustomizeChildren(TSharedRef<IPropertyHandle> StructPropertyHandle, class IDetailChildrenBuilder& ChildBuilder, IPropertyTypeCustomizationUtils& StructCustomizationUtils) override;
-
-public:
-	static FReply OnLayerFilenameButtonClicked(TSharedRef<IPropertyHandle> PropertyHandle_LayerFilename);
-	static bool ShouldFilterLayerInfo(const struct FAssetData& AssetData, FName LayerName);
-
-	static EVisibility GetImportLayerCreateVisibility(TSharedRef<IPropertyHandle> PropertyHandle_LayerInfo);
-	static TSharedRef<SWidget> OnGetImportLayerCreateMenu(TSharedRef<IPropertyHandle> PropertyHandle_LayerInfo, FName LayerName);
-	static void OnImportLayerCreateClicked(TSharedRef<IPropertyHandle> PropertyHandle_LayerInfo, FName LayerName, bool bNoWeightBlend);
-
-	static EVisibility GetErrorVisibility(TSharedRef<IPropertyHandle> PropertyHandle_ImportResult);
-	static FSlateColor GetErrorColor(TSharedRef<IPropertyHandle> PropertyHandle_ImportResult);
-	static FText GetErrorText(TSharedRef<IPropertyHandle> PropertyHandle_ErrorMessage);
-};
diff --git a/Engine/Source/Editor/LandscapeEditor/Private/LandscapeEditorDetails.cpp b/Engine/Source/Editor/LandscapeEditor/Private/LandscapeEditorDetails.cpp
--- a/Engine/Source/Editor/LandscapeEditor/Private/LandscapeEditorDetails.cpp
+++ b/Engine/Source/Editor/LandscapeEditor/Private/LandscapeEditorDetails.cpp
@@ -15,6 +15,7 @@
 #include "LandscapeEditorDetailCustomization_CopyPaste.h"
 #include "LandscapeEditorDetailCustomization_MiscTools.h"
 #include "LandscapeEditorDetailCustomization_AlphaBrush.h"
+#include "LandscapeEditorDetailCustomization_ImportExport.h"
 #include "DetailWidgetRow.h"
 #include "LandscapeEditorDetailCustomization_TargetLayers.h"
 #include "DetailCategoryBuilder.h"
@@ -142,6 +143,8 @@ void FLandscapeEditorDetails::CustomizeDetails(IDetailLayoutBuilder& DetailBuild
 	// Tools:
 	Customization_NewLandscape = MakeShareable(new FLandscapeEditorDetailCustomization_NewLandscape);
 	Customization_NewLandscape->CustomizeDetails(DetailBuilder);
+	Customization_ImportExport = MakeShareable(new FLandscapeEditorDetailCustomization_ImportExport);
+	Customization_ImportExport->CustomizeDetails(DetailBuilder);
 	Customization_ResizeLandscape = MakeShareable(new FLandscapeEditorDetailCustomization_ResizeLandscape);
 	Customization_ResizeLandscape->CustomizeDetails(DetailBuilder);
 	Customization_CopyPaste = MakeShareable(new FLandscapeEditorDetailCustomization_CopyPaste);
diff --git a/Engine/Source/Editor/LandscapeEditor/Private/LandscapeEditorDetails.h b/Engine/Source/Editor/LandscapeEditor/Private/LandscapeEditorDetails.h
--- a/Engine/Source/Editor/LandscapeEditor/Private/LandscapeEditorDetails.h
+++ b/Engine/Source/Editor/LandscapeEditor/Private/LandscapeEditorDetails.h
@@ -18,6 +18,7 @@ class FLandscapeEditorDetailCustomization_MiscTools;
 class FLandscapeEditorDetailCustomization_NewLandscape;
 class FLandscapeEditorDetailCustomization_ResizeLandscape;
 class FLandscapeEditorDetailCustomization_TargetLayers;
+class FLandscapeEditorDetailCustomization_ImportExport;
 class FUICommandList;
 class IDetailLayoutBuilder;
 class ULandscapeInfo;
@@ -79,6 +80,7 @@ protected:
 
 	TSharedPtr<FLandscapeEditorDetailCustomization_NewLandscape> Customization_NewLandscape;
 	TSharedPtr<FLandscapeEditorDetailCustomization_ResizeLandscape> Customization_ResizeLandscape;
+	TSharedPtr<FLandscapeEditorDetailCustomization_ImportExport> Customization_ImportExport;
 	TSharedPtr<FLandscapeEditorDetailCustomization_CopyPaste> Customization_CopyPaste;
 	TSharedPtr<FLandscapeEditorDetailCustomization_MiscTools> Customization_MiscTools;
 	TSharedPtr<FLandscapeEditorDetailCustomization_AlphaBrush> Customization_AlphaBrush;
diff --git a/Engine/Source/Editor/LandscapeEditor/Private/LandscapeEditorModule.cpp b/Engine/Source/Editor/LandscapeEditor/Private/LandscapeEditorModule.cpp
--- a/Engine/Source/Editor/LandscapeEditor/Private/LandscapeEditorModule.cpp
+++ b/Engine/Source/Editor/LandscapeEditor/Private/LandscapeEditorModule.cpp
@@ -24,6 +24,7 @@
 #include "LandscapeEditorDetails.h"
 #include "LandscapeEditorDetailCustomization_NewLandscape.h"
 #include "LandscapeEditorDetailCustomization_CopyPaste.h"
+#include "LandscapeEditorDetailCustomization_ImportLayers.h"
 #include "LandscapeSplineDetails.h"
 
 #include "LevelEditor.h"
diff --git a/Engine/Source/Editor/LandscapeEditor/Private/LandscapeEditorObject.cpp b/Engine/Source/Editor/LandscapeEditor/Private/LandscapeEditorObject.cpp
--- a/Engine/Source/Editor/LandscapeEditor/Private/LandscapeEditorObject.cpp
+++ b/Engine/Source/Editor/LandscapeEditor/Private/LandscapeEditorObject.cpp
@@ -552,7 +552,7 @@ bool ULandscapeEditorObject::UseSingleFileImport() const
 	return true;
 }
 
-void ULandscapeEditorObject::OnImportHeightmapFilenameChanged()
+void ULandscapeEditorObject::RefreshImports()
 {
 	ClearImportLandscapeData();
 	HeightmapImportDescriptorIndex = 0;
@@ -562,14 +562,72 @@ void ULandscapeEditorObject::OnImportHeightmapFilenameChanged()
 	ImportLandscape_HeightmapImportResult = ELandscapeImportResult::Success;
 	ImportLandscape_HeightmapErrorMessage = FText();
 
-	ImportLandscape_HeightmapImportResult =
-		FLandscapeImportHelper::GetHeightmapImportDescriptor(ImportLandscape_HeightmapFilename, UseSingleFileImport(), HeightmapImportDescriptor, ImportLandscape_HeightmapErrorMessage);
-	if (ImportLandscape_HeightmapImportResult != ELandscapeImportResult::Error)
+	if (!ImportLandscape_HeightmapFilename.IsEmpty())
 	{
-		ImportLandscape_Width = HeightmapImportDescriptor.ImportResolutions[HeightmapImportDescriptorIndex].Width;
-		ImportLandscape_Height = HeightmapImportDescriptor.ImportResolutions[HeightmapImportDescriptorIndex].Height;
-		ChooseBestComponentSizeForImport();
-		ImportLandscapeData();
+		ImportLandscape_HeightmapImportResult =
+			FLandscapeImportHelper::GetHeightmapImportDescriptor(ImportLandscape_HeightmapFilename, UseSingleFileImport(), HeightmapImportDescriptor, ImportLandscape_HeightmapErrorMessage);
+		if (ImportLandscape_HeightmapImportResult != ELandscapeImportResult::Error)
+		{
+			ImportLandscape_Width = HeightmapImportDescriptor.ImportResolutions[HeightmapImportDescriptorIndex].Width;
+			ImportLandscape_Height = HeightmapImportDescriptor.ImportResolutions[HeightmapImportDescriptorIndex].Height;
+			ChooseBestComponentSizeForImport();
+			ImportLandscapeData();
+		}
+	}
+
+	RefreshLayerImports();
+}
+
+void ULandscapeEditorObject::RefreshLayerImports()
+{
+	// Make sure to reset import width and height if we don't have a Heightmap to import
+	if (ImportLandscape_HeightmapFilename.IsEmpty())
+	{
+		HeightmapImportDescriptorIndex = 0;
+		ImportLandscape_Width = 0;
+		ImportLandscape_Height = 0;
+	}
+
+	for (FLandscapeImportLayer& UIImportLayer : ImportLandscape_Layers)
+	{
+		UIImportLayer.ErrorMessage = FText();
+		UIImportLayer.ImportResult = ELandscapeImportResult::Success;
+
+		if (UIImportLayer.LayerInfo != nullptr && !UIImportLayer.SourceFilePath.IsEmpty())
+		{
+			UIImportLayer.ImportResult = FLandscapeImportHelper::GetWeightmapImportDescriptor(UIImportLayer.SourceFilePath, UseSingleFileImport(), UIImportLayer.LayerName, UIImportLayer.ImportDescriptor, UIImportLayer.ErrorMessage);
+			if (UIImportLayer.ImportResult == ELandscapeImportResult::Error)
+			{
+				continue;
+			}
+
+			if (ImportLandscape_Height != 0 || ImportLandscape_Width != 0)
+			{
+				// Use same import index as Heightmap
+				int32 FoundIndex = INDEX_NONE;
+				for (int32 Index = 0; Index < UIImportLayer.ImportDescriptor.FileResolutions.Num(); ++Index)
+				{
+					if (UIImportLayer.ImportDescriptor.ImportResolutions[Index] == HeightmapImportDescriptor.ImportResolutions[HeightmapImportDescriptorIndex])
+					{
+						FoundIndex = Index;
+						break;
+					}
+				}
+
+				if (FoundIndex == INDEX_NONE)
+				{
+					UIImportLayer.ImportResult = ELandscapeImportResult::Error;
+					UIImportLayer.ErrorMessage = NSLOCTEXT("LandscapeEditor.ImportLandscape", "Import_WeightHeightResolutionMismatch", "Weightmap import resolution isn't same as Heightmap resolution.");
+				}
+			}
+			else // We don't have a valid Heightmap import so use the Layer Import resolution
+			{
+				HeightmapImportDescriptorIndex = 0;
+				ImportLandscape_Width = UIImportLayer.ImportDescriptor.ImportResolutions[HeightmapImportDescriptorIndex].Width;
+				ImportLandscape_Height = UIImportLayer.ImportDescriptor.ImportResolutions[HeightmapImportDescriptorIndex].Height;
+				ChooseBestComponentSizeForImport();
+			}			
+		}
 	}
 }
 
@@ -615,10 +673,9 @@ ELandscapeImportResult ULandscapeEditorObject::CreateImportLayersInfo(TArray<FLa
 		OutImportLayerInfos.Add((const FLandscapeImportLayer&)UIImportLayer); //slicing is fine here
 		FLandscapeImportLayerInfo& ImportLayer = OutImportLayerInfos.Last();
 
-		if (ImportLayer.LayerInfo != nullptr && ImportLayer.SourceFilePath != "")
+		if (ImportLayer.LayerInfo != nullptr && !ImportLayer.SourceFilePath.IsEmpty())
 		{
-			FLandscapeImportDescriptor OutImportDescriptor;
-			UIImportLayer.ImportResult = FLandscapeImportHelper::GetWeightmapImportDescriptor(ImportLayer.SourceFilePath, UseSingleFileImport(), ImportLayer.LayerName, OutImportDescriptor, UIImportLayer.ErrorMessage);
+			UIImportLayer.ImportResult = FLandscapeImportHelper::GetWeightmapImportDescriptor(ImportLayer.SourceFilePath, UseSingleFileImport(), ImportLayer.LayerName, UIImportLayer.ImportDescriptor, UIImportLayer.ErrorMessage);
 			if (UIImportLayer.ImportResult == ELandscapeImportResult::Error)
 			{
 				FMessageDialog::Open(EAppMsgType::Ok, UIImportLayer.ErrorMessage);
@@ -627,9 +684,9 @@ ELandscapeImportResult ULandscapeEditorObject::CreateImportLayersInfo(TArray<FLa
 
 			// Use same import index as Heightmap
 			int32 FoundIndex = INDEX_NONE;
-			for (int32 Index = 0; Index < OutImportDescriptor.FileResolutions.Num(); ++Index)
+			for (int32 Index = 0; Index < UIImportLayer.ImportDescriptor.FileResolutions.Num(); ++Index)
 			{
-				if (OutImportDescriptor.ImportResolutions[Index] == HeightmapImportDescriptor.ImportResolutions[HeightmapImportDescriptorIndex])
+				if (UIImportLayer.ImportDescriptor.ImportResolutions[Index] == HeightmapImportDescriptor.ImportResolutions[HeightmapImportDescriptorIndex])
 				{
 					FoundIndex = Index;
 					break;
@@ -644,7 +701,7 @@ ELandscapeImportResult ULandscapeEditorObject::CreateImportLayersInfo(TArray<FLa
 				return ELandscapeImportResult::Error;
 			}
 
-			UIImportLayer.ImportResult = FLandscapeImportHelper::GetWeightmapImportData(OutImportDescriptor, FoundIndex, ImportLayer.LayerName, ImportLayer.LayerData, UIImportLayer.ErrorMessage);
+			UIImportLayer.ImportResult = FLandscapeImportHelper::GetWeightmapImportData(UIImportLayer.ImportDescriptor, FoundIndex, ImportLayer.LayerName, ImportLayer.LayerData, UIImportLayer.ErrorMessage);
 			if (UIImportLayer.ImportResult == ELandscapeImportResult::Error)
 			{
 				FMessageDialog::Open(EAppMsgType::Ok, UIImportLayer.ErrorMessage);
@@ -719,7 +776,7 @@ void ULandscapeEditorObject::ExpandImportData(TArray<uint16>& OutHeightData, TAr
 		FLandscapeImportResolution RequiredResolution(NewLandscape_ComponentCount.X * QuadsPerComponent + 1, NewLandscape_ComponentCount.Y * QuadsPerComponent + 1);
 		FLandscapeImportResolution ImportResolution(ImportLandscape_Width, ImportLandscape_Height);
 
-		FLandscapeImportHelper::ExpandHeightmapImportData(ImportData, OutHeightData, ImportResolution, RequiredResolution);
+		FLandscapeImportHelper::TransformHeightmapImportData(ImportData, OutHeightData, ImportResolution, RequiredResolution, ELandscapeImportTransformType::ExpandCentered);
 
 		for (int32 LayerIdx = 0; LayerIdx < OutImportLayerInfos.Num(); ++LayerIdx)
 		{
@@ -727,20 +784,42 @@ void ULandscapeEditorObject::ExpandImportData(TArray<uint16>& OutHeightData, TAr
 			TArray<uint8> OutLayerData;
 			if (OutImportLayerData.Num())
 			{
-				FLandscapeImportHelper::ExpandWeightmapImportData(OutImportLayerData, OutLayerData, ImportResolution, RequiredResolution);
+				FLandscapeImportHelper::TransformWeightmapImportData(OutImportLayerData, OutLayerData, ImportResolution, RequiredResolution, ELandscapeImportTransformType::ExpandCentered);
 				OutImportLayerData = MoveTemp(OutLayerData);
 			}
 		}
 	}
 }
 
-void ULandscapeEditorObject::RefreshImportLayersList()
+void ULandscapeEditorObject::RefreshImportLayersList(bool bRefreshFromTarget)
 {
-	UTexture2D* ThumbnailWeightmap = LoadObject<UTexture2D>(NULL, TEXT("/Engine/EditorLandscapeResources/LandscapeThumbnailWeightmap.LandscapeThumbnailWeightmap"), NULL, LOAD_None, NULL);
-	UTexture2D* ThumbnailHeightmap = LoadObject<UTexture2D>(NULL, TEXT("/Engine/EditorLandscapeResources/LandscapeThumbnailHeightmap.LandscapeThumbnailHeightmap"), NULL, LOAD_None, NULL);
+	UTexture2D* ThumbnailWeightmap = nullptr;
+	UTexture2D* ThumbnailHeightmap = nullptr;
+		
+	TArray<FName> LayerNames;
+	TArray<ULandscapeLayerInfoObject*> LayerInfoObjs;
+	UMaterialInterface* Material = nullptr;
+	if (bRefreshFromTarget)
+	{
+		LayerNames.Reset(ImportLandscape_Layers.Num());
+		LayerInfoObjs.Reset(ImportLandscape_Layers.Num());
+		Material = ParentMode->GetTargetLandscapeMaterial();
+		for (const TSharedRef<FLandscapeTargetListInfo>& TargetListInfo : ParentMode->GetTargetList())
+		{
+			if (TargetListInfo->TargetType != ELandscapeToolTargetType::Weightmap)
+			{
+				continue;
+			}
 
-	UMaterialInterface* Material = NewLandscape_Material.Get();
-	TArray<FName> LayerNames = ALandscapeProxy::GetLayersFromMaterial(Material);
+			LayerNames.Add(TargetListInfo->LayerName);
+			LayerInfoObjs.Add(TargetListInfo->LayerInfoObj.Get());
+		}
+	}
+	else
+	{
+		Material = NewLandscape_Material.Get();
+		LayerNames = ALandscapeProxy::GetLayersFromMaterial(Material);
+	}
 
 	const TArray<FLandscapeImportLayer> OldLayersList = MoveTemp(ImportLandscape_Layers);
 	ImportLandscape_Layers.Reset(LayerNames.Num());
@@ -751,6 +830,9 @@ void ULandscapeEditorObject::RefreshImportLayersList()
 
 		bool bFound = false;
 		FLandscapeImportLayer NewImportLayer;
+		NewImportLayer.ImportResult = ELandscapeImportResult::Success;
+		NewImportLayer.ErrorMessage = FText();
+
 		for (int32 j = 0; j < OldLayersList.Num(); j++)
 		{
 			if (OldLayersList[j].LayerName == LayerName)
@@ -769,38 +851,49 @@ void ULandscapeEditorObject::RefreshImportLayersList()
 				NewImportLayer.ThumbnailMIC->SetParentEditorOnly(Material);
 				Context.AddMaterialInterface(NewImportLayer.ThumbnailMIC);
 			}
+		}
+		else
+		{
+			if (!ThumbnailWeightmap)
+			{
+				ThumbnailWeightmap = LoadObject<UTexture2D>(NULL, TEXT("/Engine/EditorLandscapeResources/LandscapeThumbnailWeightmap.LandscapeThumbnailWeightmap"), NULL, LOAD_None, NULL);
+			}
 
-			NewImportLayer.ImportResult = ELandscapeImportResult::Success;
-			NewImportLayer.ErrorMessage = FText();
+			if (!ThumbnailHeightmap)
+			{
+				ThumbnailHeightmap = LoadObject<UTexture2D>(NULL, TEXT("/Engine/EditorLandscapeResources/LandscapeThumbnailHeightmap.LandscapeThumbnailHeightmap"), NULL, LOAD_None, NULL);
+			}
 
-			if (!NewImportLayer.SourceFilePath.IsEmpty())
+			NewImportLayer.LayerName = LayerName;
+			NewImportLayer.ThumbnailMIC = ALandscapeProxy::GetLayerThumbnailMIC(Material, LayerName, ThumbnailWeightmap, ThumbnailHeightmap, nullptr);
+		}
+
+		if (bRefreshFromTarget)
+		{
+			NewImportLayer.LayerInfo = LayerInfoObjs[i];
+		}
+
+		if (!NewImportLayer.SourceFilePath.IsEmpty())
+		{
+			if (!NewImportLayer.LayerInfo)
 			{
-				if (!NewImportLayer.LayerInfo)
-				{
-					NewImportLayer.ImportResult = ELandscapeImportResult::Error;
-					NewImportLayer.ErrorMessage = NSLOCTEXT("LandscapeEditor.NewLandscape", "Import_LayerInfoNotSet", "Can't import a layer file without a layer info");
-				}
-				else
+				NewImportLayer.ImportResult = ELandscapeImportResult::Error;
+				NewImportLayer.ErrorMessage = NSLOCTEXT("LandscapeEditor.NewLandscape", "Import_LayerInfoNotSet", "Can't import a layer file without a layer info");
+			}
+			else
+			{
+				NewImportLayer.ImportResult = FLandscapeImportHelper::GetWeightmapImportDescriptor(NewImportLayer.SourceFilePath, UseSingleFileImport(), NewImportLayer.LayerName, NewImportLayer.ImportDescriptor, NewImportLayer.ErrorMessage);
+				if (NewImportLayer.ImportResult != ELandscapeImportResult::Error)
 				{
-					FLandscapeImportDescriptor OutImportDescriptor;
-					NewImportLayer.ImportResult = FLandscapeImportHelper::GetWeightmapImportDescriptor(NewImportLayer.SourceFilePath, UseSingleFileImport(), NewImportLayer.LayerName, OutImportDescriptor, NewImportLayer.ErrorMessage);
-					if (NewImportLayer.ImportResult != ELandscapeImportResult::Error)
+					int32 DescriptorIndex = NewImportLayer.ImportDescriptor.FindDescriptorIndex(ImportLandscape_Width, ImportLandscape_Height);
+					if (DescriptorIndex == INDEX_NONE)
 					{
-						int32 DescriptorIndex = OutImportDescriptor.FindDescriptorIndex(ImportLandscape_Width, ImportLandscape_Height);
-						if (DescriptorIndex == INDEX_NONE)
-						{
-							NewImportLayer.ImportResult = ELandscapeImportResult::Error;
-							NewImportLayer.ErrorMessage = NSLOCTEXT("LandscapeEditor.NewLandscape", "Import_LayerSizeMismatch", "Size of the layer import size does not match size of heightmap import size");
-						}
+						NewImportLayer.ImportResult = ELandscapeImportResult::Error;
+						NewImportLayer.ErrorMessage = NSLOCTEXT("LandscapeEditor.NewLandscape", "Import_LayerSizeMismatch", "Size of the layer import size does not match size of heightmap import size");
 					}
 				}
 			}
 		}
-		else
-		{
-			NewImportLayer.LayerName = LayerName;
-			NewImportLayer.ThumbnailMIC = ALandscapeProxy::GetLayerThumbnailMIC(Material, LayerName, ThumbnailWeightmap, ThumbnailHeightmap, nullptr);
-		}
 
 		ImportLandscape_Layers.Add(MoveTemp(NewImportLayer));
 	}
diff --git a/Engine/Source/Editor/LandscapeEditor/Private/LandscapeImportHelper.cpp b/Engine/Source/Editor/LandscapeEditor/Private/LandscapeImportHelper.cpp
--- a/Engine/Source/Editor/LandscapeEditor/Private/LandscapeImportHelper.cpp
+++ b/Engine/Source/Editor/LandscapeEditor/Private/LandscapeImportHelper.cpp
@@ -128,14 +128,25 @@ ELandscapeImportResult GetImportDescriptorInternal(const FString& FilePath, bool
 	FString OutFileImportPattern;
 	FString FilePathPattern;
 	TArray<FString> OutFilesToImport;
-	if (!bSingleFile && FLandscapeImportHelper::ExtractCoordinates(FPaths::GetBaseFilename(FilePath), OutCoord, OutFileImportPattern))
+
+	// If we are handling multiple files handle the case where Filepath is in the _xN_yM.extenstion format or if its a pattern already
+	if (!bSingleFile)
 	{
-		FilePathPattern = FPaths::GetPath(FilePath) / OutFileImportPattern;
+		if (FLandscapeImportHelper::ExtractCoordinates(FPaths::GetBaseFilename(FilePath), OutCoord, OutFileImportPattern))
+		{
+			FilePathPattern = FPaths::GetPath(FilePath) / OutFileImportPattern;
+		}
+		else
+		{
+			FilePathPattern = FPaths::GetBaseFilename(FilePath, false);
+		}
+
 		FLandscapeImportHelper::GetMatchingFiles(FilePathPattern, OutFilesToImport);
+		// We were expecting multiple files but only got one.
+		bSingleFile = OutFilesToImport.Num() == 1 && FilePath == OutFilesToImport[0];
 	}
 	else
 	{
-		bSingleFile = true;
 		OutFilesToImport.Add(FilePath);
 	}
 
@@ -233,16 +244,38 @@ ELandscapeImportResult GetImportDescriptorInternal(const FString& FilePath, bool
 }
 
 template<class T>
-void ExpandImportDataInternal(const TArray<T>& InData, TArray<T>& OutData, const FLandscapeImportResolution& CurrentResolution, const FLandscapeImportResolution& RequiredResolution)
+void TransformImportDataInternal(const TArray<T>& InData, TArray<T>& OutData, const FLandscapeImportResolution& CurrentResolution, const FLandscapeImportResolution& RequiredResolution, ELandscapeImportTransformType TransformType, FIntPoint Offset)
 {
 	check(InData.Num() == CurrentResolution.Width * CurrentResolution.Height);
-	// Center Import data
-	const int32 OffsetX = (int32)(RequiredResolution.Width - CurrentResolution.Width) / 2;
-	const int32 OffsetY = (int32)(RequiredResolution.Height - CurrentResolution.Height) / 2;
-
-	FIntRect SrcRegion(0, 0, CurrentResolution.Width - 1, CurrentResolution.Height - 1);
-	FIntRect DestRegion(-OffsetX, -OffsetY, RequiredResolution.Width - OffsetX - 1, RequiredResolution.Height - OffsetY - 1);
-	FLandscapeConfigHelper::ExpandData<T>(InData, OutData, SrcRegion, DestRegion, true);
+	if (TransformType == ELandscapeImportTransformType::Resample)
+	{
+		const FIntRect SrcRegion(0, 0, CurrentResolution.Width - 1, CurrentResolution.Height - 1);
+		const FIntRect DestRegion(0, 0, RequiredResolution.Width - 1, RequiredResolution.Height - 1);
+		FLandscapeConfigHelper::ResampleData<T>(InData, OutData, SrcRegion, DestRegion);
+	}
+	else if(TransformType != ELandscapeImportTransformType::None)
+	{
+		int32 OffsetX = 0;
+		int32 OffsetY = 0;
+		if (TransformType == ELandscapeImportTransformType::ExpandCentered)
+		{
+			OffsetX = (int32)(RequiredResolution.Width - CurrentResolution.Width) / 2;
+			OffsetY = (int32)(RequiredResolution.Height - CurrentResolution.Height) / 2;
+		}
+		else if (TransformType == ELandscapeImportTransformType::ExpandOffset)
+		{
+			OffsetX = Offset.X;
+			OffsetY = Offset.Y;
+		}
+				
+		const FIntRect SrcRegion(0, 0, CurrentResolution.Width - 1, CurrentResolution.Height - 1);
+		const FIntRect DestRegion(-OffsetX, -OffsetY, RequiredResolution.Width - OffsetX - 1, RequiredResolution.Height - OffsetY - 1);
+		FLandscapeConfigHelper::ExpandData<T>(InData, OutData, SrcRegion, DestRegion, OffsetX != 0 || OffsetY != 0);
+	}
+	else
+	{
+		OutData = InData;
+	}
 }
 
 ELandscapeImportResult FLandscapeImportHelper::GetHeightmapImportDescriptor(const FString& FilePath, bool bSingleFile, FLandscapeImportDescriptor& OutImportDescriptor, FText& OutMessage)
@@ -265,14 +298,14 @@ ELandscapeImportResult FLandscapeImportHelper::GetWeightmapImportData(const FLan
 	return GetImportDataInternal<uint8>(ImportDescriptor, DescriptorIndex, LayerName, 0, OutData, OutMessage);
 }
 
-void FLandscapeImportHelper::ExpandWeightmapImportData(const TArray<uint8>& InData, TArray<uint8>& OutData, const FLandscapeImportResolution& CurrentResolution, const FLandscapeImportResolution& RequiredResolution)
+void FLandscapeImportHelper::TransformWeightmapImportData(const TArray<uint8>& InData, TArray<uint8>& OutData, const FLandscapeImportResolution& CurrentResolution, const FLandscapeImportResolution& RequiredResolution, ELandscapeImportTransformType TransformType, FIntPoint Offset)
 {
-	ExpandImportDataInternal<uint8>(InData, OutData, CurrentResolution, RequiredResolution);
+	TransformImportDataInternal<uint8>(InData, OutData, CurrentResolution, RequiredResolution, TransformType, Offset);
 }
 
-void FLandscapeImportHelper::ExpandHeightmapImportData(const TArray<uint16>& InData, TArray<uint16>& OutData, const FLandscapeImportResolution& CurrentResolution, const FLandscapeImportResolution& RequiredResolution)
+void FLandscapeImportHelper::TransformHeightmapImportData(const TArray<uint16>& InData, TArray<uint16>& OutData, const FLandscapeImportResolution& CurrentResolution, const FLandscapeImportResolution& RequiredResolution, ELandscapeImportTransformType TransformType, FIntPoint Offset)
 {
-	ExpandImportDataInternal<uint16>(InData, OutData, CurrentResolution, RequiredResolution);
+	TransformImportDataInternal<uint16>(InData, OutData, CurrentResolution, RequiredResolution, TransformType, Offset);
 }
 
 void FLandscapeImportHelper::ChooseBestComponentSizeForImport(int32 Width, int32 Height, int32& InOutQuadsPerSection, int32& InOutSectionsPerComponent, FIntPoint& OutComponentCount)
diff --git a/Engine/Source/Editor/LandscapeEditor/Private/SLandscapeEditor.cpp b/Engine/Source/Editor/LandscapeEditor/Private/SLandscapeEditor.cpp
--- a/Engine/Source/Editor/LandscapeEditor/Private/SLandscapeEditor.cpp
+++ b/Engine/Source/Editor/LandscapeEditor/Private/SLandscapeEditor.cpp
@@ -92,6 +92,7 @@ void FLandscapeToolKit::Init(const TSharedPtr<IToolkitHost>& InitToolkitHost)
 #define MAP_TOOL(ToolName) CommandList->MapAction(NameToCommandMap.FindChecked("Tool_" ToolName), FUIAction(FExecuteAction::CreateSP(this, &FLandscapeToolKit::OnChangeTool, FName(ToolName)), FCanExecuteAction::CreateSP(this, &FLandscapeToolKit::IsToolEnabled, FName(ToolName)), FIsActionChecked::CreateSP(this, &FLandscapeToolKit::IsToolActive, FName(ToolName)), FIsActionButtonVisible::CreateSP(this, &FLandscapeToolKit::IsToolAvailable, FName(ToolName))));
 	MAP_TOOL("NewLandscape");
 	MAP_TOOL("ResizeLandscape");
+	MAP_TOOL("ImportExport");
 
 	MAP_TOOL("Sculpt");
 	MAP_TOOL("Erase");
@@ -194,7 +195,7 @@ void FLandscapeToolKit::BuildToolPalette(FName PaletteName, class FToolBarBuilde
 		ToolBarBuilder.BeginSection("Manage");
 
 		ToolBarBuilder.AddToolBarButton(Commands.NewLandscape);
-
+		ToolBarBuilder.AddToolBarButton(Commands.ImportExportTool);
 		ToolBarBuilder.AddToolBarButton(Commands.SelectComponentTool);
 		ToolBarBuilder.AddToolBarButton(Commands.AddComponentTool);
 		ToolBarBuilder.AddToolBarButton(Commands.DeleteComponentTool);
diff --git a/Engine/Source/Editor/LandscapeEditor/Public/LandscapeEditorModule.h b/Engine/Source/Editor/LandscapeEditor/Public/LandscapeEditorModule.h
--- a/Engine/Source/Editor/LandscapeEditor/Public/LandscapeEditorModule.h
+++ b/Engine/Source/Editor/LandscapeEditor/Public/LandscapeEditorModule.h
@@ -5,10 +5,15 @@
 #include "CoreMinimal.h"
 #include "Modules/ModuleInterface.h"
 #include "Templates/SharedPointer.h"
-#include "LandscapeFileFormatInterface.h"
 #include "Templates/UnrealTypeTraits.h"
 #include "Templates/EnableIf.h"
 
+template<class T>
+class ILandscapeFileFormat;
+
+using ILandscapeHeightmapFileFormat = ILandscapeFileFormat<uint16>;
+using ILandscapeWeightmapFileFormat = ILandscapeFileFormat<uint8>;
+
 class FUICommandList;
 
 /**
diff --git a/Engine/Source/Editor/LandscapeEditor/Public/LandscapeEditorObject.h b/Engine/Source/Editor/LandscapeEditor/Public/LandscapeEditorObject.h
--- a/Engine/Source/Editor/LandscapeEditor/Public/LandscapeEditorObject.h
+++ b/Engine/Source/Editor/LandscapeEditor/Public/LandscapeEditorObject.h
@@ -204,17 +204,26 @@ struct FLandscapeImportLayer : public FLandscapeImportLayerInfo
 	UPROPERTY(Category="Import", VisibleAnywhere)
 	ULandscapeMaterialInstanceConstant* ThumbnailMIC;
 
-	UPROPERTY(Category="Import", VisibleAnywhere)
+	UPROPERTY(Category = "Import", VisibleAnywhere)
 	ELandscapeImportResult ImportResult;
 
-	UPROPERTY(Category="Import", VisibleAnywhere)
+	UPROPERTY(Category = "Import", VisibleAnywhere)
 	FText ErrorMessage;
 
+	UPROPERTY(Category="Export", EditAnywhere, meta = (DisplayName = "Layer File"))
+	FString ExportFilePath;
+
+	UPROPERTY(Category="Import", EditAnywhere)
+	bool bSelected;
+
+	FLandscapeImportDescriptor ImportDescriptor;
+
 	FLandscapeImportLayer()
 		: FLandscapeImportLayerInfo()
 		, ThumbnailMIC(nullptr)
 		, ImportResult(ELandscapeImportResult::Success)
 		, ErrorMessage(FText())
+		, bSelected(false)
 	{
 	}
 };
@@ -523,7 +532,25 @@ class ULandscapeEditorObject : public UObject
 	UPROPERTY(NonTransactional)
 	uint32 ImportLandscape_Height;
 
+	UPROPERTY(Category="Import / Export", EditAnywhere, NonTransactional, meta=(DisplayName="Heightmap File", ShowForTools="ImportExport"))
+	FString HeightmapExportFilename;
+		
+	UPROPERTY(NonTransactional)
+	FIntPoint ImportLandscape_GizmoLocalPosition;
+
+	UPROPERTY(Category = "Import / Export", EditAnywhere, NonTransactional, meta =(ShowForTools = "ImportExport"))
+	ELandscapeImportTransformType ImportType;
+		
+	UPROPERTY(NonTransactional)
+	bool bHeightmapSelected = false;
+	
+	UPROPERTY(Category = "Import / Export", EditAnywhere, NonTransactional, meta = (DisplayName="Export Edit Layer", ShowForTools = "ImportExport", ToolTip="When true exports the selected edit layer, if false exports the blend result"))
+	bool bExportEditLayer = true;
+
+	UPROPERTY(NonTransactional)
 	FLandscapeImportDescriptor HeightmapImportDescriptor;
+	
+	UPROPERTY(NonTransactional)
 	int32 HeightmapImportDescriptorIndex;
 private:
 	UPROPERTY(NonTransactional)
@@ -649,7 +676,7 @@ public:
 	void ClearImportLandscapeData() { ImportLandscape_Data.Empty(); }
 	void ChooseBestComponentSizeForImport();
 	void ImportLandscapeData();
-	void RefreshImportLayersList();
+	void RefreshImportLayersList(bool bRefreshFromTarget = false);
 	ELandscapeImportResult CreateImportLayersInfo(TArray<FLandscapeImportLayerInfo>& OutImportLayerInfos);
 	ELandscapeImportResult CreateNewLayersInfo(TArray<FLandscapeImportLayerInfo>& OutNewLayerInfos);
 	void InitializeDefaultHeightData(TArray<uint16>& OutData);
@@ -657,7 +684,10 @@ public:
 	void UpdateComponentLayerWhitelist();
 	bool UseSingleFileImport() const;
 	void OnChangeImportLandscapeResolution(int32 DescriptorIndex);
-	void OnImportHeightmapFilenameChanged();
+	void OnImportHeightmapFilenameChanged() { RefreshImports(); }
+	void RefreshImports();
+	void OnImportWeightmapFilenameChanged() { RefreshLayerImports(); }
+	void RefreshLayerImports();
 
 	int32 ClampLandscapeSize(int32 InComponentsCount) const
 	{
diff --git a/Engine/Source/Editor/LandscapeEditor/Public/LandscapeEditorUtils.h b/Engine/Source/Editor/LandscapeEditor/Public/LandscapeEditorUtils.h
--- a/Engine/Source/Editor/LandscapeEditor/Public/LandscapeEditorUtils.h
+++ b/Engine/Source/Editor/LandscapeEditor/Public/LandscapeEditorUtils.h
@@ -9,94 +9,6 @@ class ULandscapeLayerInfoObject;
 
 namespace LandscapeEditorUtils
 {
-	template<typename T>
-	void ExpandData(T* OutData, const T* InData,
-		int32 OldMinX, int32 OldMinY, int32 OldMaxX, int32 OldMaxY,
-		int32 NewMinX, int32 NewMinY, int32 NewMaxX, int32 NewMaxY)
-	{
-		const int32 OldWidth = OldMaxX - OldMinX + 1;
-		const int32 OldHeight = OldMaxY - OldMinY + 1;
-		const int32 NewWidth = NewMaxX - NewMinX + 1;
-		const int32 NewHeight = NewMaxY - NewMinY + 1;
-		const int32 OffsetX = NewMinX - OldMinX;
-		const int32 OffsetY = NewMinY - OldMinY;
-
-		for (int32 Y = 0; Y < NewHeight; ++Y)
-		{
-			const int32 OldY = FMath::Clamp<int32>(Y + OffsetY, 0, OldHeight - 1);
-
-			// Pad anything to the left
-			const T PadLeft = InData[OldY * OldWidth + 0];
-			for (int32 X = 0; X < -OffsetX; ++X)
-			{
-				OutData[Y * NewWidth + X] = PadLeft;
-			}
-
-			// Copy one row of the old data
-			{
-				const int32 X = FMath::Max(0, -OffsetX);
-				const int32 OldX = FMath::Clamp<int32>(X + OffsetX, 0, OldWidth - 1);
-				FMemory::Memcpy(&OutData[Y * NewWidth + X], &InData[OldY * OldWidth + OldX], FMath::Min<int32>(OldWidth, NewWidth) * sizeof(T));
-			}
-
-			const T PadRight = InData[OldY * OldWidth + OldWidth - 1];
-			for (int32 X = -OffsetX + OldWidth; X < NewWidth; ++X)
-			{
-				OutData[Y * NewWidth + X] = PadRight;
-			}
-		}
-	}
-
-	template<typename T>
-	TArray<T> ExpandData(const TArray<T>& Data,
-		int32 OldMinX, int32 OldMinY, int32 OldMaxX, int32 OldMaxY,
-		int32 NewMinX, int32 NewMinY, int32 NewMaxX, int32 NewMaxY)
-	{
-		const int32 NewWidth = NewMaxX - NewMinX + 1;
-		const int32 NewHeight = NewMaxY - NewMinY + 1;
-
-		TArray<T> Result;
-		Result.Empty(NewWidth * NewHeight);
-		Result.AddUninitialized(NewWidth * NewHeight);
-
-		ExpandData(Result.GetData(), Data.GetData(),
-			OldMinX, OldMinY, OldMaxX, OldMaxY,
-			NewMinX, NewMinY, NewMaxX, NewMaxY);
-
-		return Result;
-	}
-
-	template<typename T>
-	TArray<T> ResampleData(const TArray<T>& Data, int32 OldWidth, int32 OldHeight, int32 NewWidth, int32 NewHeight)
-	{
-		TArray<T> Result;
-		Result.Empty(NewWidth * NewHeight);
-		Result.AddUninitialized(NewWidth * NewHeight);
-
-		const float XScale = (float)(OldWidth - 1) / (NewWidth - 1);
-		const float YScale = (float)(OldHeight - 1) / (NewHeight - 1);
-		for (int32 Y = 0; Y < NewHeight; ++Y)
-		{
-			for (int32 X = 0; X < NewWidth; ++X)
-			{
-				const float OldY = Y * YScale;
-				const float OldX = X * XScale;
-				const int32 X0 = FMath::FloorToInt(OldX);
-				const int32 X1 = FMath::Min(FMath::FloorToInt(OldX) + 1, OldWidth - 1);
-				const int32 Y0 = FMath::FloorToInt(OldY);
-				const int32 Y1 = FMath::Min(FMath::FloorToInt(OldY) + 1, OldHeight - 1);
-				const T& Original00 = Data[Y0 * OldWidth + X0];
-				const T& Original10 = Data[Y0 * OldWidth + X1];
-				const T& Original01 = Data[Y1 * OldWidth + X0];
-				const T& Original11 = Data[Y1 * OldWidth + X1];
-				Result[Y * NewWidth + X] = FMath::BiLerp(Original00, Original10, Original01, Original11, FMath::Fractional(OldX), FMath::Fractional(OldY));
-			}
-		}
-
-		return Result;
-	}
-
 	bool LANDSCAPEEDITOR_API SetHeightmapData(ALandscapeProxy* Landscape, const TArray<uint16>& Data);
-
 	bool LANDSCAPEEDITOR_API SetWeightmapData(ALandscapeProxy* Landscape, ULandscapeLayerInfoObject* LayerObject, const TArray<uint8>& Data);
 }
diff --git a/Engine/Source/Editor/LandscapeEditor/Public/LandscapeFileFormatInterface.h b/Engine/Source/Editor/LandscapeEditor/Public/LandscapeFileFormatInterface.h
--- a/Engine/Source/Editor/LandscapeEditor/Public/LandscapeFileFormatInterface.h
+++ b/Engine/Source/Editor/LandscapeEditor/Public/LandscapeFileFormatInterface.h
@@ -6,6 +6,7 @@
 #include "UObject/ObjectMacros.h"
 #include "Containers/ArrayView.h"
 #include "Misc/Paths.h"
+#include "LandscapeFileFormatInterface.generated.h"
 
 class Error;
 
@@ -30,10 +31,19 @@ enum class ELandscapeImportResult : uint8
 	Error,
 };
 
+USTRUCT()
 struct FLandscapeFileResolution
 {
-	uint32 Width;
-	uint32 Height;
+	GENERATED_USTRUCT_BODY()
+
+	FLandscapeFileResolution() {}
+	FLandscapeFileResolution(uint32 InWidth, uint32 InHeight) : Width(InWidth), Height(InHeight) {}
+
+	UPROPERTY()
+	uint32 Width = 0;
+
+	UPROPERTY()
+	uint32 Height = 0;
 };
 
 FORCEINLINE bool operator==(const FLandscapeFileResolution& Lhs, const FLandscapeFileResolution& Rhs)
diff --git a/Engine/Source/Editor/LandscapeEditor/Public/LandscapeImportHelper.h b/Engine/Source/Editor/LandscapeEditor/Public/LandscapeImportHelper.h
--- a/Engine/Source/Editor/LandscapeEditor/Public/LandscapeImportHelper.h
+++ b/Engine/Source/Editor/LandscapeEditor/Public/LandscapeImportHelper.h
@@ -7,24 +7,49 @@
 #include "LandscapeFileFormatInterface.h"
 #include "Templates/UnrealTypeTraits.h"
 #include "Templates/EnableIf.h"
+#include "LandscapeImportHelper.generated.h"
 
+USTRUCT()
 struct FLandscapeImportFileDescriptor
 {
+	GENERATED_USTRUCT_BODY()
+
 	FLandscapeImportFileDescriptor(const FString& InFilePath, const FIntPoint& InCoord)
 		: Coord(InCoord), FilePath(InFilePath){}
 
+	FLandscapeImportFileDescriptor() {}
+
 	/* Which tile does that descriptor represent */
+	UPROPERTY()
 	FIntPoint Coord;
 	/* File path */
+	UPROPERTY()
 	FString FilePath;
 };
 
+USTRUCT()
 struct FLandscapeImportResolution
 {
+	GENERATED_USTRUCT_BODY()
+
 	FLandscapeImportResolution(uint32 InWidth, uint32 InHeight) : Width(InWidth), Height(InHeight) {}
 
-	uint32 Width;
-	uint32 Height;
+	FLandscapeImportResolution() {}
+
+	UPROPERTY()
+	uint32 Width = 0;
+	UPROPERTY()
+	uint32 Height = 0;
+};
+
+UENUM()
+enum class ELandscapeImportTransformType : int8
+{
+
+	None UMETA(DisplayName="Original", ToolTip="Will Import the data at the gizmo location in the original size"),
+	ExpandOffset UMETA(DisplayName="Expand", ToolTip="Will Import the data at the gizmo location and expand the data to fill the landscape") ,
+	ExpandCentered UMETA(Hidden), 
+	Resample UMETA(ToolTip="Will resample Import data to fit landscape")
 };
 
 FORCEINLINE bool operator==(const FLandscapeImportResolution& Lhs, const FLandscapeImportResolution& Rhs)
@@ -37,8 +62,11 @@ FORCEINLINE bool operator!=(const FLandscapeImportResolution& Lhs, const FLandsc
 	return !(Lhs == Rhs);
 }
 
+USTRUCT()
 struct FLandscapeImportDescriptor
 {
+	GENERATED_USTRUCT_BODY()
+
 	FLandscapeImportDescriptor()
 		: Scale(100, 100, 100) {}
 
@@ -64,12 +92,16 @@ struct FLandscapeImportDescriptor
 	}
 
 	/* Landscape Import Resolution based on File Coords + Resolutions */
+	UPROPERTY()
 	TArray<FLandscapeImportResolution> ImportResolutions;
 	/* Single File Resolutions */
+	UPROPERTY()
 	TArray<FLandscapeFileResolution> FileResolutions;
 	/* Files contributing to this descriptor */
+	UPROPERTY()
 	TArray<FLandscapeImportFileDescriptor> FileDescriptors;
 	/* Scale of Import data */
+	UPROPERTY()
 	FVector Scale;
 };
 
@@ -98,8 +130,8 @@ public:
 	static ELandscapeImportResult GetHeightmapImportData(const FLandscapeImportDescriptor& ImportDescriptor, int32 DescriptorIndex, TArray<uint16>& OutData, FText& OutMessage);
 	static ELandscapeImportResult GetWeightmapImportData(const FLandscapeImportDescriptor& ImportDescriptor, int32 DescriptorIndex, FName LayerName, TArray<uint8>& OutData, FText& OutMessage);
 
-	static void ExpandWeightmapImportData(const TArray<uint8>& InData, TArray<uint8>& OutData, const FLandscapeImportResolution& CurrentResolution, const FLandscapeImportResolution& RequiredResolution);
-	static void ExpandHeightmapImportData(const TArray<uint16>& InData, TArray<uint16>& OutData, const FLandscapeImportResolution& CurrentResolution, const FLandscapeImportResolution& RequiredResolution);
+	static void TransformWeightmapImportData(const TArray<uint8>& InData, TArray<uint8>& OutData, const FLandscapeImportResolution& CurrentResolution, const FLandscapeImportResolution& RequiredResolution, ELandscapeImportTransformType TransformType = ELandscapeImportTransformType::ExpandCentered, FIntPoint Offset = FIntPoint(0,0));
+	static void TransformHeightmapImportData(const TArray<uint16>& InData, TArray<uint16>& OutData, const FLandscapeImportResolution& CurrentResolution, const FLandscapeImportResolution& RequiredResolution, ELandscapeImportTransformType TransformType = ELandscapeImportTransformType::ExpandCentered, FIntPoint Offset = FIntPoint(0, 0));
 
 	static void ChooseBestComponentSizeForImport(int32 Width, int32 Height, int32& InOutQuadsPerSection, int32& InOutSectionsPerComponent, FIntPoint& OutComponentCount);
 
@@ -119,9 +151,9 @@ public:
 		return GetHeightmapImportData(ImportDescriptor, DescriptorIndex, OutData, OutMessage);
 	}
 	template<typename T>
-	static typename TEnableIf<TIsSame<T, uint16>::Value, void>::Type ExpandImportData(const TArray<T>& InData, TArray<T>& OutData, const FLandscapeImportResolution& CurrentResolution, const FLandscapeImportResolution& RequiredResolution)
+	static typename TEnableIf<TIsSame<T, uint16>::Value, void>::Type TransformImportData(const TArray<T>& InData, TArray<T>& OutData, const FLandscapeImportResolution& CurrentResolution, const FLandscapeImportResolution& RequiredResolution, ELandscapeImportTransformType TransformType = ELandscapeImportTransformType::ExpandCentered, FIntPoint Offset = FIntPoint(0, 0))
 	{
-		return ExpandHeightmapImportData(InData, OutData, CurrentResolution, RequiredResolution);
+		return TransformHeightmapImportData(InData, OutData, CurrentResolution, RequiredResolution, TransformType, Offset);
 	}
 
 	template<typename T>
@@ -135,8 +167,8 @@ public:
 		return GetWeightmapImportData(ImportDescriptor, DescriptorIndex, LayerName, OutData, OutMessage);
 	}
 	template<typename T>
-	static typename TEnableIf<TIsSame<T, uint8>::Value, void>::Type ExpandImportData(const TArray<T>& InData, TArray<T>& OutData, const FLandscapeImportResolution& CurrentResolution, const FLandscapeImportResolution& RequiredResolution)
+	static typename TEnableIf<TIsSame<T, uint8>::Value, void>::Type TransformImportData(const TArray<T>& InData, TArray<T>& OutData, const FLandscapeImportResolution& CurrentResolution, const FLandscapeImportResolution& RequiredResolution, ELandscapeImportTransformType TransformType = ELandscapeImportTransformType::ExpandCentered, FIntPoint Offset = FIntPoint(0, 0))
 	{
-		return ExpandWeightmapImportData(InData, OutData, CurrentResolution, RequiredResolution);
+		return TransformWeightmapImportData(InData, OutData, CurrentResolution, RequiredResolution, TransformType, Offset);
 	}
 };
\ No newline at end of file
diff --git a/Engine/Source/Runtime/Landscape/Classes/LandscapeInfo.h b/Engine/Source/Runtime/Landscape/Classes/LandscapeInfo.h
--- a/Engine/Source/Runtime/Landscape/Classes/LandscapeInfo.h
+++ b/Engine/Source/Runtime/Landscape/Classes/LandscapeInfo.h
@@ -154,10 +154,14 @@ public:
 	// all below.
 	LANDSCAPE_API bool AreAllComponentsRegistered() const;
 	LANDSCAPE_API void GetComponentsInRegion(int32 X1, int32 Y1, int32 X2, int32 Y2, TSet<ULandscapeComponent*>& OutComponents, bool bOverlap = true) const;
+	LANDSCAPE_API bool GetLandscapeExtent(ALandscapeProxy* Proxy, FIntRect& ProxyExtent) const;
+	LANDSCAPE_API bool GetLandscapeExtent(FIntRect& LandscapeExtent) const;
 	LANDSCAPE_API bool GetLandscapeExtent(int32& MinX, int32& MinY, int32& MaxX, int32& MaxY) const;
 	LANDSCAPE_API void ForAllLandscapeComponents(TFunctionRef<void(ULandscapeComponent*)> Fn) const;
 	LANDSCAPE_API void ExportHeightmap(const FString& Filename);
+	LANDSCAPE_API void ExportHeightmap(const FString& Filename, const FIntRect& ExportRegion);
 	LANDSCAPE_API void ExportLayer(ULandscapeLayerInfoObject* LayerInfo, const FString& Filename);
+	LANDSCAPE_API void ExportLayer(ULandscapeLayerInfoObject* LayerInfo, const FString& Filename, const FIntRect& ExportRegion);
 	LANDSCAPE_API bool ApplySplines(bool bOnlySelected, TSet<ULandscapeComponent*>* OutModifiedComponents = nullptr, bool bMarkPackageDirty = true);
 
 	LANDSCAPE_API bool GetSelectedExtent(int32& MinX, int32& MinY, int32& MaxX, int32& MaxY) const;
diff --git a/Engine/Source/Runtime/Landscape/Classes/LandscapeProxy.h b/Engine/Source/Runtime/Landscape/Classes/LandscapeProxy.h
--- a/Engine/Source/Runtime/Landscape/Classes/LandscapeProxy.h
+++ b/Engine/Source/Runtime/Landscape/Classes/LandscapeProxy.h
@@ -140,7 +140,7 @@ struct FLandscapeImportLayerInfo
 	UPROPERTY(Category="Import", EditAnywhere)
 	ULandscapeLayerInfoObject* LayerInfo;
 
-	UPROPERTY(Category="Import", EditAnywhere)
+	UPROPERTY(Category="Import", EditAnywhere, meta=(DisplayName="Layer File"))
 	FString SourceFilePath; // Optional
 	
 	// Raw weightmap data
diff --git a/Engine/Source/Runtime/Landscape/Private/LandscapeEdit.cpp b/Engine/Source/Runtime/Landscape/Private/LandscapeEdit.cpp
--- a/Engine/Source/Runtime/Landscape/Private/LandscapeEdit.cpp
+++ b/Engine/Source/Runtime/Landscape/Private/LandscapeEdit.cpp
@@ -3419,6 +3419,26 @@ bool ALandscape::HasAllComponent()
 	return false;
 }
 
+bool ULandscapeInfo::GetLandscapeExtent(ALandscapeProxy* LandscapeProxy, FIntRect& ProxyExtent) const
+{
+	ProxyExtent.Min.X = INT32_MAX;
+	ProxyExtent.Min.Y = INT32_MAX;
+	ProxyExtent.Max.X = INT32_MIN;
+	ProxyExtent.Max.Y = INT32_MIN;
+
+	for (ULandscapeComponent* LandscapeComponent : LandscapeProxy->LandscapeComponents)
+	{
+		LandscapeComponent->GetComponentExtent(ProxyExtent.Min.X, ProxyExtent.Min.Y, ProxyExtent.Max.X, ProxyExtent.Max.Y);
+	}
+	
+	return ProxyExtent.Min.X != INT32_MAX;
+}
+
+bool ULandscapeInfo::GetLandscapeExtent(FIntRect& LandscapeExtent) const
+{
+	return GetLandscapeExtent(LandscapeExtent.Min.X, LandscapeExtent.Min.Y, LandscapeExtent.Max.X, LandscapeExtent.Max.Y);
+}
+
 bool ULandscapeInfo::GetLandscapeExtent(int32& MinX, int32& MinY, int32& MaxX, int32& MaxY) const
 {
 	MinX = MAX_int32;
@@ -3562,29 +3582,32 @@ bool ULandscapeInfo::IsValidPosition(int32 X, int32 Y)
 
 void ULandscapeInfo::ExportHeightmap(const FString& Filename)
 {
-	int32 MinX = MAX_int32;
-	int32 MinY = MAX_int32;
-	int32 MaxX = -MAX_int32;
-	int32 MaxY = -MAX_int32;
-
-	if (!GetLandscapeExtent(MinX, MinY, MaxX, MaxY))
+	FIntRect ExportRegion;
+	if (!GetLandscapeExtent(ExportRegion))
 	{
 		return;
 	}
 
+	ExportHeightmap(Filename, ExportRegion);
+}
+
+void ULandscapeInfo::ExportHeightmap(const FString& Filename, const FIntRect& ExportRegion)
+{
 	GWarn->BeginSlowTask(LOCTEXT("BeginExportingLandscapeHeightmapTask", "Exporting Landscape Heightmap"), true);
 
 	ILandscapeEditorModule& LandscapeEditorModule = FModuleManager::GetModuleChecked<ILandscapeEditorModule>("LandscapeEditor");
 	FLandscapeEditDataInterface LandscapeEdit(this);
 
 	TArray<uint16> HeightData;
-	HeightData.AddZeroed((MaxX - MinX + 1) * (MaxY - MinY + 1));
-	LandscapeEdit.GetHeightDataFast(MinX, MinY, MaxX, MaxY, HeightData.GetData(), 0);
+	int32 ExportWidth = ExportRegion.Width() + 1;
+	int32 ExportHeight = ExportRegion.Height() + 1;
+	HeightData.AddZeroed(ExportWidth * ExportHeight);
+	LandscapeEdit.GetHeightDataFast(ExportRegion.Min.X, ExportRegion.Min.Y, ExportRegion.Max.X, ExportRegion.Max.Y, HeightData.GetData(), 0);
 
 	const ILandscapeHeightmapFileFormat* HeightmapFormat = LandscapeEditorModule.GetHeightmapFormatByExtension(*FPaths::GetExtension(Filename, true));
 	if (HeightmapFormat)
 	{
-		HeightmapFormat->Export(*Filename, NAME_None, HeightData, {(uint32)(MaxX - MinX + 1), (uint32)(MaxY - MinY + 1)}, DrawScale * FVector(1, 1, LANDSCAPE_ZSCALE));
+		HeightmapFormat->Export(*Filename, NAME_None, HeightData, {(uint32)ExportWidth, (uint32)ExportHeight}, DrawScale * FVector(1, 1, LANDSCAPE_ZSCALE));
 	}
 
 	GWarn->EndSlowTask();
@@ -3592,32 +3615,35 @@ void ULandscapeInfo::ExportHeightmap(const FString& Filename)
 
 void ULandscapeInfo::ExportLayer(ULandscapeLayerInfoObject* LayerInfo, const FString& Filename)
 {
-	check(LayerInfo);
-
-	int32 MinX = MAX_int32;
-	int32 MinY = MAX_int32;
-	int32 MaxX = -MAX_int32;
-	int32 MaxY = -MAX_int32;
-
-	if (!GetLandscapeExtent(MinX, MinY, MaxX, MaxY))
+	FIntRect ExportRegion;
+	if (!GetLandscapeExtent(ExportRegion))
 	{
 		return;
 	}
 
+	ExportLayer(LayerInfo, Filename, ExportRegion);
+}
+
+void ULandscapeInfo::ExportLayer(ULandscapeLayerInfoObject* LayerInfo, const FString& Filename, const FIntRect& ExportRegion)
+{
+	check(LayerInfo);
+		
 	GWarn->BeginSlowTask(LOCTEXT("BeginExportingLandscapeWeightmapTask", "Exporting Landscape Layer Weightmap"), true);
 
 	ILandscapeEditorModule& LandscapeEditorModule = FModuleManager::GetModuleChecked<ILandscapeEditorModule>("LandscapeEditor");
 
 	TArray<uint8> WeightData;
-	WeightData.AddZeroed((MaxX - MinX + 1) * (MaxY - MinY + 1));
+	int32 ExportWidth = ExportRegion.Width() + 1;
+	int32 ExportHeight = ExportRegion.Height() + 1;
+	WeightData.AddZeroed(ExportWidth * ExportHeight);
 
 	FLandscapeEditDataInterface LandscapeEdit(this);
-	LandscapeEdit.GetWeightDataFast(LayerInfo, MinX, MinY, MaxX, MaxY, WeightData.GetData(), 0);
+	LandscapeEdit.GetWeightDataFast(LayerInfo, ExportRegion.Min.X, ExportRegion.Min.Y, ExportRegion.Max.X, ExportRegion.Max.Y, WeightData.GetData(), 0);
 
 	const ILandscapeWeightmapFileFormat* WeightmapFormat = LandscapeEditorModule.GetWeightmapFormatByExtension(*FPaths::GetExtension(Filename, true));
 	if (WeightmapFormat)
 	{
-		WeightmapFormat->Export(*Filename, LayerInfo->LayerName, WeightData, {(uint32)(MaxX - MinX + 1), (uint32)(MaxY - MinY + 1)}, DrawScale * FVector(1, 1, LANDSCAPE_ZSCALE));
+		WeightmapFormat->Export(*Filename, LayerInfo->LayerName, WeightData, { (uint32)ExportWidth, (uint32)ExportHeight }, DrawScale * FVector(1, 1, LANDSCAPE_ZSCALE));
 	}
 
 	GWarn->EndSlowTask();
diff --git a/Engine/Source/Runtime/Landscape/Public/LandscapeConfigHelper.h b/Engine/Source/Runtime/Landscape/Public/LandscapeConfigHelper.h
--- a/Engine/Source/Runtime/Landscape/Public/LandscapeConfigHelper.h
+++ b/Engine/Source/Runtime/Landscape/Public/LandscapeConfigHelper.h
@@ -88,12 +88,12 @@ public:
 		}
 		else
 		{
-			ExpandData(InData, OutData, InSrcRegion, InDestRegion);
+			ExpandData(InData, OutData, InSrcRegion, InDestRegion, true);
 		}
 	}
 
 	template<typename T>
-	static void ExpandData(const TArray<T>& InData, TArray<T>& OutData, const FIntRect& InSrcRegion, const FIntRect& InDestRegion, bool bOffset = false)
+	static void ExpandData(const TArray<T>& InData, TArray<T>& OutData, const FIntRect& InSrcRegion, const FIntRect& InDestRegion, bool bOffset)
 	{
 		// Regions are in ComponentQuads and we want Vertices (+1)
 		const int32 SrcWidth = InSrcRegion.Width() + 1;
@@ -120,11 +120,12 @@ public:
 			{
 				const int32 X = FMath::Max(0, -OffsetX);
 				const int32 OldX = FMath::Clamp<int32>(X + OffsetX, 0, SrcWidth - 1);
-				FMemory::Memcpy(&OutData[Y * DestWidth + X], &InData[OldY * SrcWidth + OldX], FMath::Min<int32>(SrcWidth, DestWidth) * sizeof(T));
+				const int32 CopySize = FMath::Min<int32>(SrcWidth - OldX, DestWidth) * sizeof(T);
+				FMemory::Memcpy(&OutData[Y * DestWidth + X], &InData[OldY * SrcWidth + OldX], CopySize);
 			}
 
 			const T PadRight = InData[OldY * SrcWidth + SrcWidth - 1];
-			for (int32 X = SrcWidth; X < DestWidth; ++X)
+			for (int32 X = -OffsetX + SrcWidth; X < DestWidth; ++X)
 			{
 				OutData[Y * DestWidth + X] = PadRight;
 			}
