Fix initialization warning

diff --git a/Engine/Source/Editor/LandscapeEditor/Public/LandscapeImportHelper.h b/Engine/Source/Editor/LandscapeEditor/Public/LandscapeImportHelper.h
--- a/Engine/Source/Editor/LandscapeEditor/Public/LandscapeImportHelper.h
+++ b/Engine/Source/Editor/LandscapeEditor/Public/LandscapeImportHelper.h
@@ -17,7 +17,8 @@ struct FLandscapeImportFileDescriptor
 	FLandscapeImportFileDescriptor(const FString& InFilePath, const FIntPoint& InCoord)
 		: Coord(InCoord), FilePath(InFilePath){}
 
-	FLandscapeImportFileDescriptor() {}
+	FLandscapeImportFileDescriptor() 
+		: Coord(EForceInit::ForceInit) {}
 
 	/* Which tile does that descriptor represent */
 	UPROPERTY()
