Scaling in Landscape mode makes actor move out
Notify selection change when switching to the landscape copy/paste tool
Doing so will properly update the GLevelEditorModeTools().PivotLocation

diff --git a/Engine/Source/Editor/LandscapeEditor/Private/LandscapeEdModeComponentTools.cpp b/Engine/Source/Editor/LandscapeEditor/Private/LandscapeEdModeComponentTools.cpp
--- a/Engine/Source/Editor/LandscapeEditor/Private/LandscapeEdModeComponentTools.cpp
+++ b/Engine/Source/Editor/LandscapeEditor/Private/LandscapeEdModeComponentTools.cpp
@@ -1599,7 +1599,7 @@ public:
 		if (Gizmo)
 		{
 			GEditor->SelectNone(false, true);
-			GEditor->SelectActor(Gizmo, true, false, true);
+			GEditor->SelectActor(Gizmo, true, true, true);
 		}
 	}
 
