Fixed Import Landscape not rendering preview of imported data

diff --git a/Engine/Source/Editor/LandscapeEditor/Private/LandscapeEdModeComponentTools.cpp b/Engine/Source/Editor/LandscapeEditor/Private/LandscapeEdModeComponentTools.cpp
--- a/Engine/Source/Editor/LandscapeEditor/Private/LandscapeEdModeComponentTools.cpp
+++ b/Engine/Source/Editor/LandscapeEditor/Private/LandscapeEdModeComponentTools.cpp
@@ -1866,7 +1866,7 @@ public:
 
 	virtual void Render(const FSceneView* View, FViewport* Viewport, FPrimitiveDrawInterface* PDI) override
 	{
-		if (NewLandscapePreviewMode != ENewLandscapePreviewMode::None)
+		if (EdMode->NewLandscapePreviewMode != ENewLandscapePreviewMode::None)
 		{
 			static const float        CornerSize = 0.33f;
 			static const FLinearColor CornerColour(1.0f, 1.0f, 0.5f);
@@ -1884,7 +1884,7 @@ public:
 			const FVector Offset = EdMode->UISettings->NewLandscape_Location + FTransform(EdMode->UISettings->NewLandscape_Rotation, FVector::ZeroVector, EdMode->UISettings->NewLandscape_Scale).TransformVector(FVector(-ComponentCountX * ComponentSize / 2, -ComponentCountY * ComponentSize / 2, 0));
 			const FTransform Transform = FTransform(EdMode->UISettings->NewLandscape_Rotation, Offset, EdMode->UISettings->NewLandscape_Scale);
 
-			if (NewLandscapePreviewMode == ENewLandscapePreviewMode::ImportLandscape)
+			if (EdMode->NewLandscapePreviewMode == ENewLandscapePreviewMode::ImportLandscape)
 			{
 				const TArray<uint16>& ImportHeights = EdMode->UISettings->GetImportLandscapeData();
 				if (ImportHeights.Num() != 0)
@@ -1967,7 +1967,7 @@ public:
 					}
 				}
 			}
-			else //if (NewLandscapePreviewMode == ENewLandscapePreviewMode::NewLandscape)
+			else //if (EdMode->NewLandscapePreviewMode == ENewLandscapePreviewMode::NewLandscape)
 			{
 				if (ViewportType == LVT_Perspective || ViewportType == LVT_OrthoXY || ViewportType == LVT_OrthoNegativeXY)
 				{
