Landscape: Fix cross-level Copy/Paste  (specify Export root)

diff --git a/Engine/Source/Editor/LandscapeEditor/Private/LandscapeEdModeSplineTools.cpp b/Engine/Source/Editor/LandscapeEditor/Private/LandscapeEdModeSplineTools.cpp
--- a/Engine/Source/Editor/LandscapeEditor/Private/LandscapeEdModeSplineTools.cpp
+++ b/Engine/Source/Editor/LandscapeEditor/Private/LandscapeEdModeSplineTools.cpp
@@ -2108,7 +2108,7 @@ public:
 			Ar.Logf(TEXT("Begin Splines\r\n"));
 			for (UObject* Object : Objects)
 			{
-				UExporter::ExportToOutputDevice(&Context, Object, NULL, Ar, TEXT("copy"), 3, PPF_ExportsNotFullyQualified | PPF_Copy | PPF_Delimited, false);
+				UExporter::ExportToOutputDevice(&Context, Object, NULL, Ar, TEXT("copy"), 3, PPF_ExportsNotFullyQualified | PPF_Copy | PPF_Delimited, false, Object->GetOuter());
 			}
 			Ar.Logf(TEXT("End Splines\r\n"));
 
