Move render/trace of AddComponent out of the LandscapeEdMode and in the tool.

diff --git a/Engine/Source/Editor/LandscapeEditor/Private/LandscapeEdMode.cpp b/Engine/Source/Editor/LandscapeEditor/Private/LandscapeEdMode.cpp
--- a/Engine/Source/Editor/LandscapeEditor/Private/LandscapeEdMode.cpp
+++ b/Engine/Source/Editor/LandscapeEditor/Private/LandscapeEdMode.cpp
@@ -94,7 +94,6 @@ FEdModeLandscape::FEdModeLandscape()
 	, CurrentGizmoActor(nullptr)
 	, CopyPasteTool(nullptr)
 	, SplinesTool(nullptr)
-	, LandscapeRenderAddCollision(nullptr)
 	, TargetLayerStartingIndex(0)
 	, CachedLandscapeMaterial(nullptr)
 	, ToolActiveViewport(nullptr)
@@ -963,29 +962,6 @@ bool FEdModeLandscape::EndTracking(FEditorViewportClient* InViewportClient, FVie
 	return false;
 }
 
-namespace
-{
-	bool RayIntersectTriangle(const FVector& Start, const FVector& End, const FVector& A, const FVector& B, const FVector& C, FVector& IntersectPoint)
-	{
-		const FVector BA = A - B;
-		const FVector CB = B - C;
-		const FVector TriNormal = BA ^ CB;
-
-		bool bCollide = FMath::SegmentPlaneIntersection(Start, End, FPlane(A, TriNormal), IntersectPoint);
-		if (!bCollide)
-		{
-			return false;
-		}
-
-		FVector BaryCentric = FMath::ComputeBaryCentric2D(IntersectPoint, A, B, C);
-		if (BaryCentric.X > 0.0f && BaryCentric.Y > 0.0f && BaryCentric.Z > 0.0f)
-		{
-			return true;
-		}
-		return false;
-	}
-};
-
 /** Trace under the mouse cursor and return the landscape hit and the hit location (in landscape quad space) */
 bool FEdModeLandscape::LandscapeMouseTrace(FEditorViewportClient* ViewportClient, float& OutHitX, float& OutHitY)
 {
@@ -1038,12 +1014,23 @@ bool FEdModeLandscape::LandscapeMouseTrace(FEditorViewportClient* ViewportClient
 
 bool FEdModeLandscape::LandscapeTrace(const FVector& InRayOrigin, const FVector& InRayEnd, FVector& OutHitLocation)
 {
+	if (!CurrentTool || !CurrentToolTarget.LandscapeInfo.IsValid())
+	{
+		return false;
+	}
+
 	FVector Start = InRayOrigin;
 	FVector End = InRayEnd;
 
 	// Cache a copy of the world pointer
 	UWorld* World = GetWorld();
 
+	// Check Tool Trace first
+	if (CurrentTool->HitTrace(Start, End, OutHitLocation))
+	{
+		return true;
+	}
+	
 	TArray<FHitResult> Results;
 	// Each landscape component has 2 collision shapes, 1 of them is specific to landscape editor
 	// Trace only ECC_Visibility channel, so we do hit only Editor specific shape
@@ -1065,46 +1052,7 @@ bool FEdModeLandscape::LandscapeTrace(const FVector& InRayOrigin, const FVector&
 			}
 		}
 	}
-
-	// For Add Landscape Component Mode
-	if (CurrentTool->GetToolName() == FName("AddComponent") &&
-		CurrentToolTarget.LandscapeInfo.IsValid())
-	{
-		bool bCollided = false;
-		FVector IntersectPoint;
-		LandscapeRenderAddCollision = NULL;
-		// Need to optimize collision for AddLandscapeComponent...?
-		for (auto& XYToAddCollisionPair : CurrentToolTarget.LandscapeInfo->XYtoAddCollisionMap)
-		{
-			FLandscapeAddCollision& AddCollision = XYToAddCollisionPair.Value;
-			// Triangle 1
-			bCollided = RayIntersectTriangle(Start, End, AddCollision.Corners[0], AddCollision.Corners[3], AddCollision.Corners[1], IntersectPoint);
-			if (bCollided)
-			{
-				LandscapeRenderAddCollision = &AddCollision;
-				break;
-			}
-			// Triangle 2
-			bCollided = RayIntersectTriangle(Start, End, AddCollision.Corners[0], AddCollision.Corners[2], AddCollision.Corners[3], IntersectPoint);
-			if (bCollided)
-			{
-				LandscapeRenderAddCollision = &AddCollision;
-				break;
-			}
-		}
-
-		if (bCollided &&
-			CurrentToolTarget.LandscapeInfo.IsValid())
-		{
-			ALandscapeProxy* Proxy = CurrentToolTarget.LandscapeInfo.Get()->GetCurrentLevelLandscapeProxy(true);
-			if (Proxy)
-			{
-				OutHitLocation = Proxy->LandscapeActorToWorld().InverseTransformPosition(IntersectPoint);
-				return true;
-			}
-		}
-	}
-
+	
 	return false;
 }
 
@@ -2679,17 +2627,6 @@ void FEdModeLandscape::Render(const FSceneView* View, FViewport* Viewport, FPrim
 		return;
 	}
 
-	if (LandscapeRenderAddCollision)
-	{
-		PDI->DrawLine(LandscapeRenderAddCollision->Corners[0], LandscapeRenderAddCollision->Corners[3], FColor(0, 255, 128), SDPG_Foreground);
-		PDI->DrawLine(LandscapeRenderAddCollision->Corners[3], LandscapeRenderAddCollision->Corners[1], FColor(0, 255, 128), SDPG_Foreground);
-		PDI->DrawLine(LandscapeRenderAddCollision->Corners[1], LandscapeRenderAddCollision->Corners[0], FColor(0, 255, 128), SDPG_Foreground);
-
-		PDI->DrawLine(LandscapeRenderAddCollision->Corners[0], LandscapeRenderAddCollision->Corners[2], FColor(0, 255, 128), SDPG_Foreground);
-		PDI->DrawLine(LandscapeRenderAddCollision->Corners[2], LandscapeRenderAddCollision->Corners[3], FColor(0, 255, 128), SDPG_Foreground);
-		PDI->DrawLine(LandscapeRenderAddCollision->Corners[3], LandscapeRenderAddCollision->Corners[0], FColor(0, 255, 128), SDPG_Foreground);
-	}
-
 	// Override Rendering for Splines Tool
 	if (CurrentTool)
 	{
diff --git a/Engine/Source/Editor/LandscapeEditor/Private/LandscapeEdMode.h b/Engine/Source/Editor/LandscapeEditor/Private/LandscapeEdMode.h
--- a/Engine/Source/Editor/LandscapeEditor/Private/LandscapeEdMode.h
+++ b/Engine/Source/Editor/LandscapeEditor/Private/LandscapeEdMode.h
@@ -608,9 +608,6 @@ public:
 	TArray<TUniquePtr<FLandscapeTool>> LandscapeTools;
 	TArray<FLandscapeBrushSet> LandscapeBrushSets;
 
-	// For collision add visualization
-	FLandscapeAddCollision* LandscapeRenderAddCollision;
-
 	ELandscapeEditingState GetEditingState() const;
 
 	bool IsEditingEnabled() const
diff --git a/Engine/Source/Editor/LandscapeEditor/Private/LandscapeEdModeComponentTools.cpp b/Engine/Source/Editor/LandscapeEditor/Private/LandscapeEdModeComponentTools.cpp
--- a/Engine/Source/Editor/LandscapeEditor/Private/LandscapeEdModeComponentTools.cpp
+++ b/Engine/Source/Editor/LandscapeEditor/Private/LandscapeEdModeComponentTools.cpp
@@ -527,7 +527,7 @@ public:
 
 	virtual void Apply(FEditorViewportClient* ViewportClient, FLandscapeBrush* Brush, const ULandscapeEditorObject* UISettings, const TArray<FLandscapeToolInteractorPosition>& InteractorPositions)
 	{
-		if (LandscapeInfo && EdMode->LandscapeRenderAddCollision)
+		if (LandscapeInfo)
 		{
 			check(Brush->GetBrushType() == ELandscapeBrushType::Component);
 
@@ -724,8 +724,6 @@ public:
 				}				
 			}
 
-			EdMode->LandscapeRenderAddCollision = nullptr;
-
 			// Add/update "add collision" around the newly added components
 			if (!bHasLandscapeLayersContent)
 			{
@@ -799,6 +797,7 @@ public:
 	virtual void EnterTool() override
 	{
 		FLandscapeToolBase<FLandscapeToolStrokeAddComponent>::EnterTool();
+		AddCollision.Reset();
 		if(ULandscapeInfo* LandscapeInfo = EdMode->CurrentToolTarget.LandscapeInfo.Get())
 		{
 			LandscapeInfo->UpdateAllAddCollisions(); // Todo - as this is only used by this tool, move it into this tool?
@@ -809,8 +808,85 @@ public:
 	{
 		FLandscapeToolBase<FLandscapeToolStrokeAddComponent>::ExitTool();
 
-		EdMode->LandscapeRenderAddCollision = nullptr;
+		AddCollision.Reset();
 	}
+
+	virtual void Render(const FSceneView* View, FViewport* Viewport, FPrimitiveDrawInterface* PDI) override
+	{
+		if (AddCollision.IsSet())
+		{
+			PDI->DrawLine(AddCollision->Corners[0], AddCollision->Corners[3], FColor(0, 255, 128), SDPG_Foreground);
+			PDI->DrawLine(AddCollision->Corners[3], AddCollision->Corners[1], FColor(0, 255, 128), SDPG_Foreground);
+			PDI->DrawLine(AddCollision->Corners[1], AddCollision->Corners[0], FColor(0, 255, 128), SDPG_Foreground);
+
+			PDI->DrawLine(AddCollision->Corners[0], AddCollision->Corners[2], FColor(0, 255, 128), SDPG_Foreground);
+			PDI->DrawLine(AddCollision->Corners[2], AddCollision->Corners[3], FColor(0, 255, 128), SDPG_Foreground);
+			PDI->DrawLine(AddCollision->Corners[3], AddCollision->Corners[0], FColor(0, 255, 128), SDPG_Foreground);
+		}
+	}
+
+	virtual bool HitTrace(const FVector& TraceStart, const FVector& TraceEnd, FVector& OutHitLocation) override
+	{
+		ULandscapeInfo* LandscapeInfo = EdMode->CurrentToolTarget.LandscapeInfo.Get();
+		if (!LandscapeInfo)
+		{
+			return false;
+		}
+
+		ALandscapeProxy* Proxy = LandscapeInfo->GetLandscapeProxy();
+		if (!Proxy)
+		{
+			return false;
+		}
+
+		AddCollision.Reset();
+
+		FVector IntersectPoint;
+		// Need to optimize collision for AddLandscapeComponent...?
+		for (auto& XYToAddCollisionPair : LandscapeInfo->XYtoAddCollisionMap)
+		{
+			FLandscapeAddCollision& CurrentAddCollision = XYToAddCollisionPair.Value;
+			// Triangle 1
+			if(RayIntersectTriangle(TraceStart, TraceEnd, CurrentAddCollision.Corners[0], CurrentAddCollision.Corners[3], CurrentAddCollision.Corners[1], IntersectPoint))
+			{
+				AddCollision = CurrentAddCollision;
+				OutHitLocation = Proxy->LandscapeActorToWorld().InverseTransformPosition(IntersectPoint);
+				return true;
+			}
+			// Triangle 2
+			if(RayIntersectTriangle(TraceStart, TraceEnd, CurrentAddCollision.Corners[0], CurrentAddCollision.Corners[2], CurrentAddCollision.Corners[3], IntersectPoint))
+			{
+				AddCollision = CurrentAddCollision;
+				OutHitLocation = Proxy->LandscapeActorToWorld().InverseTransformPosition(IntersectPoint);
+				return true;
+			}
+		}
+		
+		return false;
+	}
+
+private:
+	bool RayIntersectTriangle(const FVector& Start, const FVector& End, const FVector& A, const FVector& B, const FVector& C, FVector& IntersectPoint)
+	{
+		const FVector BA = A - B;
+		const FVector CB = B - C;
+		const FVector TriNormal = BA ^ CB;
+
+		bool bCollide = FMath::SegmentPlaneIntersection(Start, End, FPlane(A, TriNormal), IntersectPoint);
+		if (!bCollide)
+		{
+			return false;
+		}
+
+		FVector BaryCentric = FMath::ComputeBaryCentric2D(IntersectPoint, A, B, C);
+		if (BaryCentric.X > 0.0f && BaryCentric.Y > 0.0f && BaryCentric.Z > 0.0f)
+		{
+			return true;
+		}
+		return false;
+	}
+
+	TOptional<FLandscapeAddCollision> AddCollision;
 };
 
 //
diff --git a/Engine/Source/Editor/LandscapeEditor/Public/LandscapeToolInterface.h b/Engine/Source/Editor/LandscapeEditor/Public/LandscapeToolInterface.h
--- a/Engine/Source/Editor/LandscapeEditor/Public/LandscapeToolInterface.h
+++ b/Engine/Source/Editor/LandscapeEditor/Public/LandscapeToolInterface.h
@@ -206,6 +206,7 @@ public:
 	virtual FText GetDisplayMessage() = 0;
 	virtual void SetEditRenderType();
 	virtual void Render(const FSceneView* View, FViewport* Viewport, FPrimitiveDrawInterface* PDI) {}
+	virtual bool HitTrace(const FVector& TraceStart, const FVector& TraceEnd, FVector& OutHitLocation) { return false; }
 	virtual bool SupportsMask() { return true; }
 	virtual bool SupportsComponentSelection() { return false; }
 	virtual bool OverrideSelection() const { return false; }
