WorldPartition: Landscape improvements
- Block painting on unloaded landscape edges
- Load LandscapeSplineActors that have intersecting bounds with a LandscapeStreamingProxy
- Fix AlphaBrushMaterial which didn't account for bCanPaint param

diff --git a/Engine/Source/Editor/LandscapeEditor/Private/LandscapeEdModeBrushes.cpp b/Engine/Source/Editor/LandscapeEditor/Private/LandscapeEdModeBrushes.cpp
--- a/Engine/Source/Editor/LandscapeEditor/Private/LandscapeEdModeBrushes.cpp
+++ b/Engine/Source/Editor/LandscapeEditor/Private/LandscapeEdModeBrushes.cpp
@@ -53,6 +53,7 @@ protected:
 	FVector2D LastMousePosition;
 	UMaterialInterface* BrushMaterial;
 	TMap<ULandscapeComponent*, UMaterialInstanceDynamic*> BrushMaterialInstanceMap;
+	bool bCanPaint;
 
 	virtual float CalculateFalloff(float Distance, float Radius, float Falloff) = 0;
 
@@ -60,6 +61,7 @@ protected:
 	FLandscapeBrushCircle(FEdModeLandscape* InEdMode, UMaterialInterface* InBrushMaterial)
 		: LastMousePosition(0, 0)
 		, BrushMaterial(LandscapeTool::CreateMaterialInstance(InBrushMaterial))
+		, bCanPaint(false)
 		, EdMode(InEdMode)
 	{
 	}
@@ -112,6 +114,7 @@ public:
 		BrushMaterialFreeInstances += BrushMaterialInstances;
 		BrushMaterialInstanceMap.Empty();
 		BrushMaterialComponents.Empty();
+		bCanPaint = false;
 	}
 
 	virtual void BeginStroke(float LandscapeX, float LandscapeY, FLandscapeTool* CurrentTool) override
@@ -137,12 +140,13 @@ public:
 		Bounds.Max.Y = FMath::CeilToInt( LastMousePosition.Y + TotalRadius);
 
 		TSet<ULandscapeComponent*> NewComponents;
-
+		bool bHasUnloadedComponents = false;
 		// Adjusting the brush may use the same keybind as moving the camera as they can be user-set, so we need this second check.
 		if (!ViewportClient->IsMovingCamera() || EdMode->IsAdjustingBrush(ViewportClient->Viewport))
 		{
 			// GetComponentsInRegion expects an inclusive max
 			LandscapeInfo->GetComponentsInRegion(Bounds.Min.X, Bounds.Min.Y, Bounds.Max.X - 1, Bounds.Max.Y - 1, NewComponents);
+			bHasUnloadedComponents = LandscapeInfo->HasUnloadedComponentsInRegion(Bounds.Min.X, Bounds.Min.Y, Bounds.Max.X - 1, Bounds.Max.Y - 1);
 		}
 
 		// Remove the material from any old components that are no longer in the region
@@ -197,43 +201,45 @@ public:
 			MaterialInstance->SetScalarParameterValue(FName(TEXT("LocalFalloff")), Falloff);
 			MaterialInstance->SetVectorParameterValue(FName(TEXT("WorldPosition")), FLinearColor(WorldLocation.X, WorldLocation.Y, WorldLocation.Z, ScaleXY));
 
-			bool bCanPaint = true;
-
-			const ALandscapeProxy* LandscapeProxy = Component->GetLandscapeProxy();
-			const ULandscapeLayerInfoObject* LayerInfo = EdMode->CurrentToolTarget.LayerInfo.Get();
-
-			if (!EdMode->CanEditLayer())
-			{
-				bCanPaint = false;
-			}
-			else if ((EdMode->CurrentToolTarget.TargetType == ELandscapeToolTargetType::Weightmap || EdMode->CurrentToolTarget.TargetType == ELandscapeToolTargetType::Visibility) &&
-				EdMode->UISettings->PaintingRestriction != ELandscapeLayerPaintingRestriction::None)
+			bCanPaint = !bHasUnloadedComponents;
+			if (bCanPaint)
 			{
-				if (EdMode->UISettings->PaintingRestriction == ELandscapeLayerPaintingRestriction::UseComponentWhitelist &&
-					!Component->LayerWhitelist.Contains(LayerInfo))
+				const ALandscapeProxy* LandscapeProxy = Component->GetLandscapeProxy();
+				const ULandscapeLayerInfoObject* LayerInfo = EdMode->CurrentToolTarget.LayerInfo.Get();
+
+				if (!EdMode->CanEditLayer())
 				{
 					bCanPaint = false;
 				}
-				else
+				else if ((EdMode->CurrentToolTarget.TargetType == ELandscapeToolTargetType::Weightmap || EdMode->CurrentToolTarget.TargetType == ELandscapeToolTargetType::Visibility) &&
+					EdMode->UISettings->PaintingRestriction != ELandscapeLayerPaintingRestriction::None)
 				{
-					TArray<FWeightmapLayerAllocationInfo>& ComponentWeightmapLayerAllocations = Component->GetWeightmapLayerAllocations(true);
-
-					bool bExisting = ComponentWeightmapLayerAllocations.ContainsByPredicate([LayerInfo](const FWeightmapLayerAllocationInfo& Allocation) { return Allocation.LayerInfo == LayerInfo; });
-					if (!bExisting)
+					if (EdMode->UISettings->PaintingRestriction == ELandscapeLayerPaintingRestriction::UseComponentWhitelist &&
+						!Component->LayerWhitelist.Contains(LayerInfo))
 					{
-						if (EdMode->UISettings->PaintingRestriction == ELandscapeLayerPaintingRestriction::ExistingOnly ||
-							(EdMode->UISettings->PaintingRestriction == ELandscapeLayerPaintingRestriction::UseMaxLayers &&
-							 LandscapeProxy->MaxPaintedLayersPerComponent > 0 && ComponentWeightmapLayerAllocations.Num() >= LandscapeProxy->MaxPaintedLayersPerComponent))
+						bCanPaint = false;
+					}
+					else
+					{
+						TArray<FWeightmapLayerAllocationInfo>& ComponentWeightmapLayerAllocations = Component->GetWeightmapLayerAllocations(true);
+
+						bool bExisting = ComponentWeightmapLayerAllocations.ContainsByPredicate([LayerInfo](const FWeightmapLayerAllocationInfo& Allocation) { return Allocation.LayerInfo == LayerInfo; });
+						if (!bExisting)
 						{
-							bCanPaint = false;
+							if (EdMode->UISettings->PaintingRestriction == ELandscapeLayerPaintingRestriction::ExistingOnly ||
+								(EdMode->UISettings->PaintingRestriction == ELandscapeLayerPaintingRestriction::UseMaxLayers &&
+									LandscapeProxy->MaxPaintedLayersPerComponent > 0 && ComponentWeightmapLayerAllocations.Num() >= LandscapeProxy->MaxPaintedLayersPerComponent))
+							{
+								bCanPaint = false;
+							}
 						}
 					}
 				}
-			}
-			
-			if (bCanPaint && EdMode->CurrentToolTarget.TargetType == ELandscapeToolTargetType::Visibility && !Component->IsLandscapeHoleMaterialValid())
-			{
-				bCanPaint = false;
+
+				if (bCanPaint && EdMode->CurrentToolTarget.TargetType == ELandscapeToolTargetType::Visibility && !Component->IsLandscapeHoleMaterialValid())
+				{
+					bCanPaint = false;
+				}
 			}
 
 			MaterialInstance->SetScalarParameterValue("CanPaint", bCanPaint ? 1.0f : 0.0f);
@@ -247,6 +253,11 @@ public:
 
 	virtual FLandscapeBrushData ApplyBrush(const TArray<FLandscapeToolInteractorPosition>& InInteractorPositions) override
 	{
+		if (!bCanPaint)
+		{
+			return FLandscapeBrushData();
+		}
+
 		ULandscapeInfo* LandscapeInfo = EdMode->CurrentToolTarget.LandscapeInfo.Get();
 		const float ScaleXY = FMath::Abs(LandscapeInfo->DrawScale.X);
 		const float TotalRadius = EdMode->UISettings->BrushRadius / ScaleXY;
@@ -1039,6 +1050,11 @@ public:
 
 	virtual FLandscapeBrushData ApplyBrush(const TArray<FLandscapeToolInteractorPosition>& InteractorPositions) override
 	{
+		if (!bCanPaint)
+		{
+			return FLandscapeBrushData();
+		}
+
 		ULandscapeInfo* LandscapeInfo = EdMode->CurrentToolTarget.LandscapeInfo.Get();
 		const float ScaleXY = FMath::Abs(LandscapeInfo->DrawScale.X);
 		const float TotalRadius = EdMode->UISettings->BrushRadius / ScaleXY;
@@ -1234,6 +1250,11 @@ public:
 
 	virtual FLandscapeBrushData ApplyBrush(const TArray<FLandscapeToolInteractorPosition>& InteractorPositions) override
 	{
+		if (!bCanPaint)
+		{
+			return FLandscapeBrushData();
+		}
+
 		ULandscapeInfo* LandscapeInfo = EdMode->CurrentToolTarget.LandscapeInfo.Get();
 		if (EdMode->UISettings->bAlphaBrushAutoRotate && OldMousePosition.IsZero())
 		{
diff --git a/Engine/Source/Runtime/Engine/Private/WorldPartition/Landscape/LandscapeActorDesc.cpp b/Engine/Source/Runtime/Engine/Private/WorldPartition/Landscape/LandscapeActorDesc.cpp
--- a/Engine/Source/Runtime/Engine/Private/WorldPartition/Landscape/LandscapeActorDesc.cpp
+++ b/Engine/Source/Runtime/Engine/Private/WorldPartition/Landscape/LandscapeActorDesc.cpp
@@ -1,9 +1,13 @@
 // Copyright Epic Games, Inc. All Rights Reserved.
 
 #include "WorldPartition/Landscape/LandscapeActorDesc.h"
-#include "LandscapeProxy.h"
 
 #if WITH_EDITOR
+#include "LandscapeProxy.h"
+#include "LandscapeStreamingProxy.h"
+#include "LandscapeInfo.h"
+#include "WorldPartition/WorldPartition.h"
+#include "WorldPartition/WorldPartitionHandle.h"
 
 void FLandscapeActorDesc::Init(const AActor* InActor)
 {
@@ -16,4 +20,38 @@ void FLandscapeActorDesc::Init(const AActor* InActor)
 	GridIndexZ = 0;
 }
 
+void FLandscapeActorDesc::Unload()
+{
+	if (ALandscapeStreamingProxy* LandscapeStreamingProxy = Cast<ALandscapeStreamingProxy>(GetActor()))
+	{
+		LandscapeStreamingProxy->ActorDescReferences.Empty();
+	}
+
+	FPartitionActorDesc::Unload();
+}
+
+void FLandscapeActorDesc::OnRegister(UWorldPartition* WorldPartition)
+{
+	FPartitionActorDesc::OnRegister(WorldPartition);
+
+	if (ULandscapeInfo* LandscapeInfo = ULandscapeInfo::Find(WorldPartition->GetWorld(), GridGuid))
+	{
+		FWorldPartitionHandle Handle(WorldPartition, GetGuid());
+		check(Handle.IsValid());
+		LandscapeInfo->ProxyHandles.Add(MoveTemp(Handle));
+	}
+}
+
+void FLandscapeActorDesc::OnUnregister(UWorldPartition* WorldPartition)
+{
+	FPartitionActorDesc::OnUnregister(WorldPartition);
+
+	if (ULandscapeInfo* LandscapeInfo = ULandscapeInfo::Find(WorldPartition->GetWorld(), GridGuid))
+	{
+		FWorldPartitionHandle Handle(WorldPartition, GetGuid());
+		check(Handle.IsValid());
+		LandscapeInfo->ProxyHandles.Remove(Handle);
+	}
+}
+
 #endif
diff --git a/Engine/Source/Runtime/Engine/Private/WorldPartition/Landscape/LandscapeSplineActorDesc.cpp b/Engine/Source/Runtime/Engine/Private/WorldPartition/Landscape/LandscapeSplineActorDesc.cpp
--- a/Engine/Source/Runtime/Engine/Private/WorldPartition/Landscape/LandscapeSplineActorDesc.cpp
+++ b/Engine/Source/Runtime/Engine/Private/WorldPartition/Landscape/LandscapeSplineActorDesc.cpp
@@ -5,6 +5,9 @@
 #if WITH_EDITOR
 
 #include "LandscapeSplineActor.h"
+#include "LandscapeInfo.h"
+#include "WorldPartition/WorldPartition.h"
+#include "WorldPartition/WorldPartitionHandle.h"
 #include "UObject/UE5MainStreamObjectVersion.h"
 
 void FLandscapeSplineActorDesc::Init(const AActor* InActor)
@@ -28,4 +31,28 @@ void FLandscapeSplineActorDesc::Serialize(FArchive& Ar)
 	}
 }
 
+void FLandscapeSplineActorDesc::OnRegister(UWorldPartition* WorldPartition)
+{
+	FWorldPartitionActorDesc::OnRegister(WorldPartition);
+
+	if (ULandscapeInfo* LandscapeInfo = ULandscapeInfo::Find(WorldPartition->GetWorld(), LandscapeGuid))
+	{
+		FWorldPartitionHandle Handle(WorldPartition, GetGuid());
+		check(Handle.IsValid());
+		LandscapeInfo->SplineHandles.Add(MoveTemp(Handle));
+	}
+}
+
+void FLandscapeSplineActorDesc::OnUnregister(UWorldPartition* WorldPartition)
+{
+	FWorldPartitionActorDesc::OnUnregister(WorldPartition);
+
+	if (ULandscapeInfo* LandscapeInfo = ULandscapeInfo::Find(WorldPartition->GetWorld(), LandscapeGuid))
+	{
+		FWorldPartitionHandle Handle(WorldPartition, GetGuid());
+		check(Handle.IsValid());
+		LandscapeInfo->SplineHandles.Remove(Handle);
+	}
+}
+
 #endif
diff --git a/Engine/Source/Runtime/Engine/Public/WorldPartition/Landscape/LandscapeActorDesc.h b/Engine/Source/Runtime/Engine/Public/WorldPartition/Landscape/LandscapeActorDesc.h
--- a/Engine/Source/Runtime/Engine/Public/WorldPartition/Landscape/LandscapeActorDesc.h
+++ b/Engine/Source/Runtime/Engine/Public/WorldPartition/Landscape/LandscapeActorDesc.h
@@ -12,7 +12,12 @@
 class ENGINE_API FLandscapeActorDesc : public FPartitionActorDesc
 {
 #if WITH_EDITOR
-protected:
+public:
 	virtual void Init(const AActor* InActor) override;
+	virtual void Unload() override;
+
+protected:
+	virtual void OnRegister(UWorldPartition* WorldPartition) override;
+	virtual void OnUnregister(UWorldPartition* WorldPartition) override;
 #endif
 };
diff --git a/Engine/Source/Runtime/Engine/Public/WorldPartition/Landscape/LandscapeSplineActorDesc.h b/Engine/Source/Runtime/Engine/Public/WorldPartition/Landscape/LandscapeSplineActorDesc.h
--- a/Engine/Source/Runtime/Engine/Public/WorldPartition/Landscape/LandscapeSplineActorDesc.h
+++ b/Engine/Source/Runtime/Engine/Public/WorldPartition/Landscape/LandscapeSplineActorDesc.h
@@ -15,6 +15,8 @@ class ENGINE_API FLandscapeSplineActorDesc : public FWorldPartitionActorDesc
 protected:
 	virtual void Init(const AActor* InActor) override;
 	virtual void Serialize(FArchive& Ar) override;
+	virtual void OnRegister(UWorldPartition* WorldPartition) override;
+	virtual void OnUnregister(UWorldPartition* WorldPartition) override;
 
 	FGuid LandscapeGuid;
 #endif
diff --git a/Engine/Source/Runtime/Engine/Public/WorldPartition/WorldPartitionActorDesc.h b/Engine/Source/Runtime/Engine/Public/WorldPartition/WorldPartitionActorDesc.h
--- a/Engine/Source/Runtime/Engine/Public/WorldPartition/WorldPartitionActorDesc.h
+++ b/Engine/Source/Runtime/Engine/Public/WorldPartition/WorldPartitionActorDesc.h
@@ -87,7 +87,7 @@ public:
 
 	AActor* GetActor() const;
 	AActor* Load(const FLinkerInstancingContext* InstancingContext = nullptr) const;
-	void Unload();
+	virtual void Unload();
 
 	virtual void Init(const AActor* InActor);
 	void Init(const FWorldPartitionActorDescInitData& DescData);
diff --git a/Engine/Source/Runtime/Landscape/Classes/LandscapeInfo.h b/Engine/Source/Runtime/Landscape/Classes/LandscapeInfo.h
--- a/Engine/Source/Runtime/Landscape/Classes/LandscapeInfo.h
+++ b/Engine/Source/Runtime/Landscape/Classes/LandscapeInfo.h
@@ -9,6 +9,11 @@
 #include "UObject/LazyObjectPtr.h"
 #include "UObject/ScriptInterface.h"
 #include "ILandscapeSplineInterface.h"
+
+#if WITH_EDITOR
+#include "WorldPartition/WorldPartitionHandle.h"
+#endif
+
 #include "LandscapeInfo.generated.h"
 
 class ALandscape;
@@ -144,6 +149,12 @@ private:
 #if WITH_EDITORONLY_DATA
 	UPROPERTY()
 	TArray<TScriptInterface<ILandscapeSplineInterface>> SplineActors;
+
+	// World Partition handles to actors with the same LandscapeGuid
+	friend class FLandscapeActorDesc;
+	TSet<FWorldPartitionHandle> ProxyHandles;
+	friend class FLandscapeSplineActorDesc;
+	TSet<FWorldPartitionHandle> SplineHandles;
 #endif
 
 	TSet<ULandscapeComponent*> SelectedComponents;
@@ -162,6 +173,7 @@ public:
 	// all below.
 	LANDSCAPE_API bool AreAllComponentsRegistered() const;
 	LANDSCAPE_API void GetComponentsInRegion(int32 X1, int32 Y1, int32 X2, int32 Y2, TSet<ULandscapeComponent*>& OutComponents, bool bOverlap = true) const;
+	LANDSCAPE_API bool HasUnloadedComponentsInRegion(int32 X1, int32 Y1, int32 X2, int32 Y2) const;
 	LANDSCAPE_API bool GetLandscapeExtent(ALandscapeProxy* Proxy, FIntRect& ProxyExtent) const;
 	LANDSCAPE_API bool GetLandscapeExtent(FIntRect& LandscapeExtent) const;
 	LANDSCAPE_API bool GetLandscapeExtent(int32& MinX, int32& MinY, int32& MaxX, int32& MaxY) const;
@@ -279,10 +291,16 @@ public:
 
 	/** Returns LandscapeStreamingProxy Cell Size in WorldPartition */
 	LANDSCAPE_API uint32 GetGridSize(uint32 InGridSizeInComponents) const;
+
+	/** Returns LandscapeSplineActor WorldPartition handles */
+	const TSet<FWorldPartitionHandle>& GetSplineHandles() const { return SplineHandles; }
 #endif
-	static ULandscapeInfo* Find(UWorld* InWorld, const FGuid& LandscapeGuid);
+	LANDSCAPE_API static ULandscapeInfo* Find(UWorld* InWorld, const FGuid& LandscapeGuid);
 	static ULandscapeInfo* FindOrCreate(UWorld* InWorld, const FGuid& LandscapeGuid);
 
+	/** Called after creating object so that it can initialize its state */
+	void Initialize(UWorld* InWorld, const FGuid& InLandscapeGuid);
+
 	/**
 	 * Runs the given function on the root landscape actor and all streaming proxies
 	 * Most easily used with a lambda as follows:
diff --git a/Engine/Source/Runtime/Landscape/Classes/LandscapeStreamingProxy.h b/Engine/Source/Runtime/Landscape/Classes/LandscapeStreamingProxy.h
--- a/Engine/Source/Runtime/Landscape/Classes/LandscapeStreamingProxy.h
+++ b/Engine/Source/Runtime/Landscape/Classes/LandscapeStreamingProxy.h
@@ -5,6 +5,11 @@
 #include "CoreMinimal.h"
 #include "UObject/ObjectMacros.h"
 #include "LandscapeProxy.h"
+
+#if WITH_EDITOR
+#include "WorldPartition/WorldPartitionHandle.h"
+#endif
+
 #include "LandscapeStreamingProxy.generated.h"
 
 class ALandscape;
@@ -21,11 +26,17 @@ public:
 	UPROPERTY(EditAnywhere, Category=LandscapeProxy)
 	TLazyObjectPtr<ALandscape> LandscapeActor;
 
+#if WITH_EDITORONLY_DATA
+	/** hard refs to actors that need to be loaded when this proxy is loaded */
+	TSet<FWorldPartitionReference> ActorDescReferences;
+#endif
+
 	//~ Begin UObject Interface
 #if WITH_EDITOR
 	virtual bool ShouldExport() override { return false;  }
 	virtual void PostEditChangeProperty(FPropertyChangedEvent& PropertyChangedEvent) override;
 	virtual bool CanEditChange(const FProperty* InProperty) const override;
+	virtual void PostRegisterAllComponents() override;
 #endif
 	//~ End UObject Interface
 
diff --git a/Engine/Source/Runtime/Landscape/Private/Landscape.cpp b/Engine/Source/Runtime/Landscape/Private/Landscape.cpp
--- a/Engine/Source/Runtime/Landscape/Private/Landscape.cpp
+++ b/Engine/Source/Runtime/Landscape/Private/Landscape.cpp
@@ -69,6 +69,9 @@ Landscape.cpp: Terrain rendering
 #include "LandscapeDataAccess.h"
 #include "UObject/EditorObjectVersion.h"
 #include "Algo/BinarySearch.h"
+#include "WorldPartition/WorldPartition.h"
+#include "WorldPartition/WorldPartitionHandle.h"
+#include "WorldPartition/WorldPartitionActorDescIterator.h"
 #if WITH_EDITOR
 #include "Rendering/StaticLightingSystemInterface.h"
 #include "Misc/ScopedSlowTask.h"
@@ -3102,12 +3105,36 @@ ULandscapeInfo* ULandscapeInfo::FindOrCreate(UWorld* InWorld, const FGuid& Lands
 	{
 		LandscapeInfo = NewObject<ULandscapeInfo>(GetTransientPackage(), NAME_None, RF_Transactional | RF_Transient);
 		LandscapeInfoMap.Modify(false);
+		LandscapeInfo->Initialize(InWorld, LandscapeGuid);
 		LandscapeInfoMap.Map.Add(LandscapeGuid, LandscapeInfo);
 	}
 	check(LandscapeInfo);
 	return LandscapeInfo;
 }
 
+void ULandscapeInfo::Initialize(UWorld* InWorld, const FGuid& InLandscapeGuid)
+{
+	LandscapeGuid = InLandscapeGuid;
+#if WITH_EDITOR
+	// If WorldPartition world get the Proxy handles
+	if (UWorldPartition* WorldPartition = InWorld->GetWorldPartition())
+	{
+		for (TWorldPartitionActorDescIterator<AActor, FWorldPartitionActorDesc> ActorDescIterator(WorldPartition); ActorDescIterator; ++ActorDescIterator)
+		{
+			if (ActorDescIterator->GetActorClass()->IsChildOf(ALandscapeStreamingProxy::StaticClass()))
+			{
+				ProxyHandles.Add(FWorldPartitionHandle(WorldPartition, ActorDescIterator->GetGuid()));
+			}
+			else if (ActorDescIterator->GetActorClass()->IsChildOf(ALandscapeSplineActor::StaticClass()))
+			{
+				SplineHandles.Emplace(FWorldPartitionHandle(WorldPartition, ActorDescIterator->GetGuid()));
+			}
+
+		}
+	}
+#endif
+}
+
 void ULandscapeInfo::ForAllLandscapeProxies(TFunctionRef<void(ALandscapeProxy*)> Fn) const
 {
 	ALandscape* Landscape = LandscapeActor.Get();
@@ -3128,16 +3155,15 @@ void ULandscapeInfo::RegisterActor(ALandscapeProxy* Proxy, bool bMapCheck)
 	// do not pass here invalid actors
 	checkSlow(Proxy);
 	check(Proxy->GetLandscapeGuid().IsValid());
-	
+	check(LandscapeGuid.IsValid());
+
 #if WITH_EDITOR
 	if (!OwningWorld->IsGameWorld())
 	{
 		// in case this Info object is not initialized yet
 		// initialized it with properties from passed actor
-		if (LandscapeGuid.IsValid() == false ||
-			(GetLandscapeProxy() == nullptr && ensure(LandscapeGuid == Proxy->GetLandscapeGuid())))
+		if (GetLandscapeProxy() == nullptr)
 		{
-			LandscapeGuid = Proxy->GetLandscapeGuid();
 			ComponentSizeQuads = Proxy->ComponentSizeQuads;
 			ComponentNumSubsections = Proxy->NumSubsections;
 			SubsectionSizeQuads = Proxy->SubsectionSizeQuads;
diff --git a/Engine/Source/Runtime/Landscape/Private/LandscapeEdit.cpp b/Engine/Source/Runtime/Landscape/Private/LandscapeEdit.cpp
--- a/Engine/Source/Runtime/Landscape/Private/LandscapeEdit.cpp
+++ b/Engine/Source/Runtime/Landscape/Private/LandscapeEdit.cpp
@@ -60,6 +60,7 @@ LandscapeEdit.cpp: Landscape editing
 #include "PhysicalMaterials/PhysicalMaterial.h"
 #include "WorldPartition/WorldPartitionActorDesc.h"
 #include "WorldPartition/Landscape/LandscapeActorDesc.h"
+#include "ActorPartition/ActorPartitionSubsystem.h"
 #endif
 #include "Algo/Count.h"
 #include "Serialization/MemoryWriter.h"
@@ -2235,6 +2236,40 @@ bool ULandscapeInfo::AreAllComponentsRegistered() const
 
 #define MAX_LANDSCAPE_SUBSECTIONS 2
 
+bool ULandscapeInfo::HasUnloadedComponentsInRegion(int32 X1, int32 Y1, int32 X2, int32 Y2) const
+{
+	if (!LandscapeActor)
+	{
+		return false;
+	}
+
+	UWorld* World = LandscapeActor->GetWorld();
+
+	int32 ComponentIndexX1, ComponentIndexY1, ComponentIndexX2, ComponentIndexY2;
+	ALandscape::CalcComponentIndicesOverlap(X1, Y1, X2, Y2, ComponentSizeQuads, ComponentIndexX1, ComponentIndexY1, ComponentIndexX2, ComponentIndexY2);
+
+	UActorPartitionSubsystem::FCellCoord MinCoord = UActorPartitionSubsystem::FCellCoord::GetCellCoord(FIntPoint(ComponentIndexX1 * ComponentSizeQuads, ComponentIndexY1 * ComponentSizeQuads), World->PersistentLevel, LandscapeActor->GridSize);
+	UActorPartitionSubsystem::FCellCoord MaxCoord = UActorPartitionSubsystem::FCellCoord::GetCellCoord(FIntPoint(ComponentIndexX2 * ComponentSizeQuads, ComponentIndexY2 * ComponentSizeQuads), World->PersistentLevel, LandscapeActor->GridSize);
+
+	for (const FWorldPartitionHandle& Handle : ProxyHandles)
+	{
+		if (FLandscapeActorDesc* LandscapeActorDesc = (FLandscapeActorDesc*)Handle.Get())
+		{
+			check(LandscapeActorDesc->GridGuid == LandscapeGuid);
+			UActorPartitionSubsystem::FCellCoord ActorCoord(LandscapeActorDesc->GridIndexX, LandscapeActorDesc->GridIndexY, LandscapeActorDesc->GridIndexZ, World->PersistentLevel);
+			if (ActorCoord.X >= MinCoord.X && ActorCoord.Y >= MinCoord.Y && ActorCoord.X <= MaxCoord.X && ActorCoord.Y <= MaxCoord.Y)
+			{
+				if (!LandscapeActorDesc->GetHardRefCount() || !LandscapeActorDesc->GetActor())
+				{
+					return true;
+				}
+			}
+		}
+	}
+
+	return false;
+}
+
 void ULandscapeInfo::GetComponentsInRegion(int32 X1, int32 Y1, int32 X2, int32 Y2, TSet<ULandscapeComponent*>& OutComponents, bool bOverlap) const
 {
 	// Find component range for this block of data
@@ -4984,6 +5019,34 @@ void ALandscapeStreamingProxy::PostEditChangeProperty(FPropertyChangedEvent& Pro
 	Super::PostEditChangeProperty(PropertyChangedEvent);
 }
 
+void ALandscapeStreamingProxy::PostRegisterAllComponents()
+{
+	ALandscapeProxy::PostRegisterAllComponents();
+
+	if (LandscapeGuid.IsValid())
+	{
+		ULandscapeInfo* LandscapeInfo = GetLandscapeInfo();
+		check(LandscapeInfo);
+		if (GEditor && !GetWorld()->IsGameWorld())
+		{
+			const TSet<FWorldPartitionHandle>& SplineHandles = LandscapeInfo->GetSplineHandles();
+			if (SplineHandles.Num())
+			{
+				FVector ActorLocation = GetActorLocation();
+				FBox Bounds(ActorLocation, ActorLocation + (GridSize * LandscapeInfo->DrawScale));
+				// Get a reference to Spline Actors that have intersecting bounds with us
+				for (const FWorldPartitionHandle& SplineHandle : LandscapeInfo->GetSplineHandles())
+				{
+					if (SplineHandle.IsValid() && Bounds.IntersectXY(SplineHandle->GetBounds()))
+					{
+						ActorDescReferences.Add(FWorldPartitionHandleHelpers::ConvertHandleToReference(SplineHandle));
+					}
+				}
+			}
+		}
+	}
+}
+
 void ALandscape::PreEditChange(FProperty* PropertyThatWillChange)
 {
 	PreEditLandscapeMaterial = LandscapeMaterial;
