Fix uninitialized script struct members in landscape structs.

diff --git a/Engine/Source/Editor/LandscapeEditor/Private/LandscapeEditorObject.cpp b/Engine/Source/Editor/LandscapeEditor/Private/LandscapeEditorObject.cpp
--- a/Engine/Source/Editor/LandscapeEditor/Private/LandscapeEditorObject.cpp
+++ b/Engine/Source/Editor/LandscapeEditor/Private/LandscapeEditorObject.cpp
@@ -98,7 +98,7 @@ ULandscapeEditorObject::ULandscapeEditorObject(const FObjectInitializer& ObjectI
 	, AlphaBrushPanU(0.5f)
 	, AlphaBrushPanV(0.5f)
 	, bUseWorldSpacePatternBrush(false)
-	, WorldSpacePatternBrushSettings(FLandscapePatternBrushWorldSpaceSettings{FVector2D::ZeroVector, 0.0f, false, 3200})
+	, WorldSpacePatternBrushSettings(FVector2D::ZeroVector, 0.0f, false, 3200)
 	, AlphaTexture(NULL)
 	, AlphaTextureChannel(EColorChannel::Red)
 	, AlphaTextureSizeX(1)
diff --git a/Engine/Source/Editor/LandscapeEditor/Public/LandscapeEditorObject.h b/Engine/Source/Editor/LandscapeEditor/Public/LandscapeEditorObject.h
--- a/Engine/Source/Editor/LandscapeEditor/Public/LandscapeEditorObject.h
+++ b/Engine/Source/Editor/LandscapeEditor/Public/LandscapeEditorObject.h
@@ -247,7 +247,19 @@ struct FLandscapePatternBrushWorldSpaceSettings
 	UPROPERTY(Category = "World-Space", EditAnywhere)
 	float RepeatSize;
 
-	FLandscapePatternBrushWorldSpaceSettings() = default;
+	FLandscapePatternBrushWorldSpaceSettings()
+		: Origin(ForceInit)
+		, Rotation(0.0f)
+		, bCenterTextureOnOrigin(false)
+		, RepeatSize(0.0f)
+	{}
+
+	FLandscapePatternBrushWorldSpaceSettings(const FVector2D& InOrigin, float InRotation, bool bInCenterTextureOnOrigin, float bRepeatSize)
+		: Origin(InOrigin)
+		, Rotation(InRotation)
+		, bCenterTextureOnOrigin(bInCenterTextureOnOrigin)
+		, RepeatSize(bRepeatSize)
+	{}
 };
 
 UCLASS(MinimalAPI)
diff --git a/Engine/Source/Runtime/Landscape/Classes/LandscapeMaterialInstanceConstant.h b/Engine/Source/Runtime/Landscape/Classes/LandscapeMaterialInstanceConstant.h
--- a/Engine/Source/Runtime/Landscape/Classes/LandscapeMaterialInstanceConstant.h
+++ b/Engine/Source/Runtime/Landscape/Classes/LandscapeMaterialInstanceConstant.h
@@ -12,6 +12,10 @@ struct FLandscapeMaterialTextureStreamingInfo
 {
 	GENERATED_USTRUCT_BODY()
 
+	FLandscapeMaterialTextureStreamingInfo()
+		: TexelFactor(0.0f)
+	{}
+
 	UPROPERTY()
 	FName TextureName;
 
