Fix crash when changing map in landscape mode.

diff --git a/Engine/Source/Editor/LandscapeEditor/Private/LandscapeEdMode.cpp b/Engine/Source/Editor/LandscapeEditor/Private/LandscapeEdMode.cpp
--- a/Engine/Source/Editor/LandscapeEditor/Private/LandscapeEdMode.cpp
+++ b/Engine/Source/Editor/LandscapeEditor/Private/LandscapeEdMode.cpp
@@ -595,8 +595,11 @@ void FEdModeLandscape::Exit()
 
 	if (CurrentToolTarget.LandscapeInfo.IsValid())
 	{
-		ALandscapeProxy* LandscapeProxy = CurrentToolTarget.LandscapeInfo->GetLandscapeProxy();
-		LandscapeProxy->OnMaterialChangedDelegate().RemoveAll(this);
+		if (ALandscapeProxy* LandscapeProxy = CurrentToolTarget.LandscapeInfo->GetLandscapeProxy())
+		{
+			LandscapeProxy->OnMaterialChangedDelegate().RemoveAll(this);
+		}
+
 		if (ALandscape* Landscape = GetLandscape())
 		{
 			Landscape->OnBlueprintBrushChangedDelegate().RemoveAll(this);
