Landscape: Fix Import tool crash (when resetting heightmap file and going out and in LS Mode)

diff --git a/Engine/Source/Editor/LandscapeEditor/Private/LandscapeEditorObject.cpp b/Engine/Source/Editor/LandscapeEditor/Private/LandscapeEditorObject.cpp
--- a/Engine/Source/Editor/LandscapeEditor/Private/LandscapeEditorObject.cpp
+++ b/Engine/Source/Editor/LandscapeEditor/Private/LandscapeEditorObject.cpp
@@ -557,6 +557,7 @@ void ULandscapeEditorObject::RefreshImports()
 {
 	ClearImportLandscapeData();
 	HeightmapImportDescriptorIndex = 0;
+	HeightmapImportDescriptor.Reset();
 	ImportLandscape_Width = 0;
 	ImportLandscape_Height = 0;
 
