Added an option to flip coordinate on Y axis when importing tiled landscape in World Partition.

diff --git a/Engine/Source/Editor/LandscapeEditor/Private/LandscapeEdMode.cpp b/Engine/Source/Editor/LandscapeEditor/Private/LandscapeEdMode.cpp
--- a/Engine/Source/Editor/LandscapeEditor/Private/LandscapeEdMode.cpp
+++ b/Engine/Source/Editor/LandscapeEditor/Private/LandscapeEdMode.cpp
@@ -2872,7 +2872,7 @@ void FEdModeLandscape::ReimportData(const FLandscapeTargetListInfo& TargetInfo)
 }
 
 template<class T>
-void ImportDataInternal(ULandscapeInfo* LandscapeInfo, const FString& Filename, FName LayerName, bool bSingleFile, const FIntRect& ImportRegionVerts, ELandscapeImportTransformType TransformType, FIntPoint Offset, TFunctionRef<void(int32, int32, int32, int32, const TArray<T>&)> SetDataFunc)
+void ImportDataInternal(ULandscapeInfo* LandscapeInfo, const FString& Filename, FName LayerName, bool bSingleFile, bool bFlipYAxis, const FIntRect& ImportRegionVerts, ELandscapeImportTransformType TransformType, FIntPoint Offset, TFunctionRef<void(int32, int32, int32, int32, const TArray<T>&)> SetDataFunc)
 {
 	if (LandscapeInfo)
 	{
@@ -2880,7 +2880,7 @@ void ImportDataInternal(ULandscapeInfo* LandscapeInfo, const FString& Filename,
 		FLandscapeImportDescriptor OutImportDescriptor;
 		FText OutMessage;
 
-		ELandscapeImportResult ImportResult = FLandscapeImportHelper::GetImportDescriptor<T>(Filename, bSingleFile, LayerName, OutImportDescriptor, OutMessage);
+		ELandscapeImportResult ImportResult = FLandscapeImportHelper::GetImportDescriptor<T>(Filename, bSingleFile, bFlipYAxis, LayerName, OutImportDescriptor, OutMessage);
 		if (ImportResult == ELandscapeImportResult::Error)
 		{
 			FMessageDialog::Open(EAppMsgType::Ok, OutMessage);
@@ -2968,9 +2968,9 @@ void ImportDataInternal(ULandscapeInfo* LandscapeInfo, const FString& Filename,
 	}			
 }
 
-void FEdModeLandscape::ImportHeightData(ULandscapeInfo* LandscapeInfo, const FGuid& LayerGuid, const FString& Filename, const FIntRect& ImportRegionVerts, ELandscapeImportTransformType TransformType, FIntPoint Offset, const ELandscapeLayerPaintingRestriction& PaintRestriction)
+void FEdModeLandscape::ImportHeightData(ULandscapeInfo* LandscapeInfo, const FGuid& LayerGuid, const FString& Filename, const FIntRect& ImportRegionVerts, ELandscapeImportTransformType TransformType, FIntPoint Offset, const ELandscapeLayerPaintingRestriction& PaintRestriction, bool bFlipYAxis)
 {
-	ImportDataInternal<uint16>(LandscapeInfo, Filename, NAME_None, UseSingleFileImport(), ImportRegionVerts, TransformType, Offset, [LandscapeInfo, LayerGuid, PaintRestriction](int32 MinX, int32 MinY, int32 MaxX, int32 MaxY, const TArray<uint16>& Data)
+	ImportDataInternal<uint16>(LandscapeInfo, Filename, NAME_None, UseSingleFileImport(), bFlipYAxis, ImportRegionVerts, TransformType, Offset, [LandscapeInfo, LayerGuid, PaintRestriction](int32 MinX, int32 MinY, int32 MaxX, int32 MaxY, const TArray<uint16>& Data)
 	{
 		ALandscape* Landscape = LandscapeInfo->LandscapeActor.Get();
 		FScopedSetLandscapeEditingLayer Scope(Landscape, LayerGuid, [&] { check(Landscape); Landscape->RequestLayersContentUpdate(ELandscapeLayerUpdateMode::Update_Heightmap_All); });
@@ -2982,11 +2982,11 @@ void FEdModeLandscape::ImportHeightData(ULandscapeInfo* LandscapeInfo, const FGu
 	});
 }
 
-void FEdModeLandscape::ImportWeightData(ULandscapeInfo* LandscapeInfo, const FGuid& LayerGuid, ULandscapeLayerInfoObject* LayerInfo, const FString& Filename, const FIntRect& ImportRegionVerts, ELandscapeImportTransformType TransformType, FIntPoint Offset, const ELandscapeLayerPaintingRestriction& PaintRestriction)
+void FEdModeLandscape::ImportWeightData(ULandscapeInfo* LandscapeInfo, const FGuid& LayerGuid, ULandscapeLayerInfoObject* LayerInfo, const FString& Filename, const FIntRect& ImportRegionVerts, ELandscapeImportTransformType TransformType, FIntPoint Offset, const ELandscapeLayerPaintingRestriction& PaintRestriction, bool bFlipYAxis)
 {
 	if (LayerInfo)
 	{
-		ImportDataInternal<uint8>(LandscapeInfo, Filename, LayerInfo->LayerName, UseSingleFileImport(), ImportRegionVerts, TransformType, Offset, [LandscapeInfo, LayerGuid, LayerInfo, PaintRestriction](int32 MinX, int32 MinY, int32 MaxX, int32 MaxY, const TArray<uint8>& Data)
+		ImportDataInternal<uint8>(LandscapeInfo, Filename, LayerInfo->LayerName, UseSingleFileImport(), bFlipYAxis, ImportRegionVerts, TransformType, Offset, [LandscapeInfo, LayerGuid, LayerInfo, PaintRestriction](int32 MinX, int32 MinY, int32 MaxX, int32 MaxY, const TArray<uint8>& Data)
 		{
 			ALandscape* Landscape = LandscapeInfo->LandscapeActor.Get();
 			FScopedSetLandscapeEditingLayer Scope(Landscape, LayerGuid, [&] { check(Landscape); Landscape->RequestLayersContentUpdate(ELandscapeLayerUpdateMode::Update_Weightmap_All); });
diff --git a/Engine/Source/Editor/LandscapeEditor/Private/LandscapeEdMode.h b/Engine/Source/Editor/LandscapeEditor/Private/LandscapeEdMode.h
--- a/Engine/Source/Editor/LandscapeEditor/Private/LandscapeEdMode.h
+++ b/Engine/Source/Editor/LandscapeEditor/Private/LandscapeEdMode.h
@@ -594,8 +594,8 @@ public:
 
 	void ReimportData(const FLandscapeTargetListInfo& TargetInfo);
 	void ImportData(const FLandscapeTargetListInfo& TargetInfo, const FString& Filename);
-	void ImportHeightData(ULandscapeInfo* LandscapeInfo, const FGuid& LayerGuid, const FString& Filename, const FIntRect& ImportRegionVerts, ELandscapeImportTransformType TransformType = ELandscapeImportTransformType::ExpandCentered, FIntPoint Offset = FIntPoint(0,0), const ELandscapeLayerPaintingRestriction& PaintRestriction = ELandscapeLayerPaintingRestriction::None);
-	void ImportWeightData(ULandscapeInfo* LandscapeInfo, const FGuid& LayerGuid, ULandscapeLayerInfoObject* LayerInfo, const FString& Filename, const FIntRect& ImportRegionVerts, ELandscapeImportTransformType TransformType = ELandscapeImportTransformType::ExpandCentered, FIntPoint Offset = FIntPoint(0, 0), const ELandscapeLayerPaintingRestriction& PaintRestriction = ELandscapeLayerPaintingRestriction::None);
+	void ImportHeightData(ULandscapeInfo* LandscapeInfo, const FGuid& LayerGuid, const FString& Filename, const FIntRect& ImportRegionVerts, ELandscapeImportTransformType TransformType = ELandscapeImportTransformType::ExpandCentered, FIntPoint Offset = FIntPoint(0,0), const ELandscapeLayerPaintingRestriction& PaintRestriction = ELandscapeLayerPaintingRestriction::None, bool bFlipYAxis = false);
+	void ImportWeightData(ULandscapeInfo* LandscapeInfo, const FGuid& LayerGuid, ULandscapeLayerInfoObject* LayerInfo, const FString& Filename, const FIntRect& ImportRegionVerts, ELandscapeImportTransformType TransformType = ELandscapeImportTransformType::ExpandCentered, FIntPoint Offset = FIntPoint(0, 0), const ELandscapeLayerPaintingRestriction& PaintRestriction = ELandscapeLayerPaintingRestriction::None, bool bFlipYAxis = false);
 	bool UseSingleFileImport() const { return !IsGridBased(); }
 
 	/** Resample landscape to a different resolution or change the component size */
diff --git a/Engine/Source/Editor/LandscapeEditor/Private/LandscapeEditorDetailCustomization_ImportExport.cpp b/Engine/Source/Editor/LandscapeEditor/Private/LandscapeEditorDetailCustomization_ImportExport.cpp
--- a/Engine/Source/Editor/LandscapeEditor/Private/LandscapeEditorDetailCustomization_ImportExport.cpp
+++ b/Engine/Source/Editor/LandscapeEditor/Private/LandscapeEditorDetailCustomization_ImportExport.cpp
@@ -155,6 +155,9 @@ void FLandscapeEditorDetailCustomization_ImportExport::CustomizeDetails(IDetailL
 	TSharedRef<IPropertyHandle> PropertyHandle_ImportType = DetailBuilder.GetProperty(GET_MEMBER_NAME_CHECKED(ULandscapeEditorObject, ImportType));
 	ImportExportCategory.AddProperty(PropertyHandle_ImportType).Visibility(MakeAttributeLambda([]() { return FLandscapeEditorDetailCustomization_ImportExport::GetImportExportVisibility(true); }));
 
+	TSharedRef<IPropertyHandle> PropertyHandle_FlipYAxis = DetailBuilder.GetProperty(GET_MEMBER_NAME_CHECKED(ULandscapeEditorObject, bFlipYAxis));
+	ImportExportCategory.AddProperty(PropertyHandle_FlipYAxis).Visibility(MakeAttributeLambda([]() { return (IsImporting() && !GetEditorMode()->UseSingleFileImport()) ? EVisibility::Visible : EVisibility::Collapsed; }));
+
 	ImportExportCategory.AddProperty(PropertyHandle_ExportEditLayer).Visibility(MakeAttributeLambda([this]()
 	{
 		if (ULandscapeInfo* LandscapeInfo = GetEditorMode()->CurrentToolTarget.LandscapeInfo.Get())
@@ -454,7 +457,7 @@ FReply FLandscapeEditorDetailCustomization_ImportExport::OnImportExportButtonCli
 			if (LandscapeEdMode->UISettings->bHeightmapSelected)
 			{
 				check(LandscapeEdMode->UISettings->ImportLandscape_HeightmapImportResult != ELandscapeImportResult::Error);
-				LandscapeEdMode->ImportHeightData(LandscapeInfo, CurrentLayerGuid, LandscapeEdMode->UISettings->ImportLandscape_HeightmapFilename, ImportRegion, TransformType, ImportOffset, PaintRestriction);
+				LandscapeEdMode->ImportHeightData(LandscapeInfo, CurrentLayerGuid, LandscapeEdMode->UISettings->ImportLandscape_HeightmapFilename, ImportRegion, TransformType, ImportOffset, PaintRestriction, LandscapeEdMode->UISettings->bFlipYAxis);
 			}
 
 			for (const FLandscapeImportLayer& ImportLayer : LandscapeEdMode->UISettings->ImportLandscape_Layers)
@@ -462,7 +465,7 @@ FReply FLandscapeEditorDetailCustomization_ImportExport::OnImportExportButtonCli
 				if (ImportLayer.bSelected)
 				{
 					check(ImportLayer.ImportResult != ELandscapeImportResult::Error);
-					LandscapeEdMode->ImportWeightData(LandscapeInfo, CurrentLayerGuid, ImportLayer.LayerInfo, ImportLayer.SourceFilePath, ImportRegion, TransformType, ImportOffset, PaintRestriction);
+					LandscapeEdMode->ImportWeightData(LandscapeInfo, CurrentLayerGuid, ImportLayer.LayerInfo, ImportLayer.SourceFilePath, ImportRegion, TransformType, ImportOffset, PaintRestriction, LandscapeEdMode->UISettings->bFlipYAxis);
 				}
 			}
 		}
diff --git a/Engine/Source/Editor/LandscapeEditor/Private/LandscapeEditorDetailCustomization_NewLandscape.cpp b/Engine/Source/Editor/LandscapeEditor/Private/LandscapeEditorDetailCustomization_NewLandscape.cpp
--- a/Engine/Source/Editor/LandscapeEditor/Private/LandscapeEditorDetailCustomization_NewLandscape.cpp
+++ b/Engine/Source/Editor/LandscapeEditor/Private/LandscapeEditorDetailCustomization_NewLandscape.cpp
@@ -93,6 +93,15 @@ void FLandscapeEditorDetailCustomization_NewLandscape::CustomizeDetails(IDetailL
 
 	TSharedRef<IPropertyHandle> PropertyHandle_CanHaveLayersContent = DetailBuilder.GetProperty(GET_MEMBER_NAME_CHECKED(ULandscapeEditorObject, bCanHaveLayersContent));
 	NewLandscapeCategory.AddProperty(PropertyHandle_CanHaveLayersContent);
+
+	TSharedRef<IPropertyHandle> PropertyHandle_FlipYAxis = DetailBuilder.GetProperty(GET_MEMBER_NAME_CHECKED(ULandscapeEditorObject, bFlipYAxis));
+	NewLandscapeCategory.AddProperty(PropertyHandle_FlipYAxis).Visibility(MakeAttributeLambda([]()
+	{ 
+		return ((GetEditorMode()->NewLandscapePreviewMode == ENewLandscapePreviewMode::ImportLandscape) && !GetEditorMode()->UseSingleFileImport()) ? EVisibility::Visible : EVisibility::Collapsed; 
+	}));
+	
+	PropertyHandle_FlipYAxis->SetOnPropertyValueChanged(FSimpleDelegate::CreateLambda([this]() { OnImportHeightmapFilenameChanged(); }));
+
 	TSharedRef<IPropertyHandle> PropertyHandle_GridSize = DetailBuilder.GetProperty(GET_MEMBER_NAME_CHECKED(ULandscapeEditorObject, WorldPartitionGridSize));
 	
 	FEdModeLandscape* LandscapeEdMode = GetEditorMode();
diff --git a/Engine/Source/Editor/LandscapeEditor/Private/LandscapeEditorObject.cpp b/Engine/Source/Editor/LandscapeEditor/Private/LandscapeEditorObject.cpp
--- a/Engine/Source/Editor/LandscapeEditor/Private/LandscapeEditorObject.cpp
+++ b/Engine/Source/Editor/LandscapeEditor/Private/LandscapeEditorObject.cpp
@@ -567,7 +567,7 @@ void ULandscapeEditorObject::RefreshImports()
 	if (!ImportLandscape_HeightmapFilename.IsEmpty())
 	{
 		ImportLandscape_HeightmapImportResult =
-			FLandscapeImportHelper::GetHeightmapImportDescriptor(ImportLandscape_HeightmapFilename, UseSingleFileImport(), HeightmapImportDescriptor, ImportLandscape_HeightmapErrorMessage);
+			FLandscapeImportHelper::GetHeightmapImportDescriptor(ImportLandscape_HeightmapFilename, UseSingleFileImport(), bFlipYAxis, HeightmapImportDescriptor, ImportLandscape_HeightmapErrorMessage);
 		if (ImportLandscape_HeightmapImportResult != ELandscapeImportResult::Error)
 		{
 			ImportLandscape_Width = HeightmapImportDescriptor.ImportResolutions[HeightmapImportDescriptorIndex].Width;
@@ -597,7 +597,7 @@ void ULandscapeEditorObject::RefreshLayerImports()
 
 		if (UIImportLayer.LayerInfo != nullptr && !UIImportLayer.SourceFilePath.IsEmpty())
 		{
-			UIImportLayer.ImportResult = FLandscapeImportHelper::GetWeightmapImportDescriptor(UIImportLayer.SourceFilePath, UseSingleFileImport(), UIImportLayer.LayerName, UIImportLayer.ImportDescriptor, UIImportLayer.ErrorMessage);
+			UIImportLayer.ImportResult = FLandscapeImportHelper::GetWeightmapImportDescriptor(UIImportLayer.SourceFilePath, UseSingleFileImport(), bFlipYAxis, UIImportLayer.LayerName, UIImportLayer.ImportDescriptor, UIImportLayer.ErrorMessage);
 			if (UIImportLayer.ImportResult == ELandscapeImportResult::Error)
 			{
 				continue;
@@ -677,7 +677,7 @@ ELandscapeImportResult ULandscapeEditorObject::CreateImportLayersInfo(TArray<FLa
 
 		if (ImportLayer.LayerInfo != nullptr && !ImportLayer.SourceFilePath.IsEmpty())
 		{
-			UIImportLayer.ImportResult = FLandscapeImportHelper::GetWeightmapImportDescriptor(ImportLayer.SourceFilePath, UseSingleFileImport(), ImportLayer.LayerName, UIImportLayer.ImportDescriptor, UIImportLayer.ErrorMessage);
+			UIImportLayer.ImportResult = FLandscapeImportHelper::GetWeightmapImportDescriptor(ImportLayer.SourceFilePath, UseSingleFileImport(), bFlipYAxis, ImportLayer.LayerName, UIImportLayer.ImportDescriptor, UIImportLayer.ErrorMessage);
 			if (UIImportLayer.ImportResult == ELandscapeImportResult::Error)
 			{
 				FMessageDialog::Open(EAppMsgType::Ok, UIImportLayer.ErrorMessage);
@@ -884,7 +884,7 @@ void ULandscapeEditorObject::RefreshImportLayersList(bool bRefreshFromTarget)
 			}
 			else
 			{
-				NewImportLayer.ImportResult = FLandscapeImportHelper::GetWeightmapImportDescriptor(NewImportLayer.SourceFilePath, UseSingleFileImport(), NewImportLayer.LayerName, NewImportLayer.ImportDescriptor, NewImportLayer.ErrorMessage);
+				NewImportLayer.ImportResult = FLandscapeImportHelper::GetWeightmapImportDescriptor(NewImportLayer.SourceFilePath, UseSingleFileImport(), bFlipYAxis, NewImportLayer.LayerName, NewImportLayer.ImportDescriptor, NewImportLayer.ErrorMessage);
 				if (NewImportLayer.ImportResult != ELandscapeImportResult::Error)
 				{
 					int32 DescriptorIndex = NewImportLayer.ImportDescriptor.FindDescriptorIndex(ImportLandscape_Width, ImportLandscape_Height);
diff --git a/Engine/Source/Editor/LandscapeEditor/Private/LandscapeImportHelper.cpp b/Engine/Source/Editor/LandscapeEditor/Private/LandscapeImportHelper.cpp
--- a/Engine/Source/Editor/LandscapeEditor/Private/LandscapeImportHelper.cpp
+++ b/Engine/Source/Editor/LandscapeEditor/Private/LandscapeImportHelper.cpp
@@ -113,7 +113,7 @@ ELandscapeImportResult GetImportDataInternal(const FLandscapeImportDescriptor& I
 }
 
 template<class T>
-ELandscapeImportResult GetImportDescriptorInternal(const FString& FilePath, bool bSingleFile, FName LayerName, FLandscapeImportDescriptor& OutImportDescriptor, FText& OutMessage)
+ELandscapeImportResult GetImportDescriptorInternal(const FString& FilePath, bool bSingleFile, bool bFlipYAxis, FName LayerName, FLandscapeImportDescriptor& OutImportDescriptor, FText& OutMessage)
 {
 	OutImportDescriptor.Reset();
 	if (FilePath.IsEmpty())
@@ -234,6 +234,15 @@ ELandscapeImportResult GetImportDescriptorInternal(const FString& FilePath, bool
 	}
 	MaxCoord -= MinCoord;
 
+	// Flip Coordinates on Y axis
+	if (bFlipYAxis)
+	{
+		for (FLandscapeImportFileDescriptor& FileDescriptor : OutImportDescriptor.FileDescriptors)
+		{
+			FileDescriptor.Coord.Y = MaxCoord.Y - FileDescriptor.Coord.Y;
+		}
+	}
+
 	// Compute Import Total Size
 	for (const FLandscapeFileResolution& Resolution : OutImportDescriptor.FileResolutions)
 	{
@@ -278,14 +287,14 @@ void TransformImportDataInternal(const TArray<T>& InData, TArray<T>& OutData, co
 	}
 }
 
-ELandscapeImportResult FLandscapeImportHelper::GetHeightmapImportDescriptor(const FString& FilePath, bool bSingleFile, FLandscapeImportDescriptor& OutImportDescriptor, FText& OutMessage)
+ELandscapeImportResult FLandscapeImportHelper::GetHeightmapImportDescriptor(const FString& FilePath, bool bSingleFile, bool bFlipYAxis, FLandscapeImportDescriptor& OutImportDescriptor, FText& OutMessage)
 {
-	return GetImportDescriptorInternal<uint16>(FilePath, bSingleFile, NAME_None, OutImportDescriptor, OutMessage);
+	return GetImportDescriptorInternal<uint16>(FilePath, bSingleFile, bFlipYAxis, NAME_None, OutImportDescriptor, OutMessage);
 }
 
-ELandscapeImportResult FLandscapeImportHelper::GetWeightmapImportDescriptor(const FString& FilePath, bool bSingleFile, FName LayerName, FLandscapeImportDescriptor& OutImportDescriptor, FText& OutMessage)
+ELandscapeImportResult FLandscapeImportHelper::GetWeightmapImportDescriptor(const FString& FilePath, bool bSingleFile, bool bFlipYAxis, FName LayerName, FLandscapeImportDescriptor& OutImportDescriptor, FText& OutMessage)
 {
-	return GetImportDescriptorInternal<uint8>(FilePath, bSingleFile, LayerName, OutImportDescriptor, OutMessage);
+	return GetImportDescriptorInternal<uint8>(FilePath, bSingleFile, bFlipYAxis, LayerName, OutImportDescriptor, OutMessage);
 }
 
 ELandscapeImportResult FLandscapeImportHelper::GetHeightmapImportData(const FLandscapeImportDescriptor& ImportDescriptor, int32 DescriptorIndex, TArray<uint16>& OutData, FText& OutMessage)
diff --git a/Engine/Source/Editor/LandscapeEditor/Private/NewLandscapeUtils.cpp b/Engine/Source/Editor/LandscapeEditor/Private/NewLandscapeUtils.cpp
--- a/Engine/Source/Editor/LandscapeEditor/Private/NewLandscapeUtils.cpp
+++ b/Engine/Source/Editor/LandscapeEditor/Private/NewLandscapeUtils.cpp
@@ -40,7 +40,7 @@ void FNewLandscapeUtils::ImportLandscapeData( ULandscapeEditorObject* UISettings
 	{
 		FLandscapeImportDescriptor OutImportDescriptor;
 		UISettings->ImportLandscape_HeightmapImportResult =
-			FLandscapeImportHelper::GetHeightmapImportDescriptor(UISettings->ImportLandscape_HeightmapFilename, UISettings->UseSingleFileImport(), OutImportDescriptor, UISettings->ImportLandscape_HeightmapErrorMessage);
+			FLandscapeImportHelper::GetHeightmapImportDescriptor(UISettings->ImportLandscape_HeightmapFilename, UISettings->UseSingleFileImport(), UISettings->bFlipYAxis, OutImportDescriptor, UISettings->ImportLandscape_HeightmapErrorMessage);
 
 		if (UISettings->ImportLandscape_HeightmapImportResult != ELandscapeImportResult::Error)
 		{
diff --git a/Engine/Source/Editor/LandscapeEditor/Public/LandscapeEditorObject.h b/Engine/Source/Editor/LandscapeEditor/Public/LandscapeEditorObject.h
--- a/Engine/Source/Editor/LandscapeEditor/Public/LandscapeEditorObject.h
+++ b/Engine/Source/Editor/LandscapeEditor/Public/LandscapeEditorObject.h
@@ -571,6 +571,9 @@ public:
 	UPROPERTY(Category = "New Landscape", EditAnywhere, NonTransactional, meta = (DisplayName= "Enable Edit Layers", ToolTip="Enable support for landscape edit layers.", ShowForTools= "NewLandscape"))
 	bool bCanHaveLayersContent = false;
 
+	UPROPERTY(Category = "New Landscape", EditAnywhere, NonTransactional, meta = (DisplayName = "Flip Y Axis", ToolTip = "Whether to flip Y coordinate of imported files.", ShowForTools = "NewLandscape"))
+	bool bFlipYAxis = false;
+
 	UPROPERTY(Category = "New Landscape", EditAnywhere, NonTransactional, meta = (DisplayName= "World Partition Grid Size", ToolTip="Number of components per landscape streaming proxies per axis", ShowForTools="NewLandscape", ClampMin=1, ClampMax=16, UIMin=1, UIMax=16))
 	uint32 WorldPartitionGridSize = 2;
 
diff --git a/Engine/Source/Editor/LandscapeEditor/Public/LandscapeImportHelper.h b/Engine/Source/Editor/LandscapeEditor/Public/LandscapeImportHelper.h
--- a/Engine/Source/Editor/LandscapeEditor/Public/LandscapeImportHelper.h
+++ b/Engine/Source/Editor/LandscapeEditor/Public/LandscapeImportHelper.h
@@ -116,7 +116,7 @@ public:
 	 * @param OutMessage In case the result of the operation is  an error or warning, this will contain the reason.
 	 * @return If the operation was a Success/Error/Warning. 
 	 */
-	static ELandscapeImportResult GetHeightmapImportDescriptor(const FString& FilePath, bool bSingleFile, FLandscapeImportDescriptor& OutImportDescriptor, FText& OutMessage);
+	static ELandscapeImportResult GetHeightmapImportDescriptor(const FString& FilePath, bool bSingleFile, bool bFlipYAxis, FLandscapeImportDescriptor& OutImportDescriptor, FText& OutMessage);
 
 	/**
 	 * @param FilePath File to import when bSingleFile is true.
@@ -126,7 +126,7 @@ public:
 	 * @param OutMessage In case the result of the operation is  an error or warning, this will contain the reason.
 	 * @return If the operation was a Success/Error/Warning.
 	 */
-	static ELandscapeImportResult GetWeightmapImportDescriptor(const FString& FilePath, bool bSingleFile, FName LayerName, FLandscapeImportDescriptor& OutImportDescriptor, FText& OutMessage);
+	static ELandscapeImportResult GetWeightmapImportDescriptor(const FString& FilePath, bool bSingleFile, bool bFlipYAxis, FName LayerName, FLandscapeImportDescriptor& OutImportDescriptor, FText& OutMessage);
 
 	static ELandscapeImportResult GetHeightmapImportData(const FLandscapeImportDescriptor& ImportDescriptor, int32 DescriptorIndex, TArray<uint16>& OutData, FText& OutMessage);
 	static ELandscapeImportResult GetWeightmapImportData(const FLandscapeImportDescriptor& ImportDescriptor, int32 DescriptorIndex, FName LayerName, TArray<uint8>& OutData, FText& OutMessage);
@@ -141,9 +141,9 @@ public:
 
 
 	template<typename T>
-	static typename TEnableIf<TIsSame<T, uint16>::Value, ELandscapeImportResult>::Type GetImportDescriptor(const FString& FilePath, bool bSingleFile, FName LayerName, FLandscapeImportDescriptor& OutImportDescriptor, FText& OutMessage)
+	static typename TEnableIf<TIsSame<T, uint16>::Value, ELandscapeImportResult>::Type GetImportDescriptor(const FString& FilePath, bool bSingleFile, bool bFlipYAxis, FName LayerName, FLandscapeImportDescriptor& OutImportDescriptor, FText& OutMessage)
 	{
-		return GetHeightmapImportDescriptor(FilePath, bSingleFile, OutImportDescriptor, OutMessage);
+		return GetHeightmapImportDescriptor(FilePath, bSingleFile, bFlipYAxis, OutImportDescriptor, OutMessage);
 	}
 
 	template<typename T>
@@ -158,9 +158,9 @@ public:
 	}
 
 	template<typename T>
-	static typename TEnableIf<TIsSame<T, uint8>::Value, ELandscapeImportResult>::Type GetImportDescriptor(const FString& FilePath, bool bSingleFile, FName LayerName, FLandscapeImportDescriptor& OutImportDescriptor, FText& OutMessage)
+	static typename TEnableIf<TIsSame<T, uint8>::Value, ELandscapeImportResult>::Type GetImportDescriptor(const FString& FilePath, bool bSingleFile, bool bFlipYAxis, FName LayerName, FLandscapeImportDescriptor& OutImportDescriptor, FText& OutMessage)
 	{
-		return GetWeightmapImportDescriptor(FilePath, bSingleFile, LayerName, OutImportDescriptor, OutMessage);
+		return GetWeightmapImportDescriptor(FilePath, bSingleFile, bFlipYAxis, LayerName, OutImportDescriptor, OutMessage);
 	}
 	template<typename T>
 	static typename TEnableIf<TIsSame<T, uint8>::Value, ELandscapeImportResult>::Type GetImportData(const FLandscapeImportDescriptor& ImportDescriptor, int32 DescriptorIndex, FName LayerName, TArray<T>& OutData, FText& OutMessage)
