Landscape: Replace UE4 by UE in comments

diff --git a/Engine/Source/Editor/LandscapeEditor/Public/LandscapeEditorObject.h b/Engine/Source/Editor/LandscapeEditor/Public/LandscapeEditorObject.h
--- a/Engine/Source/Editor/LandscapeEditor/Public/LandscapeEditorObject.h
+++ b/Engine/Source/Editor/LandscapeEditor/Public/LandscapeEditorObject.h
@@ -577,7 +577,7 @@ public:
 	UPROPERTY(Category = "New Landscape", EditAnywhere, NonTransactional, meta = (DisplayName= "World Partition Grid Size", ToolTip="Number of components per landscape streaming proxies per axis", ShowForTools="NewLandscape", ClampMin=1, ClampMax=16, UIMin=1, UIMax=16))
 	uint32 WorldPartitionGridSize = 2;
 
-	// Whether the imported alpha maps are to be interpreted as "layered" or "additive" (UE4 uses additive internally)
+	// Whether the imported alpha maps are to be interpreted as "layered" or "additive" (UE uses additive internally)
 	UPROPERTY(Category="New Landscape", EditAnywhere, NonTransactional, meta=(DisplayName="Layer Alphamap Type", ShowForTools="NewLandscape"))
 	ELandscapeImportAlphamapType ImportLandscape_AlphamapType;
 
diff --git a/Engine/Source/Runtime/Landscape/Classes/LandscapeProxy.h b/Engine/Source/Runtime/Landscape/Classes/LandscapeProxy.h
--- a/Engine/Source/Runtime/Landscape/Classes/LandscapeProxy.h
+++ b/Engine/Source/Runtime/Landscape/Classes/LandscapeProxy.h
@@ -119,7 +119,7 @@ enum class ELandscapeImportAlphamapType : uint8
 {
 	// Three layers blended 50/30/20 represented as 0.5, 0.3, and 0.2 in the alpha maps
 	// All alpha maps for blended layers total to 1.0
-	// This is the style used by UE4 internally for blended layers
+	// This is the style used by UE internally for blended layers
 	Additive,
 
 	// Three layers blended 50/30/20 represented as 0.5, 0.6, and 1.0 in the alpha maps
diff --git a/Engine/Source/Runtime/Landscape/Private/Landscape.cpp b/Engine/Source/Runtime/Landscape/Private/Landscape.cpp
--- a/Engine/Source/Runtime/Landscape/Private/Landscape.cpp
+++ b/Engine/Source/Runtime/Landscape/Private/Landscape.cpp
@@ -2505,7 +2505,7 @@ void ALandscapeProxy::GetSharedProperties(ALandscapeProxy* Landscape)
 
 		LandscapeGuid = Landscape->LandscapeGuid;
 
-		//@todo UE4 merge, landscape, this needs work
+		//@todo UE merge, landscape, this needs work
 		RootComponent->SetRelativeScale3D(Landscape->GetRootComponent()->GetComponentToWorld().GetScale3D());
 
 		//PrePivot = Landscape->PrePivot;
