Edit layers landscape : show the "Remove" context menu entry on null layer brushes so that brushes pointing to invalid actors can still be removed

diff --git a/Engine/Source/Editor/LandscapeEditor/Private/LandscapeEdMode.cpp b/Engine/Source/Editor/LandscapeEditor/Private/LandscapeEdMode.cpp
--- a/Engine/Source/Editor/LandscapeEditor/Private/LandscapeEdMode.cpp
+++ b/Engine/Source/Editor/LandscapeEditor/Private/LandscapeEdMode.cpp
@@ -3785,7 +3785,7 @@ void FEdModeLandscape::AddBrushToCurrentLayer(ALandscapeBlueprintBrushBase* InBr
 	RefreshDetailPanel();
 }
 
-void FEdModeLandscape::RemoveBrushFromCurrentLayer(ALandscapeBlueprintBrushBase* InBrush)
+void FEdModeLandscape::RemoveBrushFromCurrentLayer(int32 InBrushIndex)
 {
 	ALandscape* Landscape = GetLandscape();
 
@@ -3794,11 +3794,11 @@ void FEdModeLandscape::RemoveBrushFromCurrentLayer(ALandscapeBlueprintBrushBase*
 		return;
 	}
 
-	Landscape->RemoveBrushFromLayer(GetCurrentLayerIndex(), InBrush);
+	Landscape->RemoveBrushFromLayer(GetCurrentLayerIndex(), InBrushIndex);
 	RefreshDetailPanel();
 }
 
-ALandscapeBlueprintBrushBase* FEdModeLandscape::GetBrushForCurrentLayer(int8 InBrushIndex) const
+ALandscapeBlueprintBrushBase* FEdModeLandscape::GetBrushForCurrentLayer(int32 InBrushIndex) const
 {
 	if (ALandscape* Landscape = GetLandscape())
 	{
diff --git a/Engine/Source/Editor/LandscapeEditor/Private/LandscapeEdMode.h b/Engine/Source/Editor/LandscapeEditor/Private/LandscapeEdMode.h
--- a/Engine/Source/Editor/LandscapeEditor/Private/LandscapeEdMode.h
+++ b/Engine/Source/Editor/LandscapeEditor/Private/LandscapeEdMode.h
@@ -567,8 +567,8 @@ public:
 	bool CanEditLayer(FText* Reason = nullptr, FLandscapeLayer* InLayer = nullptr);
 
 	void AddBrushToCurrentLayer(class ALandscapeBlueprintBrushBase* InBrush);
-	void RemoveBrushFromCurrentLayer(class ALandscapeBlueprintBrushBase* InBrush);
-	class ALandscapeBlueprintBrushBase* GetBrushForCurrentLayer(int8 BrushIndex) const;
+	void RemoveBrushFromCurrentLayer(int32 InBrushIndex);
+	class ALandscapeBlueprintBrushBase* GetBrushForCurrentLayer(int32 InBrushIndex) const;
 	TArray<class ALandscapeBlueprintBrushBase*> GetBrushesForCurrentLayer();
 	
 	void ShowOnlySelectedBrush(class ALandscapeBlueprintBrushBase* InBrush);
diff --git a/Engine/Source/Editor/LandscapeEditor/Private/LandscapeEditorDetailCustomization_LayersBrushStack.cpp b/Engine/Source/Editor/LandscapeEditor/Private/LandscapeEditorDetailCustomization_LayersBrushStack.cpp
--- a/Engine/Source/Editor/LandscapeEditor/Private/LandscapeEditorDetailCustomization_LayersBrushStack.cpp
+++ b/Engine/Source/Editor/LandscapeEditor/Private/LandscapeEditorDetailCustomization_LayersBrushStack.cpp
@@ -207,43 +207,50 @@ END_SLATE_FUNCTION_BUILD_OPTIMIZATION
 
 TSharedPtr<SWidget> FLandscapeEditorCustomNodeBuilder_LayersBrushStack::OnBrushContextMenuOpening(int32 InBrushIndex)
 {
-	if (ALandscapeBlueprintBrushBase* CurrentBrush = GetBrush(InBrushIndex))
+	// Don't use GetBrush here as it will return nullptr if the layer contains a nullptr brush, and we still want to allow the Remove context menu when this happens : 
+	TArray<ALandscapeBlueprintBrushBase*> CurrentBrushes = GetBrushes();
+	if (CurrentBrushes.IsValidIndex(InBrushIndex))
 	{
+		ALandscapeBlueprintBrushBase* CurrentBrush = CurrentBrushes[InBrushIndex];
+
 		FMenuBuilder MenuBuilder(true, NULL);
 		MenuBuilder.BeginSection("LandscapeEditorBrushActions", LOCTEXT("LandscapeEditorBrushActions.Heading", "Brushes"));
 		{
 			TSharedRef<FLandscapeEditorCustomNodeBuilder_LayersBrushStack> SharedThis = AsShared();
 
-			// Show only selected
-			FUIAction ShowOnlySelectedAction = FUIAction(FExecuteAction::CreateLambda([SharedThis, CurrentBrush]
+			if (CurrentBrush != nullptr)
 			{
-				const FScopedTransaction Transaction(LOCTEXT("LandscapeBrushShowOnlySelectedTransaction", "Show Only Selected Brush"));
-				GetEditorMode()->ShowOnlySelectedBrush(CurrentBrush);
-			}));
-			MenuBuilder.AddMenuEntry(LOCTEXT("LandscapeBrushShowOnlySelected", "Show Only Selected"), LOCTEXT("LandscapeBrushShowOnlySelectedToolTip", "Hides all other brushes from the same layer"), FSlateIcon(), ShowOnlySelectedAction);
+				// Show only selected
+				FUIAction ShowOnlySelectedAction = FUIAction(FExecuteAction::CreateLambda([SharedThis, CurrentBrush]
+				{
+					const FScopedTransaction Transaction(LOCTEXT("LandscapeBrushShowOnlySelectedTransaction", "Show Only Selected Brush"));
+					GetEditorMode()->ShowOnlySelectedBrush(CurrentBrush);
+				}));
+				MenuBuilder.AddMenuEntry(LOCTEXT("LandscapeBrushShowOnlySelected", "Show Only Selected"), LOCTEXT("LandscapeBrushShowOnlySelectedToolTip", "Hides all other brushes from the same layer"), FSlateIcon(), ShowOnlySelectedAction);
 
-			// Hide selected
-			FUIAction ShowHideSelectedAction = FUIAction(FExecuteAction::CreateLambda([SharedThis, CurrentBrush]
-			{
-				SharedThis->OnToggleVisibility(CurrentBrush);
-			}));
+				// Hide selected
+				FUIAction ShowHideSelectedAction = FUIAction(FExecuteAction::CreateLambda([SharedThis, CurrentBrush]
+				{
+					SharedThis->OnToggleVisibility(CurrentBrush);
+				}));
 
-			FText MenuText = CurrentBrush->IsVisible() ? LOCTEXT("LandscapeBrushHideSelected", "Hide Selected") : LOCTEXT("LandscapeBrushShowSelected", "Show Selected");
-			MenuBuilder.AddMenuEntry(MenuText, LOCTEXT("LandscapeBrushToggleVisiblityToolTip", "Toggle brush visiblity"), FSlateIcon(), ShowHideSelectedAction);
+				FText MenuText = CurrentBrush->IsVisible() ? LOCTEXT("LandscapeBrushHideSelected", "Hide Selected") : LOCTEXT("LandscapeBrushShowSelected", "Show Selected");
+				MenuBuilder.AddMenuEntry(MenuText, LOCTEXT("LandscapeBrushToggleVisiblityToolTip", "Toggle brush visiblity"), FSlateIcon(), ShowHideSelectedAction);
 
-			MenuBuilder.AddMenuSeparator();
+				MenuBuilder.AddMenuSeparator();
 
-			// Duplicate Brush
-			FUIAction DuplicateAction = FUIAction(FExecuteAction::CreateLambda([SharedThis, CurrentBrush]
-			{
-				const FScopedTransaction Transaction(LOCTEXT("LandscapeBrushDuplicateTransaction", "Duplicate Brush"));
-				GetEditorMode()->DuplicateBrush(CurrentBrush);
-			}));
-			MenuBuilder.AddMenuEntry(LOCTEXT("LandscapeBrushDuplicate", "Duplicate"), LOCTEXT("LandscapeBrushDuplicateToolTip", "Duplicate brush"), FSlateIcon(), DuplicateAction);
+				// Duplicate Brush
+				FUIAction DuplicateAction = FUIAction(FExecuteAction::CreateLambda([SharedThis, CurrentBrush]
+				{
+					const FScopedTransaction Transaction(LOCTEXT("LandscapeBrushDuplicateTransaction", "Duplicate Brush"));
+					GetEditorMode()->DuplicateBrush(CurrentBrush);
+				}));
+				MenuBuilder.AddMenuEntry(LOCTEXT("LandscapeBrushDuplicate", "Duplicate"), LOCTEXT("LandscapeBrushDuplicateToolTip", "Duplicate brush"), FSlateIcon(), DuplicateAction);
+			}
 
 			// Add Brush
-			const TArray<ALandscapeBlueprintBrushBase*>& Brushes = GetEditorMode()->GetBrushList();
-			TArray<ALandscapeBlueprintBrushBase*> FilteredBrushes = Brushes.FilterByPredicate([](ALandscapeBlueprintBrushBase* Brush) { return Brush->GetOwningLandscape() == nullptr; });
+			const TArray<ALandscapeBlueprintBrushBase*>& AllBrushes = GetEditorMode()->GetBrushList();
+			TArray<ALandscapeBlueprintBrushBase*> FilteredBrushes = AllBrushes.FilterByPredicate([](ALandscapeBlueprintBrushBase* Brush) { return Brush->GetOwningLandscape() == nullptr; });
 			if (FilteredBrushes.Num())
 			{
 				MenuBuilder.AddSubMenu(
@@ -256,10 +263,10 @@ TSharedPtr<SWidget> FLandscapeEditorCustomNodeBuilder_LayersBrushStack::OnBrushC
 			}
 
 			// Remove Brush
-			FUIAction RemoveAction = FUIAction(FExecuteAction::CreateLambda([SharedThis, CurrentBrush]
+			FUIAction RemoveAction = FUIAction(FExecuteAction::CreateLambda([SharedThis, InBrushIndex]
 			{ 
 				const FScopedTransaction Transaction(LOCTEXT("LandscapeBrushRemoveTransaction", "Remove Brush"));
-				GetEditorMode()->RemoveBrushFromCurrentLayer(CurrentBrush);
+				GetEditorMode()->RemoveBrushFromCurrentLayer(InBrushIndex);
 			}));
 			MenuBuilder.AddMenuEntry(LOCTEXT("LandscapeBrushRemove", "Remove"), LOCTEXT("LandscapeBrushRemoveToolTip", "Remove brush from selected edit layer"), FSlateIcon(), RemoveAction);
 		}
@@ -409,6 +416,16 @@ ALandscapeBlueprintBrushBase* FLandscapeEditorCustomNodeBuilder_LayersBrushStack
 	return nullptr;
 }
 
+TArray<ALandscapeBlueprintBrushBase*> FLandscapeEditorCustomNodeBuilder_LayersBrushStack::GetBrushes() const
+{
+	if (FEdModeLandscape* LandscapeEdMode = GetEditorMode())
+	{
+		return LandscapeEdMode->GetBrushesForCurrentLayer();
+	}
+
+	return TArray<ALandscapeBlueprintBrushBase*>();
+}
+
 FReply FLandscapeEditorCustomNodeBuilder_LayersBrushStack::HandleDragDetected(const FGeometry& MyGeometry, const FPointerEvent& MouseEvent, int32 SlotIndex, SVerticalBox::FSlot* Slot)
 {
 	if (FEdModeLandscape* LandscapeEdMode = GetEditorMode())
diff --git a/Engine/Source/Editor/LandscapeEditor/Private/LandscapeEditorDetailCustomization_LayersBrushStack.h b/Engine/Source/Editor/LandscapeEditor/Private/LandscapeEditorDetailCustomization_LayersBrushStack.h
--- a/Engine/Source/Editor/LandscapeEditor/Private/LandscapeEditorDetailCustomization_LayersBrushStack.h
+++ b/Engine/Source/Editor/LandscapeEditor/Private/LandscapeEditorDetailCustomization_LayersBrushStack.h
@@ -77,6 +77,7 @@ protected:
 	FText GetBrushText(int32 InBrushIndex) const;
 	FSlateColor GetBrushTextColor(int32 InBrushIndex) const;
 	ALandscapeBlueprintBrushBase* GetBrush(int32 InBrushIndex) const;
+	TArray<ALandscapeBlueprintBrushBase*> GetBrushes() const;
 
 	// Drag/Drop handling
 	FReply HandleDragDetected(const FGeometry& MyGeometry, const FPointerEvent& MouseEvent, int32 SlotIndex, SVerticalBox::FSlot* Slot);
diff --git a/Engine/Source/Runtime/Landscape/Classes/Landscape.h b/Engine/Source/Runtime/Landscape/Classes/Landscape.h
--- a/Engine/Source/Runtime/Landscape/Classes/Landscape.h
+++ b/Engine/Source/Runtime/Landscape/Classes/Landscape.h
@@ -301,8 +301,10 @@ public:
 	LANDSCAPE_API void AddBrushToLayer(int32 InLayerIndex, class ALandscapeBlueprintBrushBase* InBrush);
 	LANDSCAPE_API void RemoveBrush(class ALandscapeBlueprintBrushBase* InBrush);
 	LANDSCAPE_API void RemoveBrushFromLayer(int32 InLayerIndex, class ALandscapeBlueprintBrushBase* InBrush);
+	LANDSCAPE_API void RemoveBrushFromLayer(int32 InLayerIndex, int32 InBrushIndex);
+	LANDSCAPE_API int32 GetBrushIndexForLayer(int32 InLayerIndex, class ALandscapeBlueprintBrushBase* InBrush);
 	LANDSCAPE_API bool ReorderLayerBrush(int32 InLayerIndex, int32 InStartingLayerBrushIndex, int32 InDestinationLayerBrushIndex);
-	LANDSCAPE_API class ALandscapeBlueprintBrushBase* GetBrushForLayer(int32 InLayerIndex, int8 BrushIndex) const;
+	LANDSCAPE_API class ALandscapeBlueprintBrushBase* GetBrushForLayer(int32 InLayerIndex, int32 BrushIndex) const;
 	LANDSCAPE_API TArray<class ALandscapeBlueprintBrushBase*> GetBrushesForLayer(int32 InLayerIndex) const;
 	LANDSCAPE_API void OnBlueprintBrushChanged();
 	LANDSCAPE_API void OnLayerInfoSplineFalloffModulationChanged(ULandscapeLayerInfoObject* InLayerInfo);
diff --git a/Engine/Source/Runtime/Landscape/Private/LandscapeEditLayers.cpp b/Engine/Source/Runtime/Landscape/Private/LandscapeEditLayers.cpp
--- a/Engine/Source/Runtime/Landscape/Private/LandscapeEditLayers.cpp
+++ b/Engine/Source/Runtime/Landscape/Private/LandscapeEditLayers.cpp
@@ -6614,6 +6614,33 @@ void ALandscape::RemoveBrush(ALandscapeBlueprintBrushBase* InBrush)
 }
 
 void ALandscape::RemoveBrushFromLayer(int32 InLayerIndex, ALandscapeBlueprintBrushBase* InBrush)
+{
+	int32 BrushIndex = GetBrushIndexForLayer(InLayerIndex, InBrush);
+	if (BrushIndex != INDEX_NONE)
+	{
+		RemoveBrushFromLayer(InLayerIndex, BrushIndex);
+	}
+}
+
+void ALandscape::RemoveBrushFromLayer(int32 InLayerIndex, int32 InBrushIndex)
+{
+	if (FLandscapeLayer* Layer = GetLayer(InLayerIndex))
+	{
+		if (Layer->Brushes.IsValidIndex(InBrushIndex))
+		{
+			Modify();
+			ALandscapeBlueprintBrushBase* Brush = Layer->Brushes[InBrushIndex].GetBrush();
+			Layer->Brushes.RemoveAt(InBrushIndex);
+			if (Brush != nullptr)
+			{
+				Brush->SetOwningLandscape(nullptr);
+			}
+			RequestLayersContentUpdateForceAll();
+		}
+	}
+}
+
+int32 ALandscape::GetBrushIndexForLayer(int32 InLayerIndex, ALandscapeBlueprintBrushBase* InBrush)
 {
 	if (FLandscapeLayer* Layer = GetLayer(InLayerIndex))
 	{
@@ -6621,16 +6648,15 @@ void ALandscape::RemoveBrushFromLayer(int32 InLayerIndex, ALandscapeBlueprintBru
 		{
 			if (Layer->Brushes[i].GetBrush() == InBrush)
 			{
-				Modify();
-				Layer->Brushes.RemoveAt(i);
-				InBrush->SetOwningLandscape(nullptr);
-				RequestLayersContentUpdateForceAll();
-				break;
+				return i;
 			}
 		}
 	}
+
+	return INDEX_NONE;
 }
 
+
 void ALandscape::OnBlueprintBrushChanged()
 {
 #if WITH_EDITORONLY_DATA
@@ -6666,7 +6692,7 @@ void ALandscape::OnLayerInfoSplineFalloffModulationChanged(ULandscapeLayerInfoOb
 	}
 }
 
-ALandscapeBlueprintBrushBase* ALandscape::GetBrushForLayer(int32 InLayerIndex, int8 InBrushIndex) const
+ALandscapeBlueprintBrushBase* ALandscape::GetBrushForLayer(int32 InLayerIndex, int32 InBrushIndex) const
 {
 	if (const FLandscapeLayer* Layer = GetLayer(InLayerIndex))
 	{
