Landscape: Highlight landscape proxy bounds in the New/Import landscape tool (based on WorldPartitionGridSize setting)

diff --git a/Engine/Source/Editor/LandscapeEditor/Private/LandscapeEdModeComponentTools.cpp b/Engine/Source/Editor/LandscapeEditor/Private/LandscapeEdModeComponentTools.cpp
--- a/Engine/Source/Editor/LandscapeEditor/Private/LandscapeEdModeComponentTools.cpp
+++ b/Engine/Source/Editor/LandscapeEditor/Private/LandscapeEdModeComponentTools.cpp
@@ -1945,11 +1945,11 @@ public:
 		if (EdMode->NewLandscapePreviewMode != ENewLandscapePreviewMode::None)
 		{
 			static const float        CornerSize = 0.33f;
-			static const FLinearColor CornerColour(1.0f, 1.0f, 0.5f);
-			static const FLinearColor EdgeColour(1.0f, 1.0f, 0.0f);
-			static const FLinearColor ComponentBorderColour(0.0f, 0.85f, 0.0f);
-			static const FLinearColor SectionBorderColour(0.0f, 0.4f, 0.0f);
-			static const FLinearColor InnerColour(0.0f, 0.25f, 0.0f);
+			static const FLinearColor CornerColor(1.0f, 1.0f, 0.5f);
+			static const FLinearColor EdgeColor(1.0f, 1.0f, 0.0f);
+			static const FLinearColor ComponentBorderColor(0.0f, 0.85f, 0.0f);
+			static const FLinearColor SectionBorderColor(0.0f, 0.4f, 0.0f);
+			static const FLinearColor InnerColor(0.0f, 0.25f, 0.0f);
 
 			const ELevelViewportType ViewportType = ((FEditorViewportClient*)Viewport->GetClient())->ViewportType;
 
@@ -1957,11 +1957,24 @@ public:
 			const int32 ComponentCountY = EdMode->UISettings->NewLandscape_ComponentCount.Y;
 			const int32 QuadsPerComponent = EdMode->UISettings->NewLandscape_SectionsPerComponent * EdMode->UISettings->NewLandscape_QuadsPerSection;
 			const float ComponentSize = QuadsPerComponent;
+			const int32 GridSize = EdMode->UISettings->WorldPartitionGridSize;
 			const FVector Offset = EdMode->UISettings->NewLandscape_Location + FTransform(EdMode->UISettings->NewLandscape_Rotation, FVector::ZeroVector, EdMode->UISettings->NewLandscape_Scale).TransformVector(FVector(-ComponentCountX * ComponentSize / 2, -ComponentCountY * ComponentSize / 2, 0));
 			const FTransform Transform = FTransform(EdMode->UISettings->NewLandscape_Rotation, Offset, EdMode->UISettings->NewLandscape_Scale);
 
 			if (EdMode->NewLandscapePreviewMode == ENewLandscapePreviewMode::ImportLandscape)
 			{
+				auto GetColor = [&](int32 ComponentIndex)
+				{
+					if (EdMode->IsGridBased() && (ComponentIndex % GridSize == 0))
+					{
+						return EdgeColor;
+					}
+					else
+					{
+						return ComponentBorderColor;
+					}
+				};
+
 				const TArray<uint16>& ImportHeights = EdMode->UISettings->GetImportLandscapeData();
 				if (ImportHeights.Num() != 0)
 				{
@@ -1995,48 +2008,48 @@ public:
 							if (ComponentX == 0)
 							{
 								PDI->SetHitProxy(new HNewLandscapeGrabHandleProxy(ELandscapeEdge::X_Negative));
-								PDI->DrawLine(Transform.TransformPosition(FVector(X0, Y0, Z00)), Transform.TransformPosition(FVector(X0, Y1, Z01)), ComponentBorderColour, SDPG_Foreground);
+								PDI->DrawLine(Transform.TransformPosition(FVector(X0, Y0, Z00)), Transform.TransformPosition(FVector(X0, Y1, Z01)), GetColor(ComponentX), SDPG_Foreground);
 								PDI->SetHitProxy(NULL);
 							}
 
 							if (ComponentX == ComponentCountX - 1)
 							{
 								PDI->SetHitProxy(new HNewLandscapeGrabHandleProxy(ELandscapeEdge::X_Positive));
-								PDI->DrawLine(Transform.TransformPosition(FVector(X1, Y0, Z10)), Transform.TransformPosition(FVector(X1, Y1, Z11)), ComponentBorderColour, SDPG_Foreground);
+								PDI->DrawLine(Transform.TransformPosition(FVector(X1, Y0, Z10)), Transform.TransformPosition(FVector(X1, Y1, Z11)), GetColor(ComponentX), SDPG_Foreground);
 								PDI->SetHitProxy(NULL);
 							}
 							else
 							{
-								PDI->DrawLine(Transform.TransformPosition(FVector(X1, Y0, Z10)), Transform.TransformPosition(FVector(X1, Y1, Z11)), ComponentBorderColour, SDPG_Foreground);
+								PDI->DrawLine(Transform.TransformPosition(FVector(X1, Y0, Z10)), Transform.TransformPosition(FVector(X1, Y1, Z11)), GetColor(ComponentX), SDPG_Foreground);
 							}
 
 							if (ComponentY == 0)
 							{
 								PDI->SetHitProxy(new HNewLandscapeGrabHandleProxy(ELandscapeEdge::Y_Negative));
-								PDI->DrawLine(Transform.TransformPosition(FVector(X0, Y0, Z00)), Transform.TransformPosition(FVector(X1, Y0, Z10)), ComponentBorderColour, SDPG_Foreground);
+								PDI->DrawLine(Transform.TransformPosition(FVector(X0, Y0, Z00)), Transform.TransformPosition(FVector(X1, Y0, Z10)), GetColor(ComponentY), SDPG_Foreground);
 								PDI->SetHitProxy(NULL);
 							}
 
 							if (ComponentY == ComponentCountY - 1)
 							{
 								PDI->SetHitProxy(new HNewLandscapeGrabHandleProxy(ELandscapeEdge::Y_Positive));
-								PDI->DrawLine(Transform.TransformPosition(FVector(X0, Y1, Z01)), Transform.TransformPosition(FVector(X1, Y1, Z11)), ComponentBorderColour, SDPG_Foreground);
+								PDI->DrawLine(Transform.TransformPosition(FVector(X0, Y1, Z01)), Transform.TransformPosition(FVector(X1, Y1, Z11)), GetColor(ComponentY), SDPG_Foreground);
 								PDI->SetHitProxy(NULL);
 							}
 							else
 							{
-								PDI->DrawLine(Transform.TransformPosition(FVector(X0, Y1, Z01)), Transform.TransformPosition(FVector(X1, Y1, Z11)), ComponentBorderColour, SDPG_Foreground);
+								PDI->DrawLine(Transform.TransformPosition(FVector(X0, Y1, Z01)), Transform.TransformPosition(FVector(X1, Y1, Z11)), GetColor(ComponentY), SDPG_Foreground);
 							}
 
 							// intra-component lines - too slow for big landscapes
 							/*
 							for (int32 x=1;x<QuadsPerComponent;x++)
 							{
-							PDI->DrawLine(Transform.TransformPosition(FVector(X0+x, Y0, FMath::Lerp(Z00,Z10,(float)x*InvQuadsPerComponent))), Transform.TransformPosition(FVector(X0+x, Y1, FMath::Lerp(Z01,Z11,(float)x*InvQuadsPerComponent))), ComponentBorderColour, SDPG_World);
+							PDI->DrawLine(Transform.TransformPosition(FVector(X0+x, Y0, FMath::Lerp(Z00,Z10,(float)x*InvQuadsPerComponent))), Transform.TransformPosition(FVector(X0+x, Y1, FMath::Lerp(Z01,Z11,(float)x*InvQuadsPerComponent))), ComponentBorderColor, SDPG_World);
 							}
 							for (int32 y=1;y<QuadsPerComponent;y++)
 							{
-							PDI->DrawLine(Transform.TransformPosition(FVector(X0, Y0+y, FMath::Lerp(Z00,Z01,(float)y*InvQuadsPerComponent))), Transform.TransformPosition(FVector(X1, Y0+y, FMath::Lerp(Z10,Z11,(float)y*InvQuadsPerComponent))), ComponentBorderColour, SDPG_World);
+							PDI->DrawLine(Transform.TransformPosition(FVector(X0, Y0+y, FMath::Lerp(Z00,Z01,(float)y*InvQuadsPerComponent))), Transform.TransformPosition(FVector(X1, Y0+y, FMath::Lerp(Z10,Z11,(float)y*InvQuadsPerComponent))), ComponentBorderColor, SDPG_World);
 							}
 							*/
 						}
@@ -2045,6 +2058,24 @@ public:
 			}
 			else //if (EdMode->NewLandscapePreviewMode == ENewLandscapePreviewMode::NewLandscape)
 			{
+				auto GetColor = [&](int32 QuadIndex)
+				{
+					if (EdMode->IsGridBased() && (QuadIndex % (GridSize * QuadsPerComponent) == 0))
+					{
+						return EdgeColor;
+					}
+					else if (QuadIndex % QuadsPerComponent == 0)
+					{
+						return ComponentBorderColor;
+					}
+					else if (QuadIndex % EdMode->UISettings->NewLandscape_QuadsPerSection == 0)
+					{
+						return SectionBorderColor;
+					}
+
+					return InnerColor;
+				};
+								
 				if (ViewportType == LVT_Perspective || ViewportType == LVT_OrthoXY || ViewportType == LVT_OrthoNegativeXY)
 				{
 					for (int32 x = 0; x <= ComponentCountX * QuadsPerComponent; x++)
@@ -2052,34 +2083,28 @@ public:
 						if (x == 0)
 						{
 							PDI->SetHitProxy(new HNewLandscapeGrabHandleProxy(ELandscapeEdge::X_Negative_Y_Negative));
-							PDI->DrawLine(Transform.TransformPosition(FVector(x, 0, 0)), Transform.TransformPosition(FVector(x, CornerSize * ComponentSize, 0)), CornerColour, SDPG_Foreground);
+							PDI->DrawLine(Transform.TransformPosition(FVector(x, 0, 0)), Transform.TransformPosition(FVector(x, CornerSize * ComponentSize, 0)), CornerColor, SDPG_Foreground);
 							PDI->SetHitProxy(new HNewLandscapeGrabHandleProxy(ELandscapeEdge::X_Negative));
-							PDI->DrawLine(Transform.TransformPosition(FVector(x, CornerSize * ComponentSize, 0)), Transform.TransformPosition(FVector(x, (ComponentCountY - CornerSize) * ComponentSize, 0)), EdgeColour, SDPG_Foreground);
+							PDI->DrawLine(Transform.TransformPosition(FVector(x, CornerSize * ComponentSize, 0)), Transform.TransformPosition(FVector(x, (ComponentCountY - CornerSize) * ComponentSize, 0)), EdgeColor, SDPG_Foreground);
 							PDI->SetHitProxy(new HNewLandscapeGrabHandleProxy(ELandscapeEdge::X_Negative_Y_Positive));
-							PDI->DrawLine(Transform.TransformPosition(FVector(x, (ComponentCountY - CornerSize) * ComponentSize, 0)), Transform.TransformPosition(FVector(x, ComponentCountY * ComponentSize, 0)), CornerColour, SDPG_Foreground);
+							PDI->DrawLine(Transform.TransformPosition(FVector(x, (ComponentCountY - CornerSize) * ComponentSize, 0)), Transform.TransformPosition(FVector(x, ComponentCountY * ComponentSize, 0)), CornerColor, SDPG_Foreground);
 							PDI->SetHitProxy(NULL);
 						}
 						else if (x == ComponentCountX * QuadsPerComponent)
 						{
 							PDI->SetHitProxy(new HNewLandscapeGrabHandleProxy(ELandscapeEdge::X_Positive_Y_Negative));
-							PDI->DrawLine(Transform.TransformPosition(FVector(x, 0, 0)), Transform.TransformPosition(FVector(x, CornerSize * ComponentSize, 0)), CornerColour, SDPG_Foreground);
+							PDI->DrawLine(Transform.TransformPosition(FVector(x, 0, 0)), Transform.TransformPosition(FVector(x, CornerSize * ComponentSize, 0)), CornerColor, SDPG_Foreground);
 							PDI->SetHitProxy(new HNewLandscapeGrabHandleProxy(ELandscapeEdge::X_Positive));
-							PDI->DrawLine(Transform.TransformPosition(FVector(x, CornerSize * ComponentSize, 0)), Transform.TransformPosition(FVector(x, (ComponentCountY - CornerSize) * ComponentSize, 0)), EdgeColour, SDPG_Foreground);
+							PDI->DrawLine(Transform.TransformPosition(FVector(x, CornerSize * ComponentSize, 0)), Transform.TransformPosition(FVector(x, (ComponentCountY - CornerSize) * ComponentSize, 0)), EdgeColor, SDPG_Foreground);
 							PDI->SetHitProxy(new HNewLandscapeGrabHandleProxy(ELandscapeEdge::X_Positive_Y_Positive));
-							PDI->DrawLine(Transform.TransformPosition(FVector(x, (ComponentCountY - CornerSize) * ComponentSize, 0)), Transform.TransformPosition(FVector(x, ComponentCountY * ComponentSize, 0)), CornerColour, SDPG_Foreground);
+							PDI->DrawLine(Transform.TransformPosition(FVector(x, (ComponentCountY - CornerSize) * ComponentSize, 0)), Transform.TransformPosition(FVector(x, ComponentCountY * ComponentSize, 0)), CornerColor, SDPG_Foreground);
 							PDI->SetHitProxy(NULL);
 						}
-						else if (x % QuadsPerComponent == 0)
-						{
-							PDI->DrawLine(Transform.TransformPosition(FVector(x, 0, 0)), Transform.TransformPosition(FVector(x, ComponentCountY * ComponentSize, 0)), ComponentBorderColour, SDPG_Foreground);
-						}
-						else if (x % EdMode->UISettings->NewLandscape_QuadsPerSection == 0)
-						{
-							PDI->DrawLine(Transform.TransformPosition(FVector(x, 0, 0)), Transform.TransformPosition(FVector(x, ComponentCountY * ComponentSize, 0)), SectionBorderColour, SDPG_Foreground);
-						}
 						else
 						{
-							PDI->DrawLine(Transform.TransformPosition(FVector(x, 0, 0)), Transform.TransformPosition(FVector(x, ComponentCountY * ComponentSize, 0)), InnerColour, SDPG_World);
+							FLinearColor CurrentColor = GetColor(x);
+							uint8 DepthPriority = CurrentColor == InnerColor ? SDPG_World : SDPG_Foreground;
+							PDI->DrawLine(Transform.TransformPosition(FVector(x, 0, 0)), Transform.TransformPosition(FVector(x, ComponentCountY * ComponentSize, 0)), GetColor(x), DepthPriority);
 						}
 					}
 				}
@@ -2087,8 +2112,8 @@ public:
 				{
 					// Don't allow dragging to resize in side-view
 					// and there's no point drawing the inner lines as only the outer is visible
-					PDI->DrawLine(Transform.TransformPosition(FVector(0, 0, 0)), Transform.TransformPosition(FVector(0, ComponentCountY * ComponentSize, 0)), EdgeColour, SDPG_World);
-					PDI->DrawLine(Transform.TransformPosition(FVector(ComponentCountX * QuadsPerComponent, 0, 0)), Transform.TransformPosition(FVector(ComponentCountX * QuadsPerComponent, ComponentCountY * ComponentSize, 0)), EdgeColour, SDPG_World);
+					PDI->DrawLine(Transform.TransformPosition(FVector(0, 0, 0)), Transform.TransformPosition(FVector(0, ComponentCountY * ComponentSize, 0)), EdgeColor, SDPG_World);
+					PDI->DrawLine(Transform.TransformPosition(FVector(ComponentCountX * QuadsPerComponent, 0, 0)), Transform.TransformPosition(FVector(ComponentCountX * QuadsPerComponent, ComponentCountY * ComponentSize, 0)), EdgeColor, SDPG_World);
 				}
 
 				if (ViewportType == LVT_Perspective || ViewportType == LVT_OrthoXY || ViewportType == LVT_OrthoNegativeXY)
@@ -2098,34 +2123,28 @@ public:
 						if (y == 0)
 						{
 							PDI->SetHitProxy(new HNewLandscapeGrabHandleProxy(ELandscapeEdge::X_Negative_Y_Negative));
-							PDI->DrawLine(Transform.TransformPosition(FVector(0, y, 0)), Transform.TransformPosition(FVector(CornerSize * ComponentSize, y, 0)), CornerColour, SDPG_Foreground);
+							PDI->DrawLine(Transform.TransformPosition(FVector(0, y, 0)), Transform.TransformPosition(FVector(CornerSize * ComponentSize, y, 0)), CornerColor, SDPG_Foreground);
 							PDI->SetHitProxy(new HNewLandscapeGrabHandleProxy(ELandscapeEdge::Y_Negative));
-							PDI->DrawLine(Transform.TransformPosition(FVector(CornerSize * ComponentSize, y, 0)), Transform.TransformPosition(FVector((ComponentCountX - CornerSize) * ComponentSize, y, 0)), EdgeColour, SDPG_Foreground);
+							PDI->DrawLine(Transform.TransformPosition(FVector(CornerSize * ComponentSize, y, 0)), Transform.TransformPosition(FVector((ComponentCountX - CornerSize) * ComponentSize, y, 0)), EdgeColor, SDPG_Foreground);
 							PDI->SetHitProxy(new HNewLandscapeGrabHandleProxy(ELandscapeEdge::X_Positive_Y_Negative));
-							PDI->DrawLine(Transform.TransformPosition(FVector((ComponentCountX - CornerSize) * ComponentSize, y, 0)), Transform.TransformPosition(FVector(ComponentCountX * ComponentSize, y, 0)), CornerColour, SDPG_Foreground);
+							PDI->DrawLine(Transform.TransformPosition(FVector((ComponentCountX - CornerSize) * ComponentSize, y, 0)), Transform.TransformPosition(FVector(ComponentCountX * ComponentSize, y, 0)), CornerColor, SDPG_Foreground);
 							PDI->SetHitProxy(NULL);
 						}
 						else if (y == ComponentCountY * QuadsPerComponent)
 						{
 							PDI->SetHitProxy(new HNewLandscapeGrabHandleProxy(ELandscapeEdge::X_Negative_Y_Positive));
-							PDI->DrawLine(Transform.TransformPosition(FVector(0, y, 0)), Transform.TransformPosition(FVector(CornerSize * ComponentSize, y, 0)), CornerColour, SDPG_Foreground);
+							PDI->DrawLine(Transform.TransformPosition(FVector(0, y, 0)), Transform.TransformPosition(FVector(CornerSize * ComponentSize, y, 0)), CornerColor, SDPG_Foreground);
 							PDI->SetHitProxy(new HNewLandscapeGrabHandleProxy(ELandscapeEdge::Y_Positive));
-							PDI->DrawLine(Transform.TransformPosition(FVector(CornerSize * ComponentSize, y, 0)), Transform.TransformPosition(FVector((ComponentCountX - CornerSize) * ComponentSize, y, 0)), EdgeColour, SDPG_Foreground);
+							PDI->DrawLine(Transform.TransformPosition(FVector(CornerSize * ComponentSize, y, 0)), Transform.TransformPosition(FVector((ComponentCountX - CornerSize) * ComponentSize, y, 0)), EdgeColor, SDPG_Foreground);
 							PDI->SetHitProxy(new HNewLandscapeGrabHandleProxy(ELandscapeEdge::X_Positive_Y_Positive));
-							PDI->DrawLine(Transform.TransformPosition(FVector((ComponentCountX - CornerSize) * ComponentSize, y, 0)), Transform.TransformPosition(FVector(ComponentCountX * ComponentSize, y, 0)), CornerColour, SDPG_Foreground);
+							PDI->DrawLine(Transform.TransformPosition(FVector((ComponentCountX - CornerSize) * ComponentSize, y, 0)), Transform.TransformPosition(FVector(ComponentCountX * ComponentSize, y, 0)), CornerColor, SDPG_Foreground);
 							PDI->SetHitProxy(NULL);
 						}
-						else if (y % QuadsPerComponent == 0)
-						{
-							PDI->DrawLine(Transform.TransformPosition(FVector(0, y, 0)), Transform.TransformPosition(FVector(ComponentCountX * ComponentSize, y, 0)), ComponentBorderColour, SDPG_Foreground);
-						}
-						else if (y % EdMode->UISettings->NewLandscape_QuadsPerSection == 0)
-						{
-							PDI->DrawLine(Transform.TransformPosition(FVector(0, y, 0)), Transform.TransformPosition(FVector(ComponentCountX * ComponentSize, y, 0)), SectionBorderColour, SDPG_Foreground);
-						}
 						else
 						{
-							PDI->DrawLine(Transform.TransformPosition(FVector(0, y, 0)), Transform.TransformPosition(FVector(ComponentCountX * ComponentSize, y, 0)), InnerColour, SDPG_World);
+							FLinearColor CurrentColor = GetColor(y);
+							uint8 DepthPriority = CurrentColor == InnerColor ? SDPG_World : SDPG_Foreground;
+							PDI->DrawLine(Transform.TransformPosition(FVector(0, y, 0)), Transform.TransformPosition(FVector(ComponentCountX * ComponentSize, y, 0)), GetColor(y), DepthPriority);
 						}
 					}
 				}
@@ -2133,8 +2152,8 @@ public:
 				{
 					// Don't allow dragging to resize in side-view
 					// and there's no point drawing the inner lines as only the outer is visible
-					PDI->DrawLine(Transform.TransformPosition(FVector(0, 0, 0)), Transform.TransformPosition(FVector(ComponentCountX * ComponentSize, 0, 0)), EdgeColour, SDPG_World);
-					PDI->DrawLine(Transform.TransformPosition(FVector(0, ComponentCountY * QuadsPerComponent, 0)), Transform.TransformPosition(FVector(ComponentCountX * ComponentSize, ComponentCountY * QuadsPerComponent, 0)), EdgeColour, SDPG_World);
+					PDI->DrawLine(Transform.TransformPosition(FVector(0, 0, 0)), Transform.TransformPosition(FVector(ComponentCountX * ComponentSize, 0, 0)), EdgeColor, SDPG_World);
+					PDI->DrawLine(Transform.TransformPosition(FVector(0, ComponentCountY * QuadsPerComponent, 0)), Transform.TransformPosition(FVector(ComponentCountX * ComponentSize, ComponentCountY * QuadsPerComponent, 0)), EdgeColor, SDPG_World);
 				}
 			}
 		}
@@ -2336,12 +2355,12 @@ public:
 		}
 			
 		static const float        CornerSize = 0.33f;
-		static const FLinearColor CornerColour(1.0f, 1.0f, 0.5f);
-		static const FLinearColor EdgeColour(1.0f, 1.0f, 0.0f);
-		static const FLinearColor ComponentBorderColour(0.0f, 0.85f, 0.0f);
-		static const FLinearColor SectionBorderColour(0.0f, 0.4f, 0.0f);
+		static const FLinearColor CornerColor(1.0f, 1.0f, 0.5f);
+		static const FLinearColor EdgeColor(1.0f, 1.0f, 0.0f);
+		static const FLinearColor ComponentBorderColor(0.0f, 0.85f, 0.0f);
+		static const FLinearColor SectionBorderColor(0.0f, 0.4f, 0.0f);
 		static const FLinearColor ComponentOutsideColor(0.0f, 0.0f, 0.85f);
-		static const FLinearColor InnerColour(0.0f, 0.25f, 0.0f);
+		static const FLinearColor InnerColor(0.0f, 0.25f, 0.0f);
 
 		const ELevelViewportType ViewportType = ((FEditorViewportClient*)Viewport->GetClient())->ViewportType;
 
@@ -2363,23 +2382,23 @@ public:
 				{
 					if (x == 0)
 					{
-						PDI->DrawLine(GizmoTransform.TransformPosition(FVector(x, 0, 0)), GizmoTransform.TransformPosition(FVector(x, CornerSize * ComponentSize, 0)), CornerColour, SDPG_Foreground);
-						PDI->DrawLine(GizmoTransform.TransformPosition(FVector(x, CornerSize * ComponentSize, 0)), GizmoTransform.TransformPosition(FVector(x, ImportHeight - (CornerSize * ComponentSize), 0)), EdgeColour, SDPG_Foreground);
-						PDI->DrawLine(GizmoTransform.TransformPosition(FVector(x, ImportHeight - (CornerSize * ComponentSize), 0)), GizmoTransform.TransformPosition(FVector(x, ImportHeight, 0)), CornerColour, SDPG_Foreground);
+						PDI->DrawLine(GizmoTransform.TransformPosition(FVector(x, 0, 0)), GizmoTransform.TransformPosition(FVector(x, CornerSize * ComponentSize, 0)), CornerColor, SDPG_Foreground);
+						PDI->DrawLine(GizmoTransform.TransformPosition(FVector(x, CornerSize * ComponentSize, 0)), GizmoTransform.TransformPosition(FVector(x, ImportHeight - (CornerSize * ComponentSize), 0)), EdgeColor, SDPG_Foreground);
+						PDI->DrawLine(GizmoTransform.TransformPosition(FVector(x, ImportHeight - (CornerSize * ComponentSize), 0)), GizmoTransform.TransformPosition(FVector(x, ImportHeight, 0)), CornerColor, SDPG_Foreground);
 					}
 					else if (x == (ImportWidth - 1))
 					{
-						PDI->DrawLine(GizmoTransform.TransformPosition(FVector(x, 0, 0)), GizmoTransform.TransformPosition(FVector(x, CornerSize * ComponentSize, 0)), CornerColour, SDPG_Foreground);
-						PDI->DrawLine(GizmoTransform.TransformPosition(FVector(x, CornerSize * ComponentSize, 0)), GizmoTransform.TransformPosition(FVector(x, ImportHeight - (CornerSize * ComponentSize), 0)), EdgeColour, SDPG_Foreground);
-						PDI->DrawLine(GizmoTransform.TransformPosition(FVector(x, ImportHeight - (CornerSize * ComponentSize), 0)), GizmoTransform.TransformPosition(FVector(x, ImportHeight, 0)), CornerColour, SDPG_Foreground);
+						PDI->DrawLine(GizmoTransform.TransformPosition(FVector(x, 0, 0)), GizmoTransform.TransformPosition(FVector(x, CornerSize * ComponentSize, 0)), CornerColor, SDPG_Foreground);
+						PDI->DrawLine(GizmoTransform.TransformPosition(FVector(x, CornerSize * ComponentSize, 0)), GizmoTransform.TransformPosition(FVector(x, ImportHeight - (CornerSize * ComponentSize), 0)), EdgeColor, SDPG_Foreground);
+						PDI->DrawLine(GizmoTransform.TransformPosition(FVector(x, ImportHeight - (CornerSize * ComponentSize), 0)), GizmoTransform.TransformPosition(FVector(x, ImportHeight, 0)), CornerColor, SDPG_Foreground);
 					}
 					else if ((x % ComponentSizeInt) == 0)
 					{
-						PDI->DrawLine(GizmoTransform.TransformPosition(FVector(x, 0, 0)), GizmoTransform.TransformPosition(FVector(x, ImportHeight, 0)), ComponentBorderColour, SDPG_Foreground);
+						PDI->DrawLine(GizmoTransform.TransformPosition(FVector(x, 0, 0)), GizmoTransform.TransformPosition(FVector(x, ImportHeight, 0)), ComponentBorderColor, SDPG_Foreground);
 					}
 					else if ((x % EdMode->UISettings->NewLandscape_QuadsPerSection) == 0)
 					{
-						PDI->DrawLine(GizmoTransform.TransformPosition(FVector(x, 0, 0)), GizmoTransform.TransformPosition(FVector(x, ImportHeight, 0)), SectionBorderColour, SDPG_Foreground);
+						PDI->DrawLine(GizmoTransform.TransformPosition(FVector(x, 0, 0)), GizmoTransform.TransformPosition(FVector(x, ImportHeight, 0)), SectionBorderColor, SDPG_Foreground);
 					}
 				}
 
@@ -2396,8 +2415,8 @@ public:
 			}
 			else
 			{
-				PDI->DrawLine(GizmoTransform.TransformPosition(FVector(0, 0, 0)), GizmoTransform.TransformPosition(FVector(0, ImportHeight, 0)), EdgeColour, SDPG_Foreground);
-				PDI->DrawLine(GizmoTransform.TransformPosition(FVector(ImportWidth, 0, 0)), GizmoTransform.TransformPosition(FVector(ImportWidth, ImportHeight, 0)), EdgeColour, SDPG_Foreground);
+				PDI->DrawLine(GizmoTransform.TransformPosition(FVector(0, 0, 0)), GizmoTransform.TransformPosition(FVector(0, ImportHeight, 0)), EdgeColor, SDPG_Foreground);
+				PDI->DrawLine(GizmoTransform.TransformPosition(FVector(ImportWidth, 0, 0)), GizmoTransform.TransformPosition(FVector(ImportWidth, ImportHeight, 0)), EdgeColor, SDPG_Foreground);
 			}
 
 			if (ViewportType == LVT_Perspective || ViewportType == LVT_OrthoXY || ViewportType == LVT_OrthoNegativeXY)
@@ -2406,23 +2425,23 @@ public:
 				{
 					if (y == 0)
 					{
-						PDI->DrawLine(GizmoTransform.TransformPosition(FVector(0, y, 0)), GizmoTransform.TransformPosition(FVector(CornerSize * ComponentSize, y, 0)), CornerColour, SDPG_Foreground);
-						PDI->DrawLine(GizmoTransform.TransformPosition(FVector(CornerSize * ComponentSize, y, 0)), GizmoTransform.TransformPosition(FVector(ImportWidth - (CornerSize * ComponentSize), y, 0)), EdgeColour, SDPG_Foreground);
-						PDI->DrawLine(GizmoTransform.TransformPosition(FVector(ImportWidth - (CornerSize * ComponentSize), y, 0)), GizmoTransform.TransformPosition(FVector(ImportWidth, y, 0)), CornerColour, SDPG_Foreground);
+						PDI->DrawLine(GizmoTransform.TransformPosition(FVector(0, y, 0)), GizmoTransform.TransformPosition(FVector(CornerSize * ComponentSize, y, 0)), CornerColor, SDPG_Foreground);
+						PDI->DrawLine(GizmoTransform.TransformPosition(FVector(CornerSize * ComponentSize, y, 0)), GizmoTransform.TransformPosition(FVector(ImportWidth - (CornerSize * ComponentSize), y, 0)), EdgeColor, SDPG_Foreground);
+						PDI->DrawLine(GizmoTransform.TransformPosition(FVector(ImportWidth - (CornerSize * ComponentSize), y, 0)), GizmoTransform.TransformPosition(FVector(ImportWidth, y, 0)), CornerColor, SDPG_Foreground);
 					}
 					else if (y == (ImportHeight - 1))
 					{
-						PDI->DrawLine(GizmoTransform.TransformPosition(FVector(0, y, 0)), GizmoTransform.TransformPosition(FVector(CornerSize * ComponentSize, y, 0)), CornerColour, SDPG_Foreground);
-						PDI->DrawLine(GizmoTransform.TransformPosition(FVector(CornerSize * ComponentSize, y, 0)), GizmoTransform.TransformPosition(FVector(ImportWidth - (CornerSize * ComponentSize), y, 0)), EdgeColour, SDPG_Foreground);
-						PDI->DrawLine(GizmoTransform.TransformPosition(FVector(ImportWidth - (CornerSize * ComponentSize), y, 0)), GizmoTransform.TransformPosition(FVector(ImportWidth, y, 0)), CornerColour, SDPG_Foreground);
+						PDI->DrawLine(GizmoTransform.TransformPosition(FVector(0, y, 0)), GizmoTransform.TransformPosition(FVector(CornerSize * ComponentSize, y, 0)), CornerColor, SDPG_Foreground);
+						PDI->DrawLine(GizmoTransform.TransformPosition(FVector(CornerSize * ComponentSize, y, 0)), GizmoTransform.TransformPosition(FVector(ImportWidth - (CornerSize * ComponentSize), y, 0)), EdgeColor, SDPG_Foreground);
+						PDI->DrawLine(GizmoTransform.TransformPosition(FVector(ImportWidth - (CornerSize * ComponentSize), y, 0)), GizmoTransform.TransformPosition(FVector(ImportWidth, y, 0)), CornerColor, SDPG_Foreground);
 					}
 					else if ((y % ComponentSizeInt) == 0)
 					{
-						PDI->DrawLine(GizmoTransform.TransformPosition(FVector(0, y, 0)), GizmoTransform.TransformPosition(FVector(ImportWidth, y, 0)), ComponentBorderColour, SDPG_Foreground);
+						PDI->DrawLine(GizmoTransform.TransformPosition(FVector(0, y, 0)), GizmoTransform.TransformPosition(FVector(ImportWidth, y, 0)), ComponentBorderColor, SDPG_Foreground);
 					}
 					else if ((y % EdMode->UISettings->NewLandscape_QuadsPerSection) == 0)
 					{
-						PDI->DrawLine(GizmoTransform.TransformPosition(FVector(0, y, 0)), GizmoTransform.TransformPosition(FVector(ImportWidth, y, 0)), SectionBorderColour, SDPG_Foreground);
+						PDI->DrawLine(GizmoTransform.TransformPosition(FVector(0, y, 0)), GizmoTransform.TransformPosition(FVector(ImportWidth, y, 0)), SectionBorderColor, SDPG_Foreground);
 					}
 				}
 
@@ -2440,8 +2459,8 @@ public:
 			else
 			{
 				// and there's no point drawing the inner lines as only the outer is visible
-				PDI->DrawLine(GizmoTransform.TransformPosition(FVector(0, 0, 0)), GizmoTransform.TransformPosition(FVector(ImportWidth, 0, 0)), EdgeColour, SDPG_Foreground);
-				PDI->DrawLine(GizmoTransform.TransformPosition(FVector(0, ImportHeight, 0)), GizmoTransform.TransformPosition(FVector(ImportWidth, ImportHeight, 0)), EdgeColour, SDPG_Foreground);
+				PDI->DrawLine(GizmoTransform.TransformPosition(FVector(0, 0, 0)), GizmoTransform.TransformPosition(FVector(ImportWidth, 0, 0)), EdgeColor, SDPG_Foreground);
+				PDI->DrawLine(GizmoTransform.TransformPosition(FVector(0, ImportHeight, 0)), GizmoTransform.TransformPosition(FVector(ImportWidth, ImportHeight, 0)), EdgeColor, SDPG_Foreground);
 			}
 		}
 	}
