The global mode manager should not be accessed in a commandlet environment.

# Moved access that was in the landscape editor module to the landscape editor mode, and added a check which should help better track down invalid access in the future, until we can clean up the global mode tools completely.

diff --git a/Engine/Source/Editor/LandscapeEditor/Private/LandscapeEdMode.cpp b/Engine/Source/Editor/LandscapeEditor/Private/LandscapeEdMode.cpp
--- a/Engine/Source/Editor/LandscapeEditor/Private/LandscapeEdMode.cpp
+++ b/Engine/Source/Editor/LandscapeEditor/Private/LandscapeEdMode.cpp
@@ -383,6 +383,7 @@ void FEdModeLandscape::Enter()
 
 	OnLevelActorDeletedDelegateHandle = GEngine->OnLevelActorDeleted().AddSP(this, &FEdModeLandscape::OnLevelActorRemoved);
 	OnLevelActorAddedDelegateHandle = GEngine->OnLevelActorAdded().AddSP(this, &FEdModeLandscape::OnLevelActorAdded);
+	PreSaveWorldHandle = FEditorDelegates::PreSaveWorldWithContext.AddSP(this, &FEdModeLandscape::OnPreSaveWorld);
 		
 	UpdateToolModes();
 
@@ -589,6 +590,7 @@ void FEdModeLandscape::Exit()
 		}
 	}
 
+	FEditorDelegates::PreSaveWorldWithContext.Remove(PreSaveWorldHandle);
 	GEngine->OnLevelActorDeleted().Remove(OnLevelActorDeletedDelegateHandle);
 	GEngine->OnLevelActorAdded().Remove(OnLevelActorAddedDelegateHandle);
 	
@@ -670,8 +672,14 @@ void FEdModeLandscape::Exit()
 	FEdMode::Exit();
 }
 
-void FEdModeLandscape::OnPreSaveWorld(const class UWorld* InWorld, FObjectPreSaveContext ObjectSaveContext)
+void FEdModeLandscape::OnPreSaveWorld(UWorld* InWorld, FObjectPreSaveContext ObjectSaveContext)
 {
+	// If the mode is pending deletion, don't run the presave routine.
+	if (IsPendingDeletion())
+	{
+		return;
+	}
+
 	// Avoid doing this during procedural saves to keep determinism and we don't want to do this on GameWorlds.
 	if (!InWorld->IsGameWorld() && !ObjectSaveContext.IsProceduralSave())
 	{
diff --git a/Engine/Source/Editor/LandscapeEditor/Private/LandscapeEdMode.h b/Engine/Source/Editor/LandscapeEditor/Private/LandscapeEdMode.h
--- a/Engine/Source/Editor/LandscapeEditor/Private/LandscapeEdMode.h
+++ b/Engine/Source/Editor/LandscapeEditor/Private/LandscapeEdMode.h
@@ -585,7 +585,7 @@ public:
 	DECLARE_EVENT(FEdModeLandscape, FTargetsListUpdated);
 	static FTargetsListUpdated TargetsListUpdated;
 
-	void OnPreSaveWorld(const class UWorld* InWorld, FObjectPreSaveContext ObjectSaveContext);
+	void OnPreSaveWorld(class UWorld* InWorld, FObjectPreSaveContext ObjectSaveContext);
 
 	/** Handle notification that visible levels may have changed and we should update the editable landscapes list */
 	void HandleLevelsChanged(bool ShouldExitMode);
@@ -636,6 +636,7 @@ private:
 
 	FDelegateHandle OnLevelActorDeletedDelegateHandle;
 	FDelegateHandle OnLevelActorAddedDelegateHandle;
+	FDelegateHandle PreSaveWorldHandle;
 	
 	/** Check if we are painting using the VREditor */
 	bool bIsPaintingInVR;
diff --git a/Engine/Source/Editor/LandscapeEditor/Private/LandscapeEditorModule.cpp b/Engine/Source/Editor/LandscapeEditor/Private/LandscapeEditorModule.cpp
--- a/Engine/Source/Editor/LandscapeEditor/Private/LandscapeEditorModule.cpp
+++ b/Engine/Source/Editor/LandscapeEditor/Private/LandscapeEditorModule.cpp
@@ -67,8 +67,6 @@ public:
 	{
 		FLandscapeEditorCommands::Register();
 
-		PreSaveWorldHandle = FEditorDelegates::PreSaveWorldWithContext.AddRaw(this, &FLandscapeEditorModule::OnPreSaveWorld);
-
 		// register the editor mode
 		FEditorModeRegistry::Get().RegisterMode<FEdModeLandscape>(
 			FBuiltinEditorModes::EM_Landscape,
@@ -143,8 +141,6 @@ public:
 	{
 		FLandscapeEditorCommands::Unregister();
 
-		FEditorDelegates::PreSaveWorldWithContext.Remove(PreSaveWorldHandle);
-
 		// unregister the editor mode
 		FEditorModeRegistry::Get().UnregisterMode(FBuiltinEditorModes::EM_Landscape);
 
@@ -294,15 +290,6 @@ public:
 		}
 	}
 
-	void OnPreSaveWorld( class UWorld* World, FObjectPreSaveContext ObjectSaveContext)
-	{
-		FEdModeLandscape* EdMode = (FEdModeLandscape*)GLevelEditorModeTools().GetActiveMode(FBuiltinEditorModes::EM_Landscape);
-		if (EdMode)
-		{
-			EdMode->OnPreSaveWorld(World, ObjectSaveContext);
-		}
-	}
-
 	virtual const TCHAR* GetHeightmapImportDialogTypeString() const override;
 	virtual const TCHAR* GetWeightmapImportDialogTypeString() const override;
 
@@ -315,7 +302,6 @@ public:
 	virtual TSharedPtr<FUICommandList> GetLandscapeLevelViewportCommandList() const override;
 		
 protected:
-	FDelegateHandle PreSaveWorldHandle;
 	TSharedPtr<FExtender> ViewportMenuExtender;
 	TSharedPtr<FUICommandList> GlobalUICommandList;
 	TArray<FRegisteredLandscapeHeightmapFileFormat> HeightmapFormats;
diff --git a/Engine/Source/Editor/UnrealEd/Private/UnrealEdGlobals.cpp b/Engine/Source/Editor/UnrealEd/Private/UnrealEdGlobals.cpp
--- a/Engine/Source/Editor/UnrealEd/Private/UnrealEdGlobals.cpp
+++ b/Engine/Source/Editor/UnrealEd/Private/UnrealEdGlobals.cpp
@@ -48,6 +48,8 @@ namespace Internal
 {
 	static TSharedPtr<FEditorModeTools>& GetGlobalModeManager()
 	{
+		checkf(!IsRunningCommandlet(), TEXT("The global mode manager should not be created or accessed in a commandlet environment. Check that your mode or module is not accessing the global mode tools or that scriptable features of modes have been moved to subsystems."));
+
 		static TSharedPtr<FEditorModeTools> EditorModeToolsSingleton = MakeShared<FEditorModeTools>();
 		return EditorModeToolsSingleton;
 	}
