Landscape: Fix Import crashes/validation inconsistencies

diff --git a/Engine/Source/Editor/LandscapeEditor/Private/LandscapeEditorDetailCustomization_NewLandscape.cpp b/Engine/Source/Editor/LandscapeEditor/Private/LandscapeEditorDetailCustomization_NewLandscape.cpp
--- a/Engine/Source/Editor/LandscapeEditor/Private/LandscapeEditorDetailCustomization_NewLandscape.cpp
+++ b/Engine/Source/Editor/LandscapeEditor/Private/LandscapeEditorDetailCustomization_NewLandscape.cpp
@@ -539,6 +539,7 @@ void FLandscapeEditorDetailCustomization_NewLandscape::CustomizeDetails(IDetailL
 			.Text(LOCTEXT("FitToData", "Fit To Data"))
 			.AddMetaData<FTagMetaData>(TEXT("ImportButton"))
 			.OnClicked(this, &FLandscapeEditorDetailCustomization_NewLandscape::OnFitImportDataButtonClicked)
+			.IsEnabled(this, &FLandscapeEditorDetailCustomization_NewLandscape::GetImportButtonIsEnabled)
 		]
 		+ SHorizontalBox::Slot()
 		.FillWidth(1)
@@ -961,15 +962,12 @@ bool FLandscapeEditorDetailCustomization_NewLandscape::GetImportButtonIsEnabled(
 	FEdModeLandscape* LandscapeEdMode = GetEditorMode();
 	if (LandscapeEdMode != nullptr)
 	{
-		bool bAllSourceFilePathsEmpty = true;
-		if (LandscapeEdMode->UISettings->ImportLandscape_HeightmapImportResult == ELandscapeImportResult::Error)
+		if (LandscapeEdMode->UISettings->ImportLandscape_HeightmapImportResult == ELandscapeImportResult::Error ||
+			LandscapeEdMode->UISettings->ImportLandscape_HeightmapFilename.IsEmpty() ||
+			LandscapeEdMode->UISettings->GetImportLandscapeData().IsEmpty())
 		{
 			return false;
 		}
-		else if (!LandscapeEdMode->UISettings->ImportLandscape_HeightmapFilename.IsEmpty())
-		{
-			bAllSourceFilePathsEmpty = false;
-		}
 
 		for (int32 i = 0; i < LandscapeEdMode->UISettings->ImportLandscape_Layers.Num(); ++i)
 		{
@@ -977,14 +975,11 @@ bool FLandscapeEditorDetailCustomization_NewLandscape::GetImportButtonIsEnabled(
 			{
 				return false;
 			}
-			else if (!LandscapeEdMode->UISettings->ImportLandscape_Layers[i].SourceFilePath.IsEmpty())
-			{
-				bAllSourceFilePathsEmpty = false;
-			}
 		}
 
-		return !bAllSourceFilePathsEmpty;
+		return true;
 	}
+	
 	return false;
 }
 
diff --git a/Engine/Source/Editor/LandscapeEditor/Private/LandscapeEditorObject.cpp b/Engine/Source/Editor/LandscapeEditor/Private/LandscapeEditorObject.cpp
--- a/Engine/Source/Editor/LandscapeEditor/Private/LandscapeEditorObject.cpp
+++ b/Engine/Source/Editor/LandscapeEditor/Private/LandscapeEditorObject.cpp
@@ -159,7 +159,8 @@ void ULandscapeEditorObject::PostEditChangeProperty(FPropertyChangedEvent& Prope
 
 	if (PropertyChangedEvent.MemberProperty == nullptr ||
 		PropertyChangedEvent.MemberProperty->GetFName() == GET_MEMBER_NAME_CHECKED(ULandscapeEditorObject, NewLandscape_Material) ||
-		PropertyChangedEvent.MemberProperty->GetFName() == GET_MEMBER_NAME_CHECKED(ULandscapeEditorObject, ImportLandscape_HeightmapFilename))
+		PropertyChangedEvent.MemberProperty->GetFName() == GET_MEMBER_NAME_CHECKED(ULandscapeEditorObject, ImportLandscape_HeightmapFilename) ||
+		PropertyChangedEvent.MemberProperty->GetFName() == GET_MEMBER_NAME_CHECKED(ULandscapeEditorObject, ImportLandscape_Layers))
 	{
 		RefreshImportLayersList();
 	}
@@ -592,43 +593,47 @@ void ULandscapeEditorObject::RefreshLayerImports()
 
 	for (FLandscapeImportLayer& UIImportLayer : ImportLandscape_Layers)
 	{
-		UIImportLayer.ErrorMessage = FText();
-		UIImportLayer.ImportResult = ELandscapeImportResult::Success;
+		RefreshLayerImport(UIImportLayer);
+	}
+}
 
-		if (UIImportLayer.LayerInfo != nullptr && !UIImportLayer.SourceFilePath.IsEmpty())
-		{
-			UIImportLayer.ImportResult = FLandscapeImportHelper::GetWeightmapImportDescriptor(UIImportLayer.SourceFilePath, UseSingleFileImport(), bFlipYAxis, UIImportLayer.LayerName, UIImportLayer.ImportDescriptor, UIImportLayer.ErrorMessage);
-			if (UIImportLayer.ImportResult == ELandscapeImportResult::Error)
-			{
-				continue;
-			}
+void ULandscapeEditorObject::RefreshLayerImport(FLandscapeImportLayer& ImportLayer)
+{
+	ImportLayer.ErrorMessage = FText();
+	ImportLayer.ImportResult = ELandscapeImportResult::Success;
 
-			if (ImportLandscape_Height != 0 || ImportLandscape_Width != 0)
+	if (!ImportLayer.SourceFilePath.IsEmpty())
+	{
+		if (!ImportLayer.LayerInfo)
+		{
+			ImportLayer.ImportResult = ELandscapeImportResult::Error;
+			ImportLayer.ErrorMessage = NSLOCTEXT("LandscapeEditor.NewLandscape", "Import_LayerInfoNotSet", "Can't import a layer file without a layer info");
+		}
+		else
+		{
+			ImportLayer.ImportResult = FLandscapeImportHelper::GetWeightmapImportDescriptor(ImportLayer.SourceFilePath, UseSingleFileImport(), bFlipYAxis, ImportLayer.LayerName, ImportLayer.ImportDescriptor, ImportLayer.ErrorMessage);
+			if (ImportLayer.ImportResult != ELandscapeImportResult::Error)
 			{
-				// Use same import index as Heightmap
-				int32 FoundIndex = INDEX_NONE;
-				for (int32 Index = 0; Index < UIImportLayer.ImportDescriptor.FileResolutions.Num(); ++Index)
+				if (ImportLandscape_Height != 0 || ImportLandscape_Width != 0)
 				{
-					if (UIImportLayer.ImportDescriptor.ImportResolutions[Index] == HeightmapImportDescriptor.ImportResolutions[HeightmapImportDescriptorIndex])
+					// Use same import index as Heightmap
+					int32 FoundIndex = INDEX_NONE;
+					for (int32 Index = 0; Index < ImportLayer.ImportDescriptor.FileResolutions.Num(); ++Index)
 					{
-						FoundIndex = Index;
-						break;
+						if (ImportLayer.ImportDescriptor.ImportResolutions[Index] == HeightmapImportDescriptor.ImportResolutions[HeightmapImportDescriptorIndex])
+						{
+							FoundIndex = Index;
+							break;
+						}
 					}
-				}
 
-				if (FoundIndex == INDEX_NONE)
-				{
-					UIImportLayer.ImportResult = ELandscapeImportResult::Error;
-					UIImportLayer.ErrorMessage = NSLOCTEXT("LandscapeEditor.ImportLandscape", "Import_WeightHeightResolutionMismatch", "Weightmap import resolution isn't same as Heightmap resolution.");
+					if (FoundIndex == INDEX_NONE)
+					{
+						ImportLayer.ImportResult = ELandscapeImportResult::Error;
+						ImportLayer.ErrorMessage = NSLOCTEXT("LandscapeEditor.ImportLandscape", "Import_WeightHeightResolutionMismatch", "Weightmap import resolution isn't same as Heightmap resolution.");
+					}
 				}
 			}
-			else // We don't have a valid Heightmap import so use the Layer Import resolution
-			{
-				HeightmapImportDescriptorIndex = 0;
-				ImportLandscape_Width = UIImportLayer.ImportDescriptor.ImportResolutions[HeightmapImportDescriptorIndex].Width;
-				ImportLandscape_Height = UIImportLayer.ImportDescriptor.ImportResolutions[HeightmapImportDescriptorIndex].Height;
-				ChooseBestComponentSizeForImport();
-			}			
 		}
 	}
 }
@@ -640,6 +645,7 @@ void ULandscapeEditorObject::OnChangeImportLandscapeResolution(int32 DescriptorI
 	ImportLandscape_Width = HeightmapImportDescriptor.ImportResolutions[HeightmapImportDescriptorIndex].Width;
 	ImportLandscape_Height = HeightmapImportDescriptor.ImportResolutions[HeightmapImportDescriptorIndex].Height;
 	ClearImportLandscapeData();
+	ImportLandscapeData();
 	ChooseBestComponentSizeForImport();
 }
 
@@ -875,27 +881,7 @@ void ULandscapeEditorObject::RefreshImportLayersList(bool bRefreshFromTarget)
 			NewImportLayer.LayerInfo = LayerInfoObjs[i];
 		}
 
-		if (!NewImportLayer.SourceFilePath.IsEmpty())
-		{
-			if (!NewImportLayer.LayerInfo)
-			{
-				NewImportLayer.ImportResult = ELandscapeImportResult::Error;
-				NewImportLayer.ErrorMessage = NSLOCTEXT("LandscapeEditor.NewLandscape", "Import_LayerInfoNotSet", "Can't import a layer file without a layer info");
-			}
-			else
-			{
-				NewImportLayer.ImportResult = FLandscapeImportHelper::GetWeightmapImportDescriptor(NewImportLayer.SourceFilePath, UseSingleFileImport(), bFlipYAxis, NewImportLayer.LayerName, NewImportLayer.ImportDescriptor, NewImportLayer.ErrorMessage);
-				if (NewImportLayer.ImportResult != ELandscapeImportResult::Error)
-				{
-					int32 DescriptorIndex = NewImportLayer.ImportDescriptor.FindDescriptorIndex(ImportLandscape_Width, ImportLandscape_Height);
-					if (DescriptorIndex == INDEX_NONE)
-					{
-						NewImportLayer.ImportResult = ELandscapeImportResult::Error;
-						NewImportLayer.ErrorMessage = NSLOCTEXT("LandscapeEditor.NewLandscape", "Import_LayerSizeMismatch", "Size of the layer import size does not match size of heightmap import size");
-					}
-				}
-			}
-		}
+		RefreshLayerImport(NewImportLayer);
 
 		ImportLandscape_Layers.Add(MoveTemp(NewImportLayer));
 	}
diff --git a/Engine/Source/Editor/LandscapeEditor/Public/LandscapeEditorObject.h b/Engine/Source/Editor/LandscapeEditor/Public/LandscapeEditorObject.h
--- a/Engine/Source/Editor/LandscapeEditor/Public/LandscapeEditorObject.h
+++ b/Engine/Source/Editor/LandscapeEditor/Public/LandscapeEditorObject.h
@@ -703,6 +703,7 @@ public:
 	void RefreshImports();
 	void OnImportWeightmapFilenameChanged() { RefreshLayerImports(); }
 	void RefreshLayerImports();
+	void RefreshLayerImport(FLandscapeImportLayer& ImportLayer);
 
 	int32 ClampLandscapeSize(int32 InComponentsCount) const
 	{
