Fix VT update to happen with rest of landscape painting tools (smooth, flatten, noise)

diff --git a/Engine/Source/Editor/LandscapeEditor/Private/LandscapeEdModePaintTools.cpp b/Engine/Source/Editor/LandscapeEditor/Private/LandscapeEdModePaintTools.cpp
--- a/Engine/Source/Editor/LandscapeEditor/Private/LandscapeEdModePaintTools.cpp
+++ b/Engine/Source/Editor/LandscapeEditor/Private/LandscapeEdModePaintTools.cpp
@@ -761,6 +761,12 @@ public:
 		}
 
 		LayerDataCache.Write(X1, Y1, X2, Y2, Data, UISettings->PaintingRestriction);
+
+		if (ToolTarget::TargetType == ELandscapeToolTargetType::Weightmap)
+		{
+			// Dirty any runtime virtual textures that our landscape components write to.
+			DirtyRuntimeVirtualTextureForLandscapeArea(this->LandscapeInfo, X1, Y1, X2, Y2);
+		}
 	}
 };
 
@@ -1043,6 +1049,12 @@ public:
 		}
 
 		LayerDataCache.Write(X1, Y1, X2, Y2, Data, UISettings->PaintingRestriction);
+
+		if (ToolTarget::TargetType == ELandscapeToolTargetType::Weightmap)
+		{
+			// Dirty any runtime virtual textures that our landscape components write to.
+			DirtyRuntimeVirtualTextureForLandscapeArea(this->LandscapeInfo, X1, Y1, X2, Y2);
+		}
 	}
 };
 
@@ -1273,6 +1285,12 @@ public:
 
 		this->Cache.SetCachedData(X1, Y1, X2, Y2, Data, UISettings->PaintingRestriction);
 		this->Cache.Flush();
+
+		if (ToolTarget::TargetType == ELandscapeToolTargetType::Weightmap)
+		{
+			// Dirty any runtime virtual textures that our landscape components write to.
+			DirtyRuntimeVirtualTextureForLandscapeArea(this->LandscapeInfo, X1, Y1, X2, Y2);
+		}
 	}
 };
 
