Fixed crash in landscape resampling method : regions in ResampleData are specified in quads, not vertices

diff --git a/Engine/Source/Editor/LandscapeEditor/Private/LandscapeEdMode.cpp b/Engine/Source/Editor/LandscapeEditor/Private/LandscapeEdMode.cpp
--- a/Engine/Source/Editor/LandscapeEditor/Private/LandscapeEdMode.cpp
+++ b/Engine/Source/Editor/LandscapeEditor/Private/LandscapeEdMode.cpp
@@ -3260,7 +3260,7 @@ ALandscape* FEdModeLandscape::ChangeComponentSetting(int32 NumComponentsX, int32
 				{
 					if (bResample)
 					{
-						OutHeightData.AddZeroed(OldVertsX * OldVertsY * sizeof(uint16));
+						OutHeightData.AddZeroed(OldVertsX * OldVertsY);
 
 						// GetHeightData alters its args, so make temp copies to avoid screwing things up
 						int32 TMinX = OldMinX, TMinY = OldMinY, TMaxX = OldMaxX, TMaxY = OldMaxY;
@@ -3268,8 +3268,8 @@ ALandscape* FEdModeLandscape::ChangeComponentSetting(int32 NumComponentsX, int32
 
 						TArray<uint16> ResampledHeightData;
 
-						FIntRect SrcRegion(0, 0, OldVertsX, OldVertsY);
-						FIntRect DestRegion(0, 0, NewVertsX, NewVertsY);
+						FIntRect SrcRegion(0, 0, OldVertsX - 1, OldVertsY - 1);
+						FIntRect DestRegion(0, 0, NewVertsX - 1, NewVertsY - 1);
 						FLandscapeConfigHelper::ResampleData(OutHeightData, ResampledHeightData, SrcRegion, DestRegion);
 						OutHeightData = MoveTemp(ResampledHeightData);
 
@@ -3279,7 +3279,7 @@ ALandscape* FEdModeLandscape::ChangeComponentSetting(int32 NumComponentsX, int32
 							if (LayerSettings.LayerInfoObj != NULL)
 							{
 								auto ImportLayerInfo = new(OutImportMaterialLayerInfos) FLandscapeImportLayerInfo(LayerSettings);
-								ImportLayerInfo->LayerData.AddZeroed(OldVertsX * OldVertsY * sizeof(uint8));
+								ImportLayerInfo->LayerData.AddZeroed(OldVertsX * OldVertsY);
 
 								TMinX = OldMinX; TMinY = OldMinY; TMaxX = OldMaxX; TMaxY = OldMaxY;
 								LandscapeEdit.GetWeightData(LayerSettings.LayerInfoObj, TMinX, TMinY, TMaxX, TMaxY, ImportLayerInfo->LayerData.GetData(), 0);
