Fixed invalid tooltips in landscape

#rnx

diff --git a/Engine/Source/Editor/LandscapeEditor/Private/LandscapeEditorDetails.cpp b/Engine/Source/Editor/LandscapeEditor/Private/LandscapeEditorDetails.cpp
--- a/Engine/Source/Editor/LandscapeEditor/Private/LandscapeEditorDetails.cpp
+++ b/Engine/Source/Editor/LandscapeEditor/Private/LandscapeEditorDetails.cpp
@@ -340,7 +340,7 @@ TSharedRef<SWidget> FLandscapeEditorDetails::GetToolSelector()
 		if (LandscapeEdMode->CurrentToolMode->ToolModeName == "ToolMode_Sculpt")
 		{
 			MenuBuilder.BeginSection(NAME_None, LOCTEXT("SculptToolsTitle", "Sculpting Tools"));
-			MenuBuilder.AddToolButton(NameToCommandMap.FindChecked("Tool_Sculpt"), NAME_None, LOCTEXT("Tool.Sculpt", "Sculpt"), LOCTEXT("Tool.Sculpt.Tooltip", "Sculpt height data.\nCtrl+Click to Raise, Ctrl+Shift+Click to lower"));
+			MenuBuilder.AddToolButton(NameToCommandMap.FindChecked("Tool_Sculpt"), NAME_None, LOCTEXT("Tool.Sculpt", "Sculpt"), LOCTEXT("Tool.Sculpt.Tooltip", "Sculpt height data.\nClick to Raise, Shift+Click to lower"));
 			
 			if (LandscapeEdMode->CanHaveLandscapeLayersContent())
 			{
@@ -373,7 +373,7 @@ TSharedRef<SWidget> FLandscapeEditorDetails::GetToolSelector()
 		if (LandscapeEdMode->CurrentToolMode->ToolModeName == "ToolMode_Paint")
 		{
 			MenuBuilder.BeginSection(NAME_None, LOCTEXT("PaintToolsTitle", "Paint Tools"));
-			MenuBuilder.AddToolButton(NameToCommandMap.FindChecked("Tool_Paint"), NAME_None, LOCTEXT("Tool.Paint", "Paint"), LOCTEXT("Tool.Paint.Tooltip", "Paints weight data.\nCtrl+Click to paint, Ctrl+Shift+Click to erase"));
+			MenuBuilder.AddToolButton(NameToCommandMap.FindChecked("Tool_Paint"), NAME_None, LOCTEXT("Tool.Paint", "Paint"), LOCTEXT("Tool.Paint.Tooltip", "Paints weight data.\nClick to paint, Shift+Click to erase"));
 			MenuBuilder.AddToolButton(NameToCommandMap.FindChecked("Tool_Smooth"), NAME_None, LOCTEXT("Tool.Smooth", "Smooth"), LOCTEXT("Tool.Smooth.Tooltip", "Smooths heightmaps or blend layers"));
 			MenuBuilder.AddToolButton(NameToCommandMap.FindChecked("Tool_Flatten"), NAME_None, LOCTEXT("Tool.Flatten", "Flatten"), LOCTEXT("Tool.Flatten.Tooltip", "Flattens an area of heightmap or blend layer"));
 			MenuBuilder.AddToolButton(NameToCommandMap.FindChecked("Tool_Noise"), NAME_None, LOCTEXT("Tool.Noise", "Noise"), LOCTEXT("Tool.Noise.Tooltip", "Adds noise to the heightmap or blend layer"));
