Touch up of landscape mode styling.

diff --git a/Engine/Source/Editor/LandscapeEditor/LandscapeEditor.Build.cs b/Engine/Source/Editor/LandscapeEditor/LandscapeEditor.Build.cs
--- a/Engine/Source/Editor/LandscapeEditor/LandscapeEditor.Build.cs
+++ b/Engine/Source/Editor/LandscapeEditor/LandscapeEditor.Build.cs
@@ -26,7 +26,8 @@ public class LandscapeEditor : ModuleRules
 				"ImageWrapper",
                 "EditorWidgets",
                 "Foliage",
-				"ToolMenus"
+				"ToolMenus",
+				"ToolWidgets",
 			}
 			);
 
diff --git a/Engine/Source/Editor/LandscapeEditor/Private/LandscapeEdMode.h b/Engine/Source/Editor/LandscapeEditor/Private/LandscapeEdMode.h
--- a/Engine/Source/Editor/LandscapeEditor/Private/LandscapeEdMode.h
+++ b/Engine/Source/Editor/LandscapeEditor/Private/LandscapeEdMode.h
@@ -242,24 +242,18 @@ struct FGizmoHistory
 	}
 };
 
-namespace ENewLandscapePreviewMode
+enum class ENewLandscapePreviewMode : uint8
 {
-	enum Type
-	{
-		None,
-		NewLandscape,
-		ImportLandscape,
-	};
-}
+	None,
+	NewLandscape,
+	ImportLandscape,
+};
 
-namespace EImportExportMode
+enum class EImportExportMode : uint8
 {
-	enum Type
-	{
-		Import,
-		Export,
-	};
-}
+	Import,
+	Export,
+};
 
 enum class ELandscapeEditingState : uint8
 {
@@ -294,8 +288,8 @@ public:
 	// UI setting for additional UI Tools
 	int32 CurrentBrushSetIndex;
 
-	ENewLandscapePreviewMode::Type NewLandscapePreviewMode;
-	EImportExportMode::Type ImportExportMode;
+	ENewLandscapePreviewMode NewLandscapePreviewMode;
+	EImportExportMode ImportExportMode;
 
 	TWeakObjectPtr<ALandscapeGizmoActiveActor> CurrentGizmoActor;
 	// UI callbacks for copy/paste tool
diff --git a/Engine/Source/Editor/LandscapeEditor/Private/LandscapeEdModeComponentTools.cpp b/Engine/Source/Editor/LandscapeEditor/Private/LandscapeEdModeComponentTools.cpp
--- a/Engine/Source/Editor/LandscapeEditor/Private/LandscapeEdModeComponentTools.cpp
+++ b/Engine/Source/Editor/LandscapeEditor/Private/LandscapeEdModeComponentTools.cpp
@@ -1789,7 +1789,7 @@ class FLandscapeToolNewLandscape : public FLandscapeTool
 {
 public:
 	FEdModeLandscape* EdMode;
-	ENewLandscapePreviewMode::Type NewLandscapePreviewMode;
+	ENewLandscapePreviewMode NewLandscapePreviewMode;
 	ELandscapeEdge::Type DraggingEdge;
 	float DraggingEdge_Remainder;
 
diff --git a/Engine/Source/Editor/LandscapeEditor/Private/LandscapeEditorDetailCustomization_ImportExport.cpp b/Engine/Source/Editor/LandscapeEditor/Private/LandscapeEditorDetailCustomization_ImportExport.cpp
--- a/Engine/Source/Editor/LandscapeEditor/Private/LandscapeEditorDetailCustomization_ImportExport.cpp
+++ b/Engine/Source/Editor/LandscapeEditor/Private/LandscapeEditorDetailCustomization_ImportExport.cpp
@@ -256,14 +256,14 @@ void FLandscapeEditorDetailCustomization_ImportExport::OnHeightmapSelectedCheckS
 	}
 }
 
-ECheckBoxState FLandscapeEditorDetailCustomization_ImportExport::ModeIsChecked(EImportExportMode::Type Value)
+ECheckBoxState FLandscapeEditorDetailCustomization_ImportExport::ModeIsChecked(EImportExportMode Value)
 {
 	FEdModeLandscape* LandscapeEdMode = GetEditorMode();
 	check(LandscapeEdMode != nullptr);
 	return LandscapeEdMode->ImportExportMode == Value ? ECheckBoxState::Checked : ECheckBoxState::Unchecked;
 }
 
-void FLandscapeEditorDetailCustomization_ImportExport::OnModeChanged(ECheckBoxState NewCheckedState, EImportExportMode::Type Value)
+void FLandscapeEditorDetailCustomization_ImportExport::OnModeChanged(ECheckBoxState NewCheckedState, EImportExportMode Value)
 {
 	if (NewCheckedState == ECheckBoxState::Checked)
 	{
diff --git a/Engine/Source/Editor/LandscapeEditor/Private/LandscapeEditorDetailCustomization_ImportExport.h b/Engine/Source/Editor/LandscapeEditor/Private/LandscapeEditorDetailCustomization_ImportExport.h
--- a/Engine/Source/Editor/LandscapeEditor/Private/LandscapeEditorDetailCustomization_ImportExport.h
+++ b/Engine/Source/Editor/LandscapeEditor/Private/LandscapeEditorDetailCustomization_ImportExport.h
@@ -36,8 +36,8 @@ private:
 	static bool IsHeightmapEnabled();
 	static ECheckBoxState GetHeightmapSelectedCheckState();
 	static void OnHeightmapSelectedCheckStateChanged(ECheckBoxState CheckState);
-	static ECheckBoxState ModeIsChecked(EImportExportMode::Type Value);
-	static void OnModeChanged(ECheckBoxState NewCheckedState, EImportExportMode::Type Value);
+	static ECheckBoxState ModeIsChecked(EImportExportMode Value);
+	static void OnModeChanged(ECheckBoxState NewCheckedState, EImportExportMode Value);
 
 	static EVisibility GetImportResultErrorVisibility(TSharedRef<IPropertyHandle> PropertyHandle_ImportResult);
 	static FSlateColor GetImportResultErrorColor(TSharedRef<IPropertyHandle> PropertyHandle_ImportResult);
diff --git a/Engine/Source/Editor/LandscapeEditor/Private/LandscapeEditorDetailCustomization_NewLandscape.cpp b/Engine/Source/Editor/LandscapeEditor/Private/LandscapeEditorDetailCustomization_NewLandscape.cpp
--- a/Engine/Source/Editor/LandscapeEditor/Private/LandscapeEditorDetailCustomization_NewLandscape.cpp
+++ b/Engine/Source/Editor/LandscapeEditor/Private/LandscapeEditorDetailCustomization_NewLandscape.cpp
@@ -44,6 +44,8 @@
 #include "Settings/EditorExperimentalSettings.h"
 #include "Editor.h"
 #include "LandscapeSubsystem.h"
+#include "SPrimaryButton.h"
+#include "Widgets/Input/SSegmentedControl.h"
 
 #define LOCTEXT_NAMESPACE "LandscapeEditor.NewLandscape"
 
@@ -65,30 +67,31 @@ void FLandscapeEditorDetailCustomization_NewLandscape::CustomizeDetails(IDetailL
 
 	NewLandscapeCategory.AddCustomRow(FText::GetEmpty())
 	[
-		SNew(SUniformGridPanel)
-		.SlotPadding(FMargin(10, 2))
-		+ SUniformGridPanel::Slot(0, 0)
-		[
-			SNew(SCheckBox)
-			.Style(FEditorStyle::Get(), "RadioButton")
-			.IsChecked(this, &FLandscapeEditorDetailCustomization_NewLandscape::NewLandscapeModeIsChecked, ENewLandscapePreviewMode::NewLandscape)
-			.OnCheckStateChanged(this, &FLandscapeEditorDetailCustomization_NewLandscape::OnNewLandscapeModeChanged, ENewLandscapePreviewMode::NewLandscape)
-			[
-				SNew(STextBlock)
-				.Text(LOCTEXT("NewLandscape", "Create New"))
-			]
-		]
-		+ SUniformGridPanel::Slot(1, 0)
+		SNew(SBox)
+		.Padding(2.0f)
+		.HAlign(HAlign_Center)
 		[
-			SNew(SCheckBox)
-			.Style(FEditorStyle::Get(), "RadioButton")
-			.IsChecked(this, &FLandscapeEditorDetailCustomization_NewLandscape::NewLandscapeModeIsChecked, ENewLandscapePreviewMode::ImportLandscape)
-			.OnCheckStateChanged(this, &FLandscapeEditorDetailCustomization_NewLandscape::OnNewLandscapeModeChanged, ENewLandscapePreviewMode::ImportLandscape)
-			[
-				SNew(STextBlock)
-				.Text(LOCTEXT("ImportLandscape", "Import from File"))
-			]
+			SNew(SSegmentedControl<ENewLandscapePreviewMode>)
+			.Value_Lambda([this]()
+			{
+				FEdModeLandscape* LandscapeEdMode = GetEditorMode();
+				return LandscapeEdMode ? LandscapeEdMode->NewLandscapePreviewMode : ENewLandscapePreviewMode::NewLandscape;
+
+			})
+			.OnValueChanged_Lambda([this](ENewLandscapePreviewMode Mode)
+			{
+				FEdModeLandscape* LandscapeEdMode = GetEditorMode();
+				if (LandscapeEdMode != nullptr)
+				{
+					LandscapeEdMode->NewLandscapePreviewMode = Mode;
+				}
+			})
+			+SSegmentedControl<ENewLandscapePreviewMode>::Slot(ENewLandscapePreviewMode::NewLandscape)
+			.Text(LOCTEXT("NewLandscape", "Create New"))
+			+ SSegmentedControl<ENewLandscapePreviewMode>::Slot(ENewLandscapePreviewMode::ImportLandscape)
+			.Text(LOCTEXT("ImportLandscape", "Import from File"))
 		]
+		
 	];
 
 	TSharedRef<IPropertyHandle> PropertyHandle_CanHaveLayersContent = DetailBuilder.GetProperty(GET_MEMBER_NAME_CHECKED(ULandscapeEditorObject, bCanHaveLayersContent));
@@ -514,10 +517,13 @@ void FLandscapeEditorDetailCustomization_NewLandscape::CustomizeDetails(IDetailL
 		.Text(this, &FLandscapeEditorDetailCustomization_NewLandscape::GetTotalComponentCount)
 	];
 
-	NewLandscapeCategory.AddCustomRow(FText::GetEmpty())
+	NewLandscapeCategory.AddCustomRow(FText::GetEmpty()).WholeRowContent()
+	.VAlign(VAlign_Center)
+	.HAlign(HAlign_Right)
 	[
 		SNew(SHorizontalBox)
 		+ SHorizontalBox::Slot()
+		.Padding(4,0)
 		.AutoWidth()
 		[
 			SNew(SButton)
@@ -527,6 +533,7 @@ void FLandscapeEditorDetailCustomization_NewLandscape::CustomizeDetails(IDetailL
 			.OnClicked(this, &FLandscapeEditorDetailCustomization_NewLandscape::OnFillWorldButtonClicked)
 		]
 		+ SHorizontalBox::Slot()
+		.Padding(4, 0)
 		.AutoWidth()
 		[
 			SNew(SButton)
@@ -537,19 +544,17 @@ void FLandscapeEditorDetailCustomization_NewLandscape::CustomizeDetails(IDetailL
 			.IsEnabled(this, &FLandscapeEditorDetailCustomization_NewLandscape::GetImportButtonIsEnabled)
 		]
 		+ SHorizontalBox::Slot()
-		.FillWidth(1)
-		//[
-		//]
-		+ SHorizontalBox::Slot()
+		.Padding(4, 0)
 		.AutoWidth()
 		[
-			SNew(SButton)
+			SNew(SPrimaryButton)
 			.Visibility_Static(&GetVisibilityOnlyInNewLandscapeMode, ENewLandscapePreviewMode::NewLandscape)
 			.Text(LOCTEXT("Create", "Create"))
 			.AddMetaData<FTutorialMetaData>(FTutorialMetaData(TEXT("CreateButton"), TEXT("LevelEditorToolBox")))
 			.OnClicked(this, &FLandscapeEditorDetailCustomization_NewLandscape::OnCreateButtonClicked)
 		]
 		+ SHorizontalBox::Slot()
+		.Padding(4, 0)
 		.AutoWidth()
 		[
 			SNew(SButton)
@@ -768,7 +773,7 @@ FText FLandscapeEditorDetailCustomization_NewLandscape::GetTotalComponentCount()
 }
 
 
-EVisibility FLandscapeEditorDetailCustomization_NewLandscape::GetVisibilityOnlyInNewLandscapeMode(ENewLandscapePreviewMode::Type value)
+EVisibility FLandscapeEditorDetailCustomization_NewLandscape::GetVisibilityOnlyInNewLandscapeMode(ENewLandscapePreviewMode value)
 {
 	FEdModeLandscape* LandscapeEdMode = GetEditorMode();
 	if (LandscapeEdMode != nullptr)
@@ -781,35 +786,6 @@ EVisibility FLandscapeEditorDetailCustomization_NewLandscape::GetVisibilityOnlyI
 	return EVisibility::Collapsed;
 }
 
-ECheckBoxState FLandscapeEditorDetailCustomization_NewLandscape::NewLandscapeModeIsChecked(ENewLandscapePreviewMode::Type value) const
-{
-	FEdModeLandscape* LandscapeEdMode = GetEditorMode();
-	if (LandscapeEdMode != nullptr)
-	{
-		if (LandscapeEdMode->NewLandscapePreviewMode == value)
-		{
-			return ECheckBoxState::Checked;
-		}
-	}
-	return ECheckBoxState::Unchecked;
-}
-
-void FLandscapeEditorDetailCustomization_NewLandscape::OnNewLandscapeModeChanged(ECheckBoxState NewCheckedState, ENewLandscapePreviewMode::Type value)
-{
-	if (NewCheckedState == ECheckBoxState::Checked)
-	{
-		FEdModeLandscape* LandscapeEdMode = GetEditorMode();
-		if (LandscapeEdMode != nullptr)
-		{
-			LandscapeEdMode->NewLandscapePreviewMode = value;
-
-			if (value == ENewLandscapePreviewMode::ImportLandscape)
-			{
-				LandscapeEdMode->NewLandscapePreviewMode = ENewLandscapePreviewMode::ImportLandscape;
-			}
-		}
-	}
-}
 
 FReply FLandscapeEditorDetailCustomization_NewLandscape::OnCreateButtonClicked()
 {
diff --git a/Engine/Source/Editor/LandscapeEditor/Private/LandscapeEditorDetailCustomization_NewLandscape.h b/Engine/Source/Editor/LandscapeEditor/Private/LandscapeEditorDetailCustomization_NewLandscape.h
--- a/Engine/Source/Editor/LandscapeEditor/Private/LandscapeEditorDetailCustomization_NewLandscape.h
+++ b/Engine/Source/Editor/LandscapeEditor/Private/LandscapeEditorDetailCustomization_NewLandscape.h
@@ -16,6 +16,8 @@
 
 class IDetailLayoutBuilder;
 
+enum class ENewLandscapePreviewMode : uint8;
+
 /**
  * Slate widgets customizer for the "New Landscape" tool
  */
@@ -58,9 +60,7 @@ public:
 	FReply OnCreateButtonClicked();
 	FReply OnFillWorldButtonClicked();
 
-	static EVisibility GetVisibilityOnlyInNewLandscapeMode(ENewLandscapePreviewMode::Type value);
-	ECheckBoxState NewLandscapeModeIsChecked(ENewLandscapePreviewMode::Type value) const;
-	void OnNewLandscapeModeChanged(ECheckBoxState NewCheckedState, ENewLandscapePreviewMode::Type value);
+	static EVisibility GetVisibilityOnlyInNewLandscapeMode(ENewLandscapePreviewMode value);
 
 	// Import
 	static EVisibility GetHeightmapErrorVisibility(TSharedRef<IPropertyHandle> PropertyHandle_HeightmapImportResult);
diff --git a/Engine/Source/Editor/LandscapeEditor/Private/LandscapeEditorDetails.cpp b/Engine/Source/Editor/LandscapeEditor/Private/LandscapeEditorDetails.cpp
--- a/Engine/Source/Editor/LandscapeEditor/Private/LandscapeEditorDetails.cpp
+++ b/Engine/Source/Editor/LandscapeEditor/Private/LandscapeEditorDetails.cpp
@@ -113,7 +113,6 @@ void FLandscapeEditorDetails::CustomizeDetails(IDetailLayoutBuilder& DetailBuild
 	LandscapeEditorCategory.AddCustomRow(FText::GetEmpty())
 	.Visibility(TAttribute<EVisibility>::Create(TAttribute<EVisibility>::FGetter::CreateSP(this, &FLandscapeEditorDetails::GetBrushSelectorVisibility)))
 	.NameContent()
-	.HAlign(HAlign_Right)
 	[
 		SNew(STextBlock)
 		.TextStyle( &FCoreStyle::Get().GetWidgetStyle<FTextBlockStyle>( "SmallText" ) )
@@ -128,7 +127,6 @@ void FLandscapeEditorDetails::CustomizeDetails(IDetailLayoutBuilder& DetailBuild
 	LandscapeEditorCategory.AddCustomRow(FText::GetEmpty())
 	.Visibility(TAttribute<EVisibility>::Create(TAttribute<EVisibility>::FGetter::CreateSP(this, &FLandscapeEditorDetails::GetBrushFalloffSelectorVisibility)))
 	.NameContent()
-	.HAlign(HAlign_Right)
 	[
 		SNew(STextBlock)
 		.TextStyle( &FCoreStyle::Get().GetWidgetStyle<FTextBlockStyle>( "SmallText" ) )
diff --git a/Engine/Source/Editor/LandscapeEditor/Private/NewLandscapeUtils.cpp b/Engine/Source/Editor/LandscapeEditor/Private/NewLandscapeUtils.cpp
--- a/Engine/Source/Editor/LandscapeEditor/Private/NewLandscapeUtils.cpp
+++ b/Engine/Source/Editor/LandscapeEditor/Private/NewLandscapeUtils.cpp
@@ -61,7 +61,7 @@ void FNewLandscapeUtils::ImportLandscapeData( ULandscapeEditorObject* UISettings
 	}
 }
 
-TOptional< TArray< FLandscapeImportLayerInfo > > FNewLandscapeUtils::CreateImportLayersInfo( ULandscapeEditorObject* UISettings, int32 NewLandscapePreviewMode )
+TOptional< TArray< FLandscapeImportLayerInfo > > FNewLandscapeUtils::CreateImportLayersInfo( ULandscapeEditorObject* UISettings, ENewLandscapePreviewMode NewLandscapePreviewMode )
 {
 	TArray<FLandscapeImportLayerInfo> ImportLayers;
 	ELandscapeImportResult Result = ELandscapeImportResult::Success;
@@ -82,7 +82,7 @@ TOptional< TArray< FLandscapeImportLayerInfo > > FNewLandscapeUtils::CreateImpor
 	return MoveTemp(ImportLayers);
 }
 
-TArray<uint16> FNewLandscapeUtils::ComputeHeightData(ULandscapeEditorObject* UISettings, TArray< FLandscapeImportLayerInfo >& ImportLayers, int32 NewLandscapePreviewMode)
+TArray<uint16> FNewLandscapeUtils::ComputeHeightData(ULandscapeEditorObject* UISettings, TArray< FLandscapeImportLayerInfo >& ImportLayers, ENewLandscapePreviewMode NewLandscapePreviewMode)
 {
 	TArray<uint16> HeightData;
 	if (NewLandscapePreviewMode == ENewLandscapePreviewMode::NewLandscape)
diff --git a/Engine/Source/Editor/LandscapeEditor/Public/NewLandscapeUtils.h b/Engine/Source/Editor/LandscapeEditor/Public/NewLandscapeUtils.h
--- a/Engine/Source/Editor/LandscapeEditor/Public/NewLandscapeUtils.h
+++ b/Engine/Source/Editor/LandscapeEditor/Public/NewLandscapeUtils.h
@@ -9,6 +9,8 @@ struct FLandscapeFileResolution;
 struct FLandscapeImportLayerInfo;
 class ULandscapeEditorObject;
 
+enum class ENewLandscapePreviewMode : uint8;
+
 // This class has been replaced by FLandscapeImportHelper & ULandscapeEditorObject methods.
 // @todo_ow: deprecate (probably just delete since it was added for Datasmith which was updated to the new import code)
 class LANDSCAPEEDITOR_API FNewLandscapeUtils
@@ -16,6 +18,6 @@ class LANDSCAPEEDITOR_API FNewLandscapeUtils
 public:
 	static void ChooseBestComponentSizeForImport( ULandscapeEditorObject* UISettings );
 	static void ImportLandscapeData( ULandscapeEditorObject* UISettings, TArray< FLandscapeFileResolution >& ImportResolutions );
-	static TOptional< TArray< FLandscapeImportLayerInfo > > CreateImportLayersInfo( ULandscapeEditorObject* UISettings, int32 NewLandscapePreviewMode );
-	static TArray< uint16 > ComputeHeightData( ULandscapeEditorObject* UISettings, TArray< FLandscapeImportLayerInfo >& ImportLayers, int32 NewLandscapePreviewMode );
+	static TOptional< TArray< FLandscapeImportLayerInfo > > CreateImportLayersInfo( ULandscapeEditorObject* UISettings, ENewLandscapePreviewMode NewLandscapePreviewMode );
+	static TArray<uint16> ComputeHeightData( ULandscapeEditorObject* UISettings, TArray< FLandscapeImportLayerInfo >& ImportLayers, ENewLandscapePreviewMode NewLandscapePreviewMode );
 };
