Crash in import panel for new landscape

diff --git a/Engine/Source/Editor/LandscapeEditor/Private/LandscapeEditorDetailCustomization_ImportLayers.cpp b/Engine/Source/Editor/LandscapeEditor/Private/LandscapeEditorDetailCustomization_ImportLayers.cpp
--- a/Engine/Source/Editor/LandscapeEditor/Private/LandscapeEditorDetailCustomization_ImportLayers.cpp
+++ b/Engine/Source/Editor/LandscapeEditor/Private/LandscapeEditorDetailCustomization_ImportLayers.cpp
@@ -417,9 +417,9 @@ FSlateColor FLandscapeEditorStructCustomization_FLandscapeImportLayer::GetErrorC
 {
 	ELandscapeImportResult WeightmapImportResult;
 	FPropertyAccess::Result Result = PropertyHandle_ImportResult->GetValue((uint8&)WeightmapImportResult);
-	check(Result == FPropertyAccess::Success);
 
-	if (Result == FPropertyAccess::MultipleValues)
+	if (Result == FPropertyAccess::Fail ||
+		Result == FPropertyAccess::MultipleValues)
 	{
 		return FCoreStyle::Get().GetColor("ErrorReporting.BackgroundColor");
 	}
