Import landscape : put the thumbnails on the left side of the paint layer names to prevent misalignment

diff --git a/Engine/Source/Editor/LandscapeEditor/Private/LandscapeEditorDetailCustomization_ImportLayers.cpp b/Engine/Source/Editor/LandscapeEditor/Private/LandscapeEditorDetailCustomization_ImportLayers.cpp
--- a/Engine/Source/Editor/LandscapeEditor/Private/LandscapeEditorDetailCustomization_ImportLayers.cpp
+++ b/Engine/Source/Editor/LandscapeEditor/Private/LandscapeEditorDetailCustomization_ImportLayers.cpp
@@ -132,21 +132,21 @@ void FLandscapeEditorStructCustomization_FLandscapeImportLayer::CustomizeChildre
 			.ToolTipText_Static(&FLandscapeEditorStructCustomization_FLandscapeImportLayer::GetImportLayerSelectedToolTip, PropertyHandle_Selected, PropertyHandle_LayerInfo)
 		]
 		+ SHorizontalBox::Slot()
-		.FillWidth(1)
+		.AutoWidth()
 		.VAlign(VAlign_Center)
 		.Padding(FMargin(2))
 		[
-			SNew(STextBlock)
-			.Font(StructCustomizationUtils.GetRegularFont())
-			.Text(LayerNameText)
+			SNew(SLandscapeAssetThumbnail, ThumbnailMIC, StructCustomizationUtils.GetThumbnailPool().ToSharedRef())
+			.ThumbnailSize(FIntPoint(48, 48))
 		]
 		+ SHorizontalBox::Slot()
-		.AutoWidth()
+		.FillWidth(1)
 		.VAlign(VAlign_Center)
 		.Padding(FMargin(2))
 		[
-			SNew(SLandscapeAssetThumbnail, ThumbnailMIC, StructCustomizationUtils.GetThumbnailPool().ToSharedRef())
-			.ThumbnailSize(FIntPoint(48, 48))
+			SNew(STextBlock)
+			.Font(StructCustomizationUtils.GetRegularFont())
+			.Text(LayerNameText)
 		]
 	]
 	.ValueContent()
