Fix spacing/padding issues on Landscape Mode Spline Details Customization

diff --git a/Engine/Source/Editor/LandscapeEditor/Private/LandscapeSplineDetails.cpp b/Engine/Source/Editor/LandscapeEditor/Private/LandscapeSplineDetails.cpp
--- a/Engine/Source/Editor/LandscapeEditor/Private/LandscapeSplineDetails.cpp
+++ b/Engine/Source/Editor/LandscapeEditor/Private/LandscapeSplineDetails.cpp
@@ -50,19 +50,23 @@ void FLandscapeSplineDetails::CustomizeDetails(IDetailLayoutBuilder& DetailBuild
 			.Text(LOCTEXT("SelectAll", "Select all connected:"))
 		]
 		+ SHorizontalBox::Slot()
+		.Padding(0, 2, 0, 2)
 		.FillWidth(1)
 		[
 			SNew(SButton)
 			.Text(LOCTEXT("ControlPoints", "Control Points"))
 			.HAlign(HAlign_Center)
+			.VAlign(VAlign_Center)
 			.OnClicked(this, &FLandscapeSplineDetails::OnSelectConnectedControlPointsButtonClicked)
 		]
 		+ SHorizontalBox::Slot()
+		.Padding(0, 2, 0, 2)
 		.FillWidth(1)
 		[
 			SNew(SButton)
 			.Text(LOCTEXT("Segments", "Segments"))
 			.HAlign(HAlign_Center)
+			.VAlign(VAlign_Center)
 			.OnClicked(this, &FLandscapeSplineDetails::OnSelectConnectedSegmentsButtonClicked)
 		]
 	];
@@ -71,7 +75,6 @@ void FLandscapeSplineDetails::CustomizeDetails(IDetailLayoutBuilder& DetailBuild
 	[
 		SNew(SHorizontalBox)
 		+ SHorizontalBox::Slot()
-		.Padding(0, 0, 2, 0)
 		.VAlign(VAlign_Center)
 		.FillWidth(1)
 		[
@@ -86,7 +89,6 @@ void FLandscapeSplineDetails::CustomizeDetails(IDetailLayoutBuilder& DetailBuild
 	[
 		SNew(SHorizontalBox)
 		+ SHorizontalBox::Slot()
-		.Padding(0, 0, 2, 0)
 		.VAlign(VAlign_Center)
 		.FillWidth(1)
 		[
