Fix crash when switching to the Landscape mode and back to Select mode

diff --git a/Engine/Source/Editor/LandscapeEditor/Private/LandscapeEditorDetailCustomization_NewLandscape.cpp b/Engine/Source/Editor/LandscapeEditor/Private/LandscapeEditorDetailCustomization_NewLandscape.cpp
--- a/Engine/Source/Editor/LandscapeEditor/Private/LandscapeEditorDetailCustomization_NewLandscape.cpp
+++ b/Engine/Source/Editor/LandscapeEditor/Private/LandscapeEditorDetailCustomization_NewLandscape.cpp
@@ -101,7 +101,16 @@ void FLandscapeEditorDetailCustomization_NewLandscape::CustomizeDetails(IDetailL
 	TSharedRef<IPropertyHandle> PropertyHandle_FlipYAxis = DetailBuilder.GetProperty(GET_MEMBER_NAME_CHECKED(ULandscapeEditorObject, bFlipYAxis));
 	NewLandscapeCategory.AddProperty(PropertyHandle_FlipYAxis).Visibility(MakeAttributeLambda([]()
 	{ 
-		return ((GetEditorMode()->NewLandscapePreviewMode == ENewLandscapePreviewMode::ImportLandscape) && !GetEditorMode()->UseSingleFileImport()) ? EVisibility::Visible : EVisibility::Collapsed; 
+		FEdModeLandscape* LandscapeEdMode = GetEditorMode();
+		if (LandscapeEdMode != nullptr)
+		{
+			if ((LandscapeEdMode->NewLandscapePreviewMode == ENewLandscapePreviewMode::ImportLandscape) && !LandscapeEdMode->UseSingleFileImport())
+			{
+				return EVisibility::Visible;
+			}
+		}
+
+		return EVisibility::Collapsed; 
 	}));
 	
 	PropertyHandle_FlipYAxis->SetOnPropertyValueChanged(FSimpleDelegate::CreateLambda([this]() { OnImportHeightmapFilenameChanged(); }));
