Landscape Editing - Prevent continuous nav mesh updates when moving the mouse if the Paint - Flatten tool is active

In MouseMove(), MeshComponent->SetRelativeLocation() would trigger the nav mesh update, as the component was bNavigationRelevant

diff --git a/Engine/Source/Editor/LandscapeEditor/Private/LandscapeEdModePaintTools.cpp b/Engine/Source/Editor/LandscapeEditor/Private/LandscapeEdModePaintTools.cpp
--- a/Engine/Source/Editor/LandscapeEditor/Private/LandscapeEdModePaintTools.cpp
+++ b/Engine/Source/Editor/LandscapeEditor/Private/LandscapeEdModePaintTools.cpp
@@ -1159,6 +1159,7 @@ public:
 		ALandscapeProxy* LandscapeProxy = this->EdMode->CurrentToolTarget.LandscapeInfo->GetLandscapeProxy();
 		MeshComponent = NewObject<UStaticMeshComponent>(LandscapeProxy, NAME_None, RF_Transient);
 		MeshComponent->SetStaticMesh(PlaneMesh);
+		MeshComponent->SetCanEverAffectNavigation(false);
 		MeshComponent->AttachToComponent(LandscapeProxy->GetRootComponent(), FAttachmentTransformRules::KeepRelativeTransform);
 		MeshComponent->RegisterComponent();
 
