FLandscapeToolFlatten - Only create preview component when sculpting
* Flatten preview is never shown when painting

diff --git a/Engine/Source/Editor/LandscapeEditor/Private/LandscapeEdModePaintTools.cpp b/Engine/Source/Editor/LandscapeEditor/Private/LandscapeEdModePaintTools.cpp
--- a/Engine/Source/Editor/LandscapeEditor/Private/LandscapeEdModePaintTools.cpp
+++ b/Engine/Source/Editor/LandscapeEditor/Private/LandscapeEdModePaintTools.cpp
@@ -1062,8 +1062,8 @@ template<class ToolTarget>
 class FLandscapeToolFlatten : public FLandscapeToolPaintBase < ToolTarget, FLandscapeToolStrokeFlatten<ToolTarget> >
 {
 protected:
-	UStaticMesh* PlaneMesh;
-	UStaticMeshComponent* MeshComponent;
+	UStaticMesh* HeightmapFlattenPlaneMesh;
+	UStaticMeshComponent* HeightmapFlattenPreviewComponent;
 	bool CanToolBeActivatedNextTick;
 	bool CanToolBeActivatedValue;
 	float EyeDropperFlattenTargetValue;
@@ -1071,13 +1071,13 @@ protected:
 public:
 	FLandscapeToolFlatten(FEdModeLandscape* InEdMode)
 		: FLandscapeToolPaintBase<ToolTarget, FLandscapeToolStrokeFlatten<ToolTarget>>(InEdMode)
-		, PlaneMesh(LoadObject<UStaticMesh>(NULL, TEXT("/Engine/EditorLandscapeResources/FlattenPlaneMesh.FlattenPlaneMesh")))
-		, MeshComponent(NULL)
+		, HeightmapFlattenPlaneMesh(LoadObject<UStaticMesh>(nullptr, TEXT("/Engine/EditorLandscapeResources/FlattenPlaneMesh.FlattenPlaneMesh")))
+		, HeightmapFlattenPreviewComponent(nullptr)
 		, CanToolBeActivatedNextTick(false)
 		, CanToolBeActivatedValue(false)
 		, EyeDropperFlattenTargetValue(0.0f)
 	{
-		check(PlaneMesh);
+		check(HeightmapFlattenPlaneMesh);
 	}
 
 	virtual bool GetCursor(EMouseCursor::Type& OutCursor) const override
@@ -1099,8 +1099,8 @@ public:
 
 	virtual void AddReferencedObjects(FReferenceCollector& Collector) override
 	{
-		Collector.AddReferencedObject(PlaneMesh);
-		Collector.AddReferencedObject(MeshComponent);
+		Collector.AddReferencedObject(HeightmapFlattenPlaneMesh);
+		Collector.AddReferencedObject(HeightmapFlattenPreviewComponent);
 	}
 
 	virtual const TCHAR* GetToolName() override { return TEXT("Flatten"); }
@@ -1117,15 +1117,18 @@ public:
 
 		FLandscapeToolPaintBase<ToolTarget, FLandscapeToolStrokeFlatten<ToolTarget>>::Tick(ViewportClient, DeltaTime);
 
-		bool bShowGrid = this->EdMode->UISettings->bUseFlattenTarget && this->EdMode->CurrentToolTarget.TargetType == ELandscapeToolTargetType::Heightmap && this->EdMode->UISettings->bShowFlattenTargetPreview;
-		MeshComponent->SetVisibility(bShowGrid);
+		if (HeightmapFlattenPreviewComponent != nullptr)
+		{
+			bool bShowGrid = this->EdMode->UISettings->bUseFlattenTarget && this->EdMode->UISettings->bShowFlattenTargetPreview;
+			HeightmapFlattenPreviewComponent->SetVisibility(bShowGrid);
+		}
 	}
 
 	virtual bool MouseMove(FEditorViewportClient* ViewportClient, FViewport* Viewport, int32 x, int32 y) override
 	{
 		bool bResult = FLandscapeToolPaintBase<ToolTarget, FLandscapeToolStrokeFlatten<ToolTarget>>::MouseMove(ViewportClient, Viewport, x, y);
 
-		if (ViewportClient->IsLevelEditorClient() && MeshComponent != NULL)
+		if (ViewportClient->IsLevelEditorClient() && HeightmapFlattenPreviewComponent != nullptr)
 		{
 			FVector MousePosition;
 			this->EdMode->LandscapeMouseTrace((FEditorViewportClient*)ViewportClient, x, y, MousePosition);
@@ -1135,7 +1138,7 @@ public:
 			Origin.X = FMath::RoundToFloat(MousePosition.X);
 			Origin.Y = FMath::RoundToFloat(MousePosition.Y);
 			Origin.Z = (FMath::RoundToFloat((this->EdMode->UISettings->FlattenTarget - LocalToWorld.GetTranslation().Z) / LocalToWorld.GetScale3D().Z * LANDSCAPE_INV_ZSCALE) - 0.1f) * LANDSCAPE_ZSCALE;
-			MeshComponent->SetRelativeLocation(Origin, false);
+			HeightmapFlattenPreviewComponent->SetRelativeLocation(Origin, false);
 
 			// Clamp the value to the height map
 			uint16 TexHeight = LandscapeDataAccess::GetTexHeight(MousePosition.Z);
@@ -1156,31 +1159,34 @@ public:
 			return;
 		}
 
-		ALandscapeProxy* LandscapeProxy = this->EdMode->CurrentToolTarget.LandscapeInfo->GetLandscapeProxy();
-		MeshComponent = NewObject<UStaticMeshComponent>(LandscapeProxy, NAME_None, RF_Transient);
-		MeshComponent->SetStaticMesh(PlaneMesh);
-		MeshComponent->SetCanEverAffectNavigation(false);
-		MeshComponent->AttachToComponent(LandscapeProxy->GetRootComponent(), FAttachmentTransformRules::KeepRelativeTransform);
-		MeshComponent->RegisterComponent();
-
-		bool bShowGrid = this->EdMode->UISettings->bUseFlattenTarget && this->EdMode->CurrentToolTarget.TargetType == ELandscapeToolTargetType::Heightmap && this->EdMode->UISettings->bShowFlattenTargetPreview;
-		MeshComponent->SetVisibility(bShowGrid);
-
-		// Try to set a sane initial location for the preview grid
-		const FTransform LocalToWorld = this->EdMode->CurrentToolTarget.LandscapeInfo->GetLandscapeProxy()->GetRootComponent()->GetComponentToWorld();
-		FVector Origin = FVector::ZeroVector;
-		Origin.Z = (FMath::RoundToFloat((this->EdMode->UISettings->FlattenTarget - LocalToWorld.GetTranslation().Z) / LocalToWorld.GetScale3D().Z * LANDSCAPE_INV_ZSCALE) - 0.1f) * LANDSCAPE_ZSCALE;
-		MeshComponent->SetRelativeLocation(Origin, false);
+		if (ToolTarget::TargetType == ELandscapeToolTargetType::Heightmap)
+		{
+			ALandscapeProxy* LandscapeProxy = this->EdMode->CurrentToolTarget.LandscapeInfo->GetLandscapeProxy();
+			HeightmapFlattenPreviewComponent = NewObject<UStaticMeshComponent>(LandscapeProxy, NAME_None, RF_Transient);
+			HeightmapFlattenPreviewComponent->SetStaticMesh(HeightmapFlattenPlaneMesh);
+			HeightmapFlattenPreviewComponent->SetCanEverAffectNavigation(false);
+			HeightmapFlattenPreviewComponent->AttachToComponent(LandscapeProxy->GetRootComponent(), FAttachmentTransformRules::KeepRelativeTransform);
+			HeightmapFlattenPreviewComponent->RegisterComponent();
+
+			bool bShowGrid = this->EdMode->UISettings->bUseFlattenTarget && this->EdMode->UISettings->bShowFlattenTargetPreview;
+			HeightmapFlattenPreviewComponent->SetVisibility(bShowGrid);
+
+			// Try to set a sane initial location for the preview grid
+			const FTransform LocalToWorld = this->EdMode->CurrentToolTarget.LandscapeInfo->GetLandscapeProxy()->GetRootComponent()->GetComponentToWorld();
+			FVector Origin = FVector::ZeroVector;
+			Origin.Z = (FMath::RoundToFloat((this->EdMode->UISettings->FlattenTarget - LocalToWorld.GetTranslation().Z) / LocalToWorld.GetScale3D().Z * LANDSCAPE_INV_ZSCALE) - 0.1f) * LANDSCAPE_ZSCALE;
+			HeightmapFlattenPreviewComponent->SetRelativeLocation(Origin, false);
+		}
 	}
 
 	virtual void ExitTool() override
 	{
 		FLandscapeToolPaintBase<ToolTarget, FLandscapeToolStrokeFlatten<ToolTarget>>::ExitTool();
 
-		if (MeshComponent)
+		if (HeightmapFlattenPreviewComponent != nullptr)
 		{
-			MeshComponent->DetachFromComponent(FDetachmentTransformRules::KeepRelativeTransform);
-			MeshComponent->DestroyComponent();
+			HeightmapFlattenPreviewComponent->DetachFromComponent(FDetachmentTransformRules::KeepRelativeTransform);
+			HeightmapFlattenPreviewComponent->DestroyComponent();
 		}
 	}
 };
