Clear selected landscape components in the details panel when clearing the selection

diff --git a/Engine/Source/Editor/LandscapeEditor/Private/LandscapeEditorDetailCustomization_MiscTools.cpp b/Engine/Source/Editor/LandscapeEditor/Private/LandscapeEditorDetailCustomization_MiscTools.cpp
--- a/Engine/Source/Editor/LandscapeEditor/Private/LandscapeEditorDetailCustomization_MiscTools.cpp
+++ b/Engine/Source/Editor/LandscapeEditor/Private/LandscapeEditorDetailCustomization_MiscTools.cpp
@@ -256,7 +256,14 @@ FReply FLandscapeEditorDetailCustomization_MiscTools::OnClearComponentSelectionB
 		{
 			FScopedTransaction Transaction(LOCTEXT("Component.Undo_ClearSelected", "Clearing Selection"));
 			LandscapeInfo->Modify();
+
+			TSet<ULandscapeComponent*> PreviouslySelectedComponents = LandscapeInfo->GetSelectedComponents();
 			LandscapeInfo->ClearSelectedRegion(true);
+
+			// Remove the previously selected components from the selected objects in the details view: 
+			FPropertyEditorModule& PropertyModule = FModuleManager::Get().LoadModuleChecked<FPropertyEditorModule>(TEXT("PropertyEditor"));
+			TArray<UObject*> ObjectsToRemove(PreviouslySelectedComponents.Array());
+			PropertyModule.RemoveDeletedObjects(ObjectsToRemove);
 		}
 	}
 
