Adding a case for the Layer Context Menu to raise a little notification when right-clicking on heightmaps in WP levels to point users to Manage -> Import instead.

diff --git a/Engine/Source/Editor/LandscapeEditor/Private/LandscapeEditorDetailCustomization_TargetLayers.cpp b/Engine/Source/Editor/LandscapeEditor/Private/LandscapeEditorDetailCustomization_TargetLayers.cpp
--- a/Engine/Source/Editor/LandscapeEditor/Private/LandscapeEditorDetailCustomization_TargetLayers.cpp
+++ b/Engine/Source/Editor/LandscapeEditor/Private/LandscapeEditorDetailCustomization_TargetLayers.cpp
@@ -907,26 +907,12 @@ TSharedPtr<SWidget> FLandscapeEditorCustomNodeBuilder_TargetLayers::OnTargetLaye
 		
 		MenuBuilder.BeginSection("LandscapeEditorLayerActions", LOCTEXT("LayerContextMenu.Heading", "Layer Actions"));
 		{
-			// In Grid Based Landscape Import/Export/Reimport is done through the Import Tool
 			FEdModeLandscape* LandscapeEdMode = GetEditorMode();
-			if (LandscapeEdMode && !LandscapeEdMode->IsGridBased())
+			if (LandscapeEdMode)
 			{
-				// Export
-				FUIAction ExportAction = FUIAction(FExecuteAction::CreateStatic(&FLandscapeEditorCustomNodeBuilder_TargetLayers::OnExportLayer, Target));
-				MenuBuilder.AddMenuEntry(LOCTEXT("LayerContextMenu.Export", "Export to file"), FText(), FSlateIcon(), ExportAction);
-
-				// Import
-				FUIAction ImportAction = FUIAction(FExecuteAction::CreateStatic(&FLandscapeEditorCustomNodeBuilder_TargetLayers::OnImportLayer, Target));
-				MenuBuilder.AddMenuEntry(LOCTEXT("LayerContextMenu.Import", "Import from file"), FText(), FSlateIcon(), ImportAction);
-
-				// Reimport
-				const FString& ReimportPath = Target->GetReimportFilePath();
-
-				if (!ReimportPath.IsEmpty())
-				{
-					FUIAction ReImportAction = FUIAction(FExecuteAction::CreateStatic(&FLandscapeEditorCustomNodeBuilder_TargetLayers::OnReimportLayer, Target));
-					MenuBuilder.AddMenuEntry(FText::Format(LOCTEXT("LayerContextMenu.ReImport", "Reimport from {0}"), FText::FromString(ReimportPath)), FText(), FSlateIcon(), ReImportAction);
-				}
+				FUIAction LandscapeHeightmapChangeToolsAction = FUIAction(FExecuteAction::CreateStatic(&FLandscapeEditorCustomNodeBuilder_TargetLayers::OnHeightmapLayerContextMenu, Target));
+				MenuBuilder.AddMenuEntry(LOCTEXT("LayerContextMenu.Heightmap", "Import From/Export To File..."),
+										LOCTEXT("LayerContextMenu.HeightmapToolTip", "Opens the Landscape Import tool in order to import / export heightmaps from / to external files."), FSlateIcon(), LandscapeHeightmapChangeToolsAction);
 			}
 
 			if (Target->TargetType == ELandscapeToolTargetType::Weightmap)
@@ -1091,6 +1077,15 @@ void FLandscapeEditorCustomNodeBuilder_TargetLayers::OnReimportLayer(const TShar
 	}
 }
 
+void FLandscapeEditorCustomNodeBuilder_TargetLayers::OnHeightmapLayerContextMenu(const TSharedRef<FLandscapeTargetListInfo> Target)
+{
+	FEdModeLandscape* LandscapeEdMode = GetEditorMode();
+	if (LandscapeEdMode)
+	{
+		LandscapeEdMode->SetCurrentTool("ImportExport");
+	}
+}
+
 void FLandscapeEditorCustomNodeBuilder_TargetLayers::OnFillLayer(const TSharedRef<FLandscapeTargetListInfo> Target)
 {
 	FScopedTransaction Transaction(LOCTEXT("Undo_FillLayer", "Filling Landscape Layer"));
diff --git a/Engine/Source/Editor/LandscapeEditor/Private/LandscapeEditorDetailCustomization_TargetLayers.h b/Engine/Source/Editor/LandscapeEditor/Private/LandscapeEditorDetailCustomization_TargetLayers.h
--- a/Engine/Source/Editor/LandscapeEditor/Private/LandscapeEditorDetailCustomization_TargetLayers.h
+++ b/Engine/Source/Editor/LandscapeEditor/Private/LandscapeEditorDetailCustomization_TargetLayers.h
@@ -73,6 +73,7 @@ protected:
 	static void OnExportLayer(const TSharedRef<FLandscapeTargetListInfo> Target);
 	static void OnImportLayer(const TSharedRef<FLandscapeTargetListInfo> Target);
 	static void OnReimportLayer(const TSharedRef<FLandscapeTargetListInfo> Target);
+	static void OnHeightmapLayerContextMenu(const TSharedRef<FLandscapeTargetListInfo> Target);
 	static void OnFillLayer(const TSharedRef<FLandscapeTargetListInfo> Target);
 	static void FillEmptyLayers(ULandscapeInfo* LandscapeInfo, ULandscapeLayerInfoObject* LandscapeInfoObject);
 	static void OnClearLayer(const TSharedRef<FLandscapeTargetListInfo> Target);
