Landscape: Fix resizing of new or imported landscapes by dragging the boundaries

diff --git a/Engine/Source/Editor/LandscapeEditor/Private/LandscapeEdMode.cpp b/Engine/Source/Editor/LandscapeEditor/Private/LandscapeEdMode.cpp
--- a/Engine/Source/Editor/LandscapeEditor/Private/LandscapeEdMode.cpp
+++ b/Engine/Source/Editor/LandscapeEditor/Private/LandscapeEdMode.cpp
@@ -1605,6 +1605,14 @@ bool FEdModeLandscape::InputKey(FEditorViewportClient* ViewportClient, FViewport
 		}
 	}
 
+	if (NewLandscapePreviewMode != ENewLandscapePreviewMode::None)
+	{
+		if (CurrentTool && CurrentTool->InputKey(ViewportClient, Viewport, Key, Event))
+		{
+			return false; // false to let FEditorViewportClient.InputKey start/end mouse tracking and enable InputDelta() so we can use it
+		}
+	}
+	else
 	{
 		// Override Key Input for Selection Brush
 		if (CurrentBrush)
diff --git a/Engine/Source/Editor/LandscapeEditor/Private/LandscapeEdModeComponentTools.cpp b/Engine/Source/Editor/LandscapeEditor/Private/LandscapeEdModeComponentTools.cpp
--- a/Engine/Source/Editor/LandscapeEditor/Private/LandscapeEdModeComponentTools.cpp
+++ b/Engine/Source/Editor/LandscapeEditor/Private/LandscapeEdModeComponentTools.cpp
@@ -1860,7 +1860,7 @@ public:
 						DraggingEdge = EdgeProxy->Edge;
 						DraggingEdge_Remainder = 0;
 
-						return false; // false to let FEditorViewportClient.InputKey start mouse tracking and enable InputDelta() so we can use it
+						return true;
 					}
 				}
 			}
@@ -1871,7 +1871,7 @@ public:
 					DraggingEdge = ELandscapeEdge::None;
 					DraggingEdge_Remainder = 0;
 
-					return false; // false to let FEditorViewportClient.InputKey end mouse tracking
+					return true;
 				}
 			}
 		}
@@ -1903,35 +1903,48 @@ public:
 			HitLocation = Transform.InverseTransformPosition(HitLocation);
 
 			UISettings->Modify();
-			switch (DraggingEdge)
-			{
-			case ELandscapeEdge::X_Negative:
-			case ELandscapeEdge::X_Negative_Y_Negative:
-			case ELandscapeEdge::X_Negative_Y_Positive:
+
+			auto DragEdge = [&UISettings, &HitLocation, &Transform](const ELandscapeEdge::Type Edge)
 			{
-				const int32 InitialComponentCountX = UISettings->NewLandscape_ComponentCount.X;
-				const int32 Delta = FMath::RoundToInt(HitLocation.X + (float)InitialComponentCountX / 2);
-				UISettings->NewLandscape_ComponentCount.X = InitialComponentCountX - Delta;
+				int32& ComponentCount = Edge == ELandscapeEdge::X_Negative || Edge == ELandscapeEdge::X_Positive ? UISettings->NewLandscape_ComponentCount.X : UISettings->NewLandscape_ComponentCount.Y;
+				const float Hit = Edge == ELandscapeEdge::X_Negative || Edge == ELandscapeEdge::X_Positive ? HitLocation.X : HitLocation.Y;
+				const float PosOrNeg = Edge == ELandscapeEdge::X_Negative || Edge == ELandscapeEdge::Y_Negative ? -1.0f : 1.0f;
+				const FVector XOrY = Edge == ELandscapeEdge::X_Negative || Edge == ELandscapeEdge::X_Positive ? FVector(1, 0, 0) : FVector(0, 1, 0);
+				
+				const int32 InitialComponentCount = ComponentCount;
+				const int32 Delta = FMath::RoundToInt(Hit - PosOrNeg * static_cast<float>(InitialComponentCount) / 2.0f);
+				ComponentCount = InitialComponentCount + PosOrNeg * Delta;
 				UISettings->NewLandscape_ClampSize();
-				const int32 ActualDelta = UISettings->NewLandscape_ComponentCount.X - InitialComponentCountX;
-				UISettings->NewLandscape_Location -= Transform.TransformVector(FVector(((float)ActualDelta / 2), 0, 0));
+				const float ActualDelta = static_cast<float>(ComponentCount - InitialComponentCount) / 2.0f;
+				UISettings->NewLandscape_Location += PosOrNeg * XOrY * Transform.TransformVector(FVector(ActualDelta, ActualDelta, 0));
+			};
+
+			if (DraggingEdge == ELandscapeEdge::X_Negative ||
+				DraggingEdge == ELandscapeEdge::X_Negative_Y_Negative ||
+				DraggingEdge ==	ELandscapeEdge::X_Negative_Y_Positive)
+			{
+				DragEdge(ELandscapeEdge::X_Negative);
 			}
-			break;
-			case ELandscapeEdge::X_Positive:
-			case ELandscapeEdge::X_Positive_Y_Negative:
-			case ELandscapeEdge::X_Positive_Y_Positive:
+
+			if (DraggingEdge == ELandscapeEdge::X_Positive ||
+				DraggingEdge == ELandscapeEdge::X_Positive_Y_Negative ||
+				DraggingEdge ==	ELandscapeEdge::X_Positive_Y_Positive)
 			{
-				const int32 InitialComponentCountX = UISettings->NewLandscape_ComponentCount.X;
-				int32 Delta = FMath::RoundToInt(HitLocation.X - (float)InitialComponentCountX / 2);
-				UISettings->NewLandscape_ComponentCount.X = InitialComponentCountX + Delta;
-				UISettings->NewLandscape_ClampSize();
-				const int32 ActualDelta = UISettings->NewLandscape_ComponentCount.X - InitialComponentCountX;
-				UISettings->NewLandscape_Location += Transform.TransformVector(FVector(((float)ActualDelta / 2), 0, 0));
+				DragEdge(ELandscapeEdge::X_Positive);
+			}
+
+			if (DraggingEdge == ELandscapeEdge::Y_Negative ||
+				DraggingEdge == ELandscapeEdge::X_Negative_Y_Negative ||
+				DraggingEdge == ELandscapeEdge::X_Positive_Y_Negative)
+			{
+				DragEdge(ELandscapeEdge::Y_Negative);
 			}
-			break;
-			case  ELandscapeEdge::Y_Negative:
-			case  ELandscapeEdge::Y_Positive:
-				break;
+
+			if (DraggingEdge == ELandscapeEdge::Y_Positive ||
+				DraggingEdge == ELandscapeEdge::X_Negative_Y_Positive ||
+				DraggingEdge == ELandscapeEdge::X_Positive_Y_Positive)
+			{
+				DragEdge(ELandscapeEdge::Y_Positive);
 			}
 						
 			return true;
