Landscape: Use same edge handles and colors for importing landscape as for new landscape; refactor line drawing code

diff --git a/Engine/Source/Editor/LandscapeEditor/Private/LandscapeEdModeComponentTools.cpp b/Engine/Source/Editor/LandscapeEditor/Private/LandscapeEdModeComponentTools.cpp
--- a/Engine/Source/Editor/LandscapeEditor/Private/LandscapeEdModeComponentTools.cpp
+++ b/Engine/Source/Editor/LandscapeEditor/Private/LandscapeEdModeComponentTools.cpp
@@ -1957,12 +1957,12 @@ public:
 	{
 		if (EdMode->NewLandscapePreviewMode != ENewLandscapePreviewMode::None)
 		{
-			static const float        CornerSize = 0.33f;
-			static const FLinearColor CornerColor(1.0f, 1.0f, 0.5f);
-			static const FLinearColor EdgeColor(1.0f, 1.0f, 0.0f);
-			static const FLinearColor ComponentBorderColor(0.0f, 0.85f, 0.0f);
-			static const FLinearColor SectionBorderColor(0.0f, 0.4f, 0.0f);
-			static const FLinearColor InnerColor(0.0f, 0.25f, 0.0f);
+			static constexpr float        CornerSize = 0.33f;
+			static constexpr FLinearColor CornerColor(1.0f, 0.5f, 0.0f);
+			static constexpr FLinearColor EdgeColor(1.0f, 1.0f, 0.0f);
+			static constexpr FLinearColor ComponentBorderColor(0.0f, 0.85f, 0.0f);
+			static constexpr FLinearColor SectionBorderColor(0.0f, 0.4f, 0.0f);
+			static constexpr FLinearColor InnerColor(0.0f, 0.25f, 0.0f);
 
 			const ELevelViewportType ViewportType = ((FEditorViewportClient*)Viewport->GetClient())->ViewportType;
 
@@ -1974,6 +1974,19 @@ public:
 			const FVector Offset = EdMode->UISettings->NewLandscape_Location + FTransform(EdMode->UISettings->NewLandscape_Rotation, FVector::ZeroVector, EdMode->UISettings->NewLandscape_Scale).TransformVector(FVector(-ComponentCountX * ComponentSize / 2, -ComponentCountY * ComponentSize / 2, 0));
 			const FTransform Transform = FTransform(EdMode->UISettings->NewLandscape_Rotation, Offset, EdMode->UISettings->NewLandscape_Scale);
 
+			using LineCoords = TTuple<FVector, FVector>;
+
+			auto DrawLine = [&PDI, &Transform](const LineCoords& AB, const FLinearColor& Color, const uint8 DepthPriorityGroup)
+			{
+				PDI->DrawLine(Transform.TransformPosition(AB.Get<0>()), Transform.TransformPosition(AB.Get<1>()), Color, DepthPriorityGroup);
+			};
+
+			auto DrawLineBorder = [&PDI, &DrawLine](const ELandscapeEdge::Type Edge, const LineCoords& AB, const FLinearColor& Color)
+			{
+				PDI->SetHitProxy(new HNewLandscapeGrabHandleProxy(Edge));
+				DrawLine(AB, Color, SDPG_Foreground);
+			};
+
 			if (EdMode->NewLandscapePreviewMode == ENewLandscapePreviewMode::ImportLandscape)
 			{
 				auto GetColor = [&](int32 ComponentIndex)
@@ -1982,16 +1995,15 @@ public:
 					{
 						return EdgeColor;
 					}
-					else
-					{
-						return ComponentBorderColor;
-					}
+					return ComponentBorderColor;
 				};
 
 				const TArray<uint16>& ImportHeights = EdMode->UISettings->GetImportLandscapeData();
 				if (ImportHeights.Num() != 0)
 				{
-					const float InvQuadsPerComponent = 1.0f / (float)QuadsPerComponent;
+					const int32 NumLinesForeground = ComponentCountX * ComponentCountY * 2 + ComponentCountX + ComponentCountY + 8;
+					PDI->AddReserveLines(SDPG_Foreground, NumLinesForeground);
+
 					const int32 SizeX = ComponentCountX * QuadsPerComponent + 1;
 					const int32 SizeY = ComponentCountY * QuadsPerComponent + 1;
 					const int32 ImportSizeX = EdMode->UISettings->ImportLandscape_Width;
@@ -1999,77 +2011,125 @@ public:
 					const int32 OffsetX = (SizeX - ImportSizeX) / 2;
 					const int32 OffsetY = (SizeY - ImportSizeY) / 2;
 
-					for (int32 ComponentY = 0; ComponentY < ComponentCountY; ComponentY++)
+					// Get coordinates for a line in X direction
+					auto LineX = [QuadsPerComponent, ImportSizeX, ImportSizeY, &ImportHeights, OffsetX, OffsetY](int32 X, int32 Y)
 					{
-						const int32 Y0 = ComponentY * QuadsPerComponent;
-						const int32 Y1 = (ComponentY + 1) * QuadsPerComponent;
-
+						X *= QuadsPerComponent;
+						const int32 Y0 = Y * QuadsPerComponent;
+						const int32 Y1 = (Y + 1) * QuadsPerComponent;
+						const int32 ImportX = FMath::Clamp<int32>(X - OffsetX, 0, ImportSizeX - 1);
 						const int32 ImportY0 = FMath::Clamp<int32>(Y0 - OffsetY, 0, ImportSizeY - 1);
 						const int32 ImportY1 = FMath::Clamp<int32>(Y1 - OffsetY, 0, ImportSizeY - 1);
+						const float Z0 = (static_cast<float>(ImportHeights[ImportX + ImportY0 * ImportSizeX]) - 32768.0f) * LANDSCAPE_ZSCALE;
+						const float Z1 = (static_cast<float>(ImportHeights[ImportX + ImportY1 * ImportSizeX]) - 32768.0f) * LANDSCAPE_ZSCALE;
+						return LineCoords{FVector(X, Y0, Z0), FVector(X, Y1, Z1)};
+					};
 
-						for (int32 ComponentX = 0; ComponentX < ComponentCountX; ComponentX++)
+					// Get coordinates for a line in Y direction
+					auto LineY = [QuadsPerComponent, ImportSizeX, ImportSizeY, &ImportHeights, OffsetX, OffsetY](int32 X, int32 Y)
+					{
+						Y *= QuadsPerComponent;
+						const int32 X0 = X * QuadsPerComponent;
+						const int32 X1 = (X + 1) * QuadsPerComponent;
+						const int32 ImportY = FMath::Clamp<int32>(Y - OffsetY, 0, ImportSizeY - 1);
+						const int32 ImportX0 = FMath::Clamp<int32>(X0 - OffsetX, 0, ImportSizeX - 1);
+						const int32 ImportX1 = FMath::Clamp<int32>(X1 - OffsetX, 0, ImportSizeX - 1);
+						const float Z0 = (static_cast<float>(ImportHeights[ImportX0 + ImportY * ImportSizeX]) - 32768.0f) * LANDSCAPE_ZSCALE;
+						const float Z1 = (static_cast<float>(ImportHeights[ImportX1 + ImportY * ImportSizeX]) - 32768.0f) * LANDSCAPE_ZSCALE;
+						return LineCoords{FVector(X0, Y, Z0), FVector(X1, Y, Z1)};
+					};
+
+					// Draw a border in X direction
+					auto DrawBorderX = [&PDI, ComponentCountY, &DrawLine, &DrawLineBorder, &LineX](
+						const int32 X, ELandscapeEdge::Type FirstCornerEdge, ELandscapeEdge::Type BorderEdge, ELandscapeEdge::Type LastCornerEdge)
+					{
+						const LineCoords FirstComponent = LineX(X, 0);
+						const LineCoords LastComponent = LineX(X, ComponentCountY - 1);
+						const FVector FirstCornerEnd = FirstComponent.Get<0>() + CornerSize * (FirstComponent.Get<1>() - FirstComponent.Get<0>());
+						const FVector LastCornerBegin = LastComponent.Get<1>() - CornerSize * (LastComponent.Get<1>() - LastComponent.Get<0>());
+
+						DrawLineBorder(FirstCornerEdge, {FirstComponent.Get<0>(), FirstCornerEnd}, CornerColor);
+						PDI->SetHitProxy(new HNewLandscapeGrabHandleProxy(BorderEdge));
+						if (ComponentCountY == 1)
 						{
-							const int32 X0 = ComponentX * QuadsPerComponent;
-							const int32 X1 = (ComponentX + 1) * QuadsPerComponent;
-							const int32 ImportX0 = FMath::Clamp<int32>(X0 - OffsetX, 0, ImportSizeX - 1);
-							const int32 ImportX1 = FMath::Clamp<int32>(X1 - OffsetX, 0, ImportSizeX - 1);
-							const float Z00 = ((float)ImportHeights[ImportX0 + ImportY0 * ImportSizeX] - 32768.0f) * LANDSCAPE_ZSCALE;
-							const float Z01 = ((float)ImportHeights[ImportX0 + ImportY1 * ImportSizeX] - 32768.0f) * LANDSCAPE_ZSCALE;
-							const float Z10 = ((float)ImportHeights[ImportX1 + ImportY0 * ImportSizeX] - 32768.0f) * LANDSCAPE_ZSCALE;
-							const float Z11 = ((float)ImportHeights[ImportX1 + ImportY1 * ImportSizeX] - 32768.0f) * LANDSCAPE_ZSCALE;
-
-							if (ComponentX == 0)
+							DrawLine({FirstCornerEnd, LastCornerBegin}, EdgeColor, SDPG_Foreground);
+						}
+						else
+						{
+							DrawLine({FirstCornerEnd, FirstComponent.Get<1>()}, EdgeColor, SDPG_Foreground);
+							for (int32 Y = 1; Y < ComponentCountY - 1; ++Y)
 							{
-								PDI->SetHitProxy(new HNewLandscapeGrabHandleProxy(ELandscapeEdge::X_Negative));
-								PDI->DrawLine(Transform.TransformPosition(FVector(X0, Y0, Z00)), Transform.TransformPosition(FVector(X0, Y1, Z01)), GetColor(ComponentX), SDPG_Foreground);
-								PDI->SetHitProxy(NULL);
+								DrawLine(LineX(X, Y), EdgeColor, SDPG_Foreground);
 							}
+							DrawLine({LastComponent.Get<0>(), LastCornerBegin}, EdgeColor, SDPG_Foreground);
+						}
+						DrawLineBorder(LastCornerEdge, {LastCornerBegin, LastComponent.Get<1>()}, CornerColor);
+					};
 
-							if (ComponentX == ComponentCountX - 1)
-							{
-								PDI->SetHitProxy(new HNewLandscapeGrabHandleProxy(ELandscapeEdge::X_Positive));
-								PDI->DrawLine(Transform.TransformPosition(FVector(X1, Y0, Z10)), Transform.TransformPosition(FVector(X1, Y1, Z11)), GetColor(ComponentX), SDPG_Foreground);
-								PDI->SetHitProxy(NULL);
-							}
-							else
+					// Draw a border in Y direction
+					auto DrawBorderY = [&PDI, ComponentCountX, &DrawLine, &DrawLineBorder, &LineY](
+						const int32 Y, ELandscapeEdge::Type FirstCornerEdge, ELandscapeEdge::Type BorderEdge, ELandscapeEdge::Type LastCornerEdge)
+					{
+						const LineCoords FirstComponent = LineY(0, Y);
+						const LineCoords LastComponent = LineY(ComponentCountX - 1, Y);
+						const FVector FirstCornerEnd = FirstComponent.Get<0>() + CornerSize * (FirstComponent.Get<1>() - FirstComponent.Get<0>());
+						const FVector LastCornerBegin = LastComponent.Get<1>() - CornerSize * (LastComponent.Get<1>() - LastComponent.Get<0>());
+
+						DrawLineBorder(FirstCornerEdge, {FirstComponent.Get<0>(), FirstCornerEnd}, CornerColor);
+						PDI->SetHitProxy(new HNewLandscapeGrabHandleProxy(BorderEdge));
+						if (ComponentCountX == 1)
+						{
+							DrawLine({FirstCornerEnd, LastCornerBegin}, EdgeColor, SDPG_Foreground);
+						}
+						else
+						{
+							DrawLine({FirstCornerEnd, FirstComponent.Get<1>()}, EdgeColor, SDPG_Foreground);
+							for (int32 X = 1; X < ComponentCountX - 1; ++X)
 							{
-								PDI->DrawLine(Transform.TransformPosition(FVector(X1, Y0, Z10)), Transform.TransformPosition(FVector(X1, Y1, Z11)), GetColor(ComponentX), SDPG_Foreground);
+								DrawLine(LineY(X, Y), EdgeColor, SDPG_Foreground);
 							}
+							DrawLine({LastComponent.Get<0>(), LastCornerBegin}, EdgeColor, SDPG_Foreground);
+						}
+						DrawLineBorder(LastCornerEdge, {LastCornerBegin, LastComponent.Get<1>()}, CornerColor);
+					};
+					
+					// Left border
+					DrawBorderX(0, ELandscapeEdge::X_Negative_Y_Negative, ELandscapeEdge::X_Negative, ELandscapeEdge::X_Negative_Y_Positive);
 
-							if (ComponentY == 0)
-							{
-								PDI->SetHitProxy(new HNewLandscapeGrabHandleProxy(ELandscapeEdge::Y_Negative));
-								PDI->DrawLine(Transform.TransformPosition(FVector(X0, Y0, Z00)), Transform.TransformPosition(FVector(X1, Y0, Z10)), GetColor(ComponentY), SDPG_Foreground);
-								PDI->SetHitProxy(NULL);
-							}
+					// Right border
+					DrawBorderX(ComponentCountX, ELandscapeEdge::X_Positive_Y_Negative, ELandscapeEdge::X_Positive, ELandscapeEdge::X_Positive_Y_Positive);
 
-							if (ComponentY == ComponentCountY - 1)
-							{
-								PDI->SetHitProxy(new HNewLandscapeGrabHandleProxy(ELandscapeEdge::Y_Positive));
-								PDI->DrawLine(Transform.TransformPosition(FVector(X0, Y1, Z01)), Transform.TransformPosition(FVector(X1, Y1, Z11)), GetColor(ComponentY), SDPG_Foreground);
-								PDI->SetHitProxy(NULL);
-							}
-							else
-							{
-								PDI->DrawLine(Transform.TransformPosition(FVector(X0, Y1, Z01)), Transform.TransformPosition(FVector(X1, Y1, Z11)), GetColor(ComponentY), SDPG_Foreground);
-							}
+					// Bottom border
+					DrawBorderY(0, ELandscapeEdge::X_Negative_Y_Negative, ELandscapeEdge::Y_Negative, ELandscapeEdge::X_Positive_Y_Negative);
 
-							// intra-component lines - too slow for big landscapes
-							/*
-							for (int32 x=1;x<QuadsPerComponent;x++)
-							{
-							PDI->DrawLine(Transform.TransformPosition(FVector(X0+x, Y0, FMath::Lerp(Z00,Z10,(float)x*InvQuadsPerComponent))), Transform.TransformPosition(FVector(X0+x, Y1, FMath::Lerp(Z01,Z11,(float)x*InvQuadsPerComponent))), ComponentBorderColor, SDPG_World);
-							}
-							for (int32 y=1;y<QuadsPerComponent;y++)
-							{
-							PDI->DrawLine(Transform.TransformPosition(FVector(X0, Y0+y, FMath::Lerp(Z00,Z01,(float)y*InvQuadsPerComponent))), Transform.TransformPosition(FVector(X1, Y0+y, FMath::Lerp(Z10,Z11,(float)y*InvQuadsPerComponent))), ComponentBorderColor, SDPG_World);
-							}
-							*/
+					// Top border
+					DrawBorderY(ComponentCountY, ELandscapeEdge::X_Negative_Y_Positive, ELandscapeEdge::Y_Positive, ELandscapeEdge::X_Positive_Y_Positive);
+
+					// Reset mouse cursor after all border are drawn
+					PDI->SetHitProxy(nullptr);
+
+					// Left to right
+					for (int32 X = 1; X < ComponentCountX; ++X)
+					{
+						const FLinearColor Color = GetColor(X);
+						for (int32 Y = 0; Y < ComponentCountY; ++Y)
+						{
+							DrawLine(LineX(X, Y), Color, SDPG_Foreground);
+						}
+					}
+
+					// Bottom to top
+					for (int32 Y = 1; Y < ComponentCountY; ++Y)
+					{
+						const FLinearColor Color = GetColor(Y);
+						for (int32 X = 0; X < ComponentCountX; ++X)
+						{
+							DrawLine(LineY(X, Y), Color, SDPG_Foreground);
 						}
 					}
 				}
 			}
-			else //if (EdMode->NewLandscapePreviewMode == ENewLandscapePreviewMode::NewLandscape)
+			else // EdMode->NewLandscapePreviewMode == ENewLandscapePreviewMode::NewLandscape
 			{
 				auto GetColor = [&](int32 QuadIndex)
 				{
@@ -2077,96 +2137,91 @@ public:
 					{
 						return EdgeColor;
 					}
-					else if (QuadIndex % QuadsPerComponent == 0)
+					if (QuadIndex % QuadsPerComponent == 0)
 					{
 						return ComponentBorderColor;
 					}
-					else if (QuadIndex % EdMode->UISettings->NewLandscape_QuadsPerSection == 0)
+					if (QuadIndex % EdMode->UISettings->NewLandscape_QuadsPerSection == 0)
 					{
 						return SectionBorderColor;
 					}
-
 					return InnerColor;
 				};
 								
 				if (ViewportType == LVT_Perspective || ViewportType == LVT_OrthoXY || ViewportType == LVT_OrthoNegativeXY)
 				{
-					for (int32 x = 0; x <= ComponentCountX * QuadsPerComponent; x++)
+					const int32 NumLines = ComponentCountX * QuadsPerComponent + 1 + ComponentCountY * QuadsPerComponent + 1;
+					const int32 NumLinesForeground = ComponentCountX + 1 + ComponentCountY + 1;
+					const int32 NumLinesWorld = NumLines - NumLinesForeground;
+					constexpr int32 NumLinesForegroundCorners = 8;
+
+					PDI->AddReserveLines(SDPG_Foreground, NumLinesForeground + NumLinesForegroundCorners);
+					PDI->AddReserveLines(SDPG_World, NumLinesWorld);
+
+					// Draw a border in X direction
+					auto DrawBorderX = [ComponentSize, ComponentCountY, &DrawLineBorder](
+						const int32 X, ELandscapeEdge::Type FirstCornerEdge, ELandscapeEdge::Type BorderEdge, ELandscapeEdge::Type LastCornerEdge)
 					{
-						if (x == 0)
-						{
-							PDI->SetHitProxy(new HNewLandscapeGrabHandleProxy(ELandscapeEdge::X_Negative_Y_Negative));
-							PDI->DrawLine(Transform.TransformPosition(FVector(x, 0, 0)), Transform.TransformPosition(FVector(x, CornerSize * ComponentSize, 0)), CornerColor, SDPG_Foreground);
-							PDI->SetHitProxy(new HNewLandscapeGrabHandleProxy(ELandscapeEdge::X_Negative));
-							PDI->DrawLine(Transform.TransformPosition(FVector(x, CornerSize * ComponentSize, 0)), Transform.TransformPosition(FVector(x, (ComponentCountY - CornerSize) * ComponentSize, 0)), EdgeColor, SDPG_Foreground);
-							PDI->SetHitProxy(new HNewLandscapeGrabHandleProxy(ELandscapeEdge::X_Negative_Y_Positive));
-							PDI->DrawLine(Transform.TransformPosition(FVector(x, (ComponentCountY - CornerSize) * ComponentSize, 0)), Transform.TransformPosition(FVector(x, ComponentCountY * ComponentSize, 0)), CornerColor, SDPG_Foreground);
-							PDI->SetHitProxy(NULL);
-						}
-						else if (x == ComponentCountX * QuadsPerComponent)
-						{
-							PDI->SetHitProxy(new HNewLandscapeGrabHandleProxy(ELandscapeEdge::X_Positive_Y_Negative));
-							PDI->DrawLine(Transform.TransformPosition(FVector(x, 0, 0)), Transform.TransformPosition(FVector(x, CornerSize * ComponentSize, 0)), CornerColor, SDPG_Foreground);
-							PDI->SetHitProxy(new HNewLandscapeGrabHandleProxy(ELandscapeEdge::X_Positive));
-							PDI->DrawLine(Transform.TransformPosition(FVector(x, CornerSize * ComponentSize, 0)), Transform.TransformPosition(FVector(x, (ComponentCountY - CornerSize) * ComponentSize, 0)), EdgeColor, SDPG_Foreground);
-							PDI->SetHitProxy(new HNewLandscapeGrabHandleProxy(ELandscapeEdge::X_Positive_Y_Positive));
-							PDI->DrawLine(Transform.TransformPosition(FVector(x, (ComponentCountY - CornerSize) * ComponentSize, 0)), Transform.TransformPosition(FVector(x, ComponentCountY * ComponentSize, 0)), CornerColor, SDPG_Foreground);
-							PDI->SetHitProxy(NULL);
-						}
-						else
-						{
-							FLinearColor CurrentColor = GetColor(x);
-							uint8 DepthPriority = CurrentColor == InnerColor ? SDPG_World : SDPG_Foreground;
-							PDI->DrawLine(Transform.TransformPosition(FVector(x, 0, 0)), Transform.TransformPosition(FVector(x, ComponentCountY * ComponentSize, 0)), GetColor(x), DepthPriority);
-						}
+						DrawLineBorder(FirstCornerEdge, {FVector(X, 0, 0), FVector(X, CornerSize * ComponentSize, 0)}, CornerColor);
+						DrawLineBorder(BorderEdge, {FVector(X, CornerSize * ComponentSize, 0), FVector(X, (ComponentCountY - CornerSize) * ComponentSize, 0)}, EdgeColor);
+						DrawLineBorder(LastCornerEdge, {FVector(X, (ComponentCountY - CornerSize) * ComponentSize, 0), FVector(X, ComponentCountY * ComponentSize, 0)}, CornerColor);
+					};
+
+					// Draw a border in Y direction
+					auto DrawBorderY = [ComponentSize, ComponentCountX, &DrawLineBorder](
+						const int32 Y, ELandscapeEdge::Type FirstCornerEdge, ELandscapeEdge::Type BorderEdge, ELandscapeEdge::Type LastCornerEdge)
+					{
+						DrawLineBorder(FirstCornerEdge, {FVector(0, Y, 0), FVector(CornerSize * ComponentSize, Y, 0)}, CornerColor);
+						DrawLineBorder(BorderEdge, {FVector(CornerSize * ComponentSize, Y, 0), FVector((ComponentCountX - CornerSize) * ComponentSize, Y, 0)}, EdgeColor);
+						DrawLineBorder(LastCornerEdge, {FVector((ComponentCountX - CornerSize) * ComponentSize, Y, 0), FVector(ComponentCountX * ComponentSize, Y, 0)}, CornerColor);
+					};
+
+					// Left border
+					DrawBorderX(0, ELandscapeEdge::X_Negative_Y_Negative, ELandscapeEdge::X_Negative, ELandscapeEdge::X_Negative_Y_Positive);
+
+					// Right border
+					DrawBorderX(ComponentCountX * QuadsPerComponent, ELandscapeEdge::X_Positive_Y_Negative, ELandscapeEdge::X_Positive, ELandscapeEdge::X_Positive_Y_Positive);
+
+					// Bottom border
+					DrawBorderY(0, ELandscapeEdge::X_Negative_Y_Negative, ELandscapeEdge::Y_Negative, ELandscapeEdge::X_Positive_Y_Negative);
+
+					// Top border
+					DrawBorderY(ComponentCountY * QuadsPerComponent, ELandscapeEdge::X_Negative_Y_Positive, ELandscapeEdge::Y_Positive, ELandscapeEdge::X_Positive_Y_Positive);
+
+					// Reset mouse cursor after all border are drawn
+					PDI->SetHitProxy(nullptr);
+
+					// Left to right
+					for (int32 X = 1; X < ComponentCountX * QuadsPerComponent; ++X)
+					{
+						const FLinearColor CurrentColor = GetColor(X);
+						const uint8 DepthPriority = CurrentColor == InnerColor ? SDPG_World : SDPG_Foreground;
+						DrawLine({FVector(X, 0, 0), FVector(X, ComponentCountY * ComponentSize, 0)}, CurrentColor, DepthPriority);
 					}
-				}
-				else
-				{
-					// Don't allow dragging to resize in side-view
-					// and there's no point drawing the inner lines as only the outer is visible
-					PDI->DrawLine(Transform.TransformPosition(FVector(0, 0, 0)), Transform.TransformPosition(FVector(0, ComponentCountY * ComponentSize, 0)), EdgeColor, SDPG_World);
-					PDI->DrawLine(Transform.TransformPosition(FVector(ComponentCountX * QuadsPerComponent, 0, 0)), Transform.TransformPosition(FVector(ComponentCountX * QuadsPerComponent, ComponentCountY * ComponentSize, 0)), EdgeColor, SDPG_World);
-				}
 
-				if (ViewportType == LVT_Perspective || ViewportType == LVT_OrthoXY || ViewportType == LVT_OrthoNegativeXY)
-				{
-					for (int32 y = 0; y <= ComponentCountY * QuadsPerComponent; y++)
+					// Bottom to top
+					for (int32 Y = 1; Y < ComponentCountY * QuadsPerComponent; ++Y)
 					{
-						if (y == 0)
-						{
-							PDI->SetHitProxy(new HNewLandscapeGrabHandleProxy(ELandscapeEdge::X_Negative_Y_Negative));
-							PDI->DrawLine(Transform.TransformPosition(FVector(0, y, 0)), Transform.TransformPosition(FVector(CornerSize * ComponentSize, y, 0)), CornerColor, SDPG_Foreground);
-							PDI->SetHitProxy(new HNewLandscapeGrabHandleProxy(ELandscapeEdge::Y_Negative));
-							PDI->DrawLine(Transform.TransformPosition(FVector(CornerSize * ComponentSize, y, 0)), Transform.TransformPosition(FVector((ComponentCountX - CornerSize) * ComponentSize, y, 0)), EdgeColor, SDPG_Foreground);
-							PDI->SetHitProxy(new HNewLandscapeGrabHandleProxy(ELandscapeEdge::X_Positive_Y_Negative));
-							PDI->DrawLine(Transform.TransformPosition(FVector((ComponentCountX - CornerSize) * ComponentSize, y, 0)), Transform.TransformPosition(FVector(ComponentCountX * ComponentSize, y, 0)), CornerColor, SDPG_Foreground);
-							PDI->SetHitProxy(NULL);
-						}
-						else if (y == ComponentCountY * QuadsPerComponent)
-						{
-							PDI->SetHitProxy(new HNewLandscapeGrabHandleProxy(ELandscapeEdge::X_Negative_Y_Positive));
-							PDI->DrawLine(Transform.TransformPosition(FVector(0, y, 0)), Transform.TransformPosition(FVector(CornerSize * ComponentSize, y, 0)), CornerColor, SDPG_Foreground);
-							PDI->SetHitProxy(new HNewLandscapeGrabHandleProxy(ELandscapeEdge::Y_Positive));
-							PDI->DrawLine(Transform.TransformPosition(FVector(CornerSize * ComponentSize, y, 0)), Transform.TransformPosition(FVector((ComponentCountX - CornerSize) * ComponentSize, y, 0)), EdgeColor, SDPG_Foreground);
-							PDI->SetHitProxy(new HNewLandscapeGrabHandleProxy(ELandscapeEdge::X_Positive_Y_Positive));
-							PDI->DrawLine(Transform.TransformPosition(FVector((ComponentCountX - CornerSize) * ComponentSize, y, 0)), Transform.TransformPosition(FVector(ComponentCountX * ComponentSize, y, 0)), CornerColor, SDPG_Foreground);
-							PDI->SetHitProxy(NULL);
-						}
-						else
-						{
-							FLinearColor CurrentColor = GetColor(y);
-							uint8 DepthPriority = CurrentColor == InnerColor ? SDPG_World : SDPG_Foreground;
-							PDI->DrawLine(Transform.TransformPosition(FVector(0, y, 0)), Transform.TransformPosition(FVector(ComponentCountX * ComponentSize, y, 0)), GetColor(y), DepthPriority);
-						}
+						const FLinearColor CurrentColor = GetColor(Y);
+						const uint8 DepthPriority = CurrentColor == InnerColor ? SDPG_World : SDPG_Foreground;
+						DrawLine({FVector(0, Y, 0), FVector(ComponentCountX * ComponentSize, Y, 0)}, CurrentColor, DepthPriority);
 					}
 				}
 				else
 				{
-					// Don't allow dragging to resize in side-view
-					// and there's no point drawing the inner lines as only the outer is visible
-					PDI->DrawLine(Transform.TransformPosition(FVector(0, 0, 0)), Transform.TransformPosition(FVector(ComponentCountX * ComponentSize, 0, 0)), EdgeColor, SDPG_World);
-					PDI->DrawLine(Transform.TransformPosition(FVector(0, ComponentCountY * QuadsPerComponent, 0)), Transform.TransformPosition(FVector(ComponentCountX * ComponentSize, ComponentCountY * QuadsPerComponent, 0)), EdgeColor, SDPG_World);
+					// Don't allow dragging to resize in side-view, and there is no point drawing the inner lines as only the outer are visible.
+
+					if (ViewportType == LVT_OrthoXZ || ViewportType == LVT_OrthoNegativeXZ)
+					{
+						DrawLine({FVector(0, 0, 0), FVector(ComponentCountX * ComponentSize, 0, 0)}, EdgeColor, SDPG_World);
+						DrawLine({FVector(0, ComponentCountY * QuadsPerComponent, 0), FVector(ComponentCountX * ComponentSize, ComponentCountY * QuadsPerComponent, 0)}, EdgeColor, SDPG_World);
+					}
+
+					if (ViewportType == LVT_OrthoYZ || ViewportType == LVT_OrthoNegativeYZ)
+					{
+						DrawLine({FVector(0, 0, 0), FVector(0, ComponentCountY * ComponentSize, 0)}, EdgeColor, SDPG_World);
+						DrawLine({FVector(ComponentCountX * QuadsPerComponent, 0, 0), FVector(ComponentCountX * QuadsPerComponent, ComponentCountY * ComponentSize, 0)}, EdgeColor, SDPG_World);
+					}
 				}
 			}
 		}
