Fix import descriptor not refreshing when toggling Flip Y Axis setting from ImportExport tool. Also, made option to flip coordinate on Y axis available when exporting tiled landscape.

diff --git a/Engine/Source/Editor/LandscapeEditor/Private/LandscapeEditorDetailCustomization_ImportExport.cpp b/Engine/Source/Editor/LandscapeEditor/Private/LandscapeEditorDetailCustomization_ImportExport.cpp
--- a/Engine/Source/Editor/LandscapeEditor/Private/LandscapeEditorDetailCustomization_ImportExport.cpp
+++ b/Engine/Source/Editor/LandscapeEditor/Private/LandscapeEditorDetailCustomization_ImportExport.cpp
@@ -156,7 +156,8 @@ void FLandscapeEditorDetailCustomization_ImportExport::CustomizeDetails(IDetailL
 	ImportExportCategory.AddProperty(PropertyHandle_ImportType).Visibility(MakeAttributeLambda([]() { return FLandscapeEditorDetailCustomization_ImportExport::GetImportExportVisibility(true); }));
 
 	TSharedRef<IPropertyHandle> PropertyHandle_FlipYAxis = DetailBuilder.GetProperty(GET_MEMBER_NAME_CHECKED(ULandscapeEditorObject, bFlipYAxis));
-	ImportExportCategory.AddProperty(PropertyHandle_FlipYAxis).Visibility(MakeAttributeLambda([]() { return (IsImporting() && !GetEditorMode()->UseSingleFileImport()) ? EVisibility::Visible : EVisibility::Collapsed; }));
+	PropertyHandle_FlipYAxis->SetOnPropertyValueChanged(FSimpleDelegate::CreateLambda([PropertyHandle_FlipYAxis]() { FLandscapeEditorDetailCustomization_ImportExport::OnImportHeightmapFilenameChanged(); }));
+	ImportExportCategory.AddProperty(PropertyHandle_FlipYAxis).Visibility(MakeAttributeLambda([]() { return (!GetEditorMode()->UseSingleFileImport()) ? EVisibility::Visible : EVisibility::Collapsed; }));
 
 	ImportExportCategory.AddProperty(PropertyHandle_ExportEditLayer).Visibility(MakeAttributeLambda([this]()
 	{
@@ -495,8 +496,10 @@ FReply FLandscapeEditorDetailCustomization_ImportExport::OnImportExportButtonCli
 					{
 						ALandscape* Landscape = LandscapeInfo->LandscapeActor.Get();
 						check(Landscape);
+
+						const int32 YCoord = LandscapeEdMode->UISettings->bFlipYAxis ? LandscapeExtent.Max.Y - ExportRegion.Max.Y : ExportRegion.Min.Y - LandscapeExtent.Min.Y;
 						FileOffset = FIntPoint((ExportRegion.Min.X - LandscapeExtent.Min.X) / Landscape->GridSize,
-								(ExportRegion.Min.Y - LandscapeExtent.Min.Y) / Landscape->GridSize);
+							YCoord / Landscape->GridSize);
 
 						// Remove the shared line/column that this proxy has with its neighbors because it 
 						// will be included by the neighbor or lost if there is none (that could become an option to avoid that loss)
