Added deprecation tag to FNewLandscapeUtils to properly deprecate it in a later version of UE

diff --git a/Engine/Source/Editor/LandscapeEditor/Public/NewLandscapeUtils.h b/Engine/Source/Editor/LandscapeEditor/Public/NewLandscapeUtils.h
--- a/Engine/Source/Editor/LandscapeEditor/Public/NewLandscapeUtils.h
+++ b/Engine/Source/Editor/LandscapeEditor/Public/NewLandscapeUtils.h
@@ -12,8 +12,7 @@ class ULandscapeEditorObject;
 enum class ENewLandscapePreviewMode : uint8;
 
 // This class has been replaced by FLandscapeImportHelper & ULandscapeEditorObject methods.
-// @todo_ow: deprecate (probably just delete since it was added for Datasmith which was updated to the new import code)
-class LANDSCAPEEDITOR_API FNewLandscapeUtils
+class LANDSCAPEEDITOR_API UE_DEPRECATED(5.1, "This class has been replaced by FLandscapeImportHelper & ULandscapeEditorObject methods. Use the these instead.") FNewLandscapeUtils
 {
 public:
 	static void ChooseBestComponentSizeForImport( ULandscapeEditorObject* UISettings );
