Landscape splines tool : added buttons to select all spline segments or control points

diff --git a/Engine/Source/Editor/LandscapeEditor/Private/LandscapeEdMode.h b/Engine/Source/Editor/LandscapeEditor/Private/LandscapeEdMode.h
--- a/Engine/Source/Editor/LandscapeEditor/Private/LandscapeEdMode.h
+++ b/Engine/Source/Editor/LandscapeEditor/Private/LandscapeEdMode.h
@@ -303,6 +303,8 @@ public:
 	bool HasSelectedSplineSegments() const;
 	void FlipSelectedSplineSegments();
 	void GetSelectedSplineOwners(TSet<AActor*>& SelectedSplineOwners) const;
+	virtual void SelectAllSplineControlPoints();
+	virtual void SelectAllSplineSegments();
 	virtual void SelectAllConnectedSplineControlPoints();
 	virtual void SelectAllConnectedSplineSegments();
 	virtual void SplineMoveToCurrentLevel();
diff --git a/Engine/Source/Editor/LandscapeEditor/Private/LandscapeEdModeSplineTools.cpp b/Engine/Source/Editor/LandscapeEditor/Private/LandscapeEdModeSplineTools.cpp
--- a/Engine/Source/Editor/LandscapeEditor/Private/LandscapeEdModeSplineTools.cpp
+++ b/Engine/Source/Editor/LandscapeEditor/Private/LandscapeEdModeSplineTools.cpp
@@ -48,10 +48,10 @@ class FLandscapeToolSplines : public FLandscapeTool, public FEditorUndoClient
 public:
 	FLandscapeToolSplines(FEdModeLandscape* InEdMode)
 		: EdMode(InEdMode)
-		, LandscapeInfo(NULL)
+		, LandscapeInfo(nullptr)
 		, SelectedSplineControlPoints()
 		, SelectedSplineSegments()
-		, DraggingTangent_Segment(NULL)
+		, DraggingTangent_Segment(nullptr)
 		, DraggingTangent_Length(0.0f)
 		, DraggingTangent_CacheCoordSpace(ECoordSystem::COORD_None)
 		, DraggingTangent_End(false)
@@ -216,6 +216,39 @@ public:
 		}
 	}
 
+	void SelectAllSplineSegments(const ULandscapeInfo& InLandscapeInfo)
+	{
+		TArray<TScriptInterface<ILandscapeSplineInterface>> SplineActors(InLandscapeInfo.GetSplineActors());
+		for (TScriptInterface<ILandscapeSplineInterface> SplineActor : SplineActors)
+		{
+			if (ULandscapeSplinesComponent* SplineComponent = SplineActor->GetSplinesComponent())
+			{
+				SplineComponent->ForEachControlPoint([this](ULandscapeSplineControlPoint* ControlPoint) 
+				{ 
+					for (const FLandscapeSplineConnection& Connection : ControlPoint->ConnectedSegments)
+					{
+						if (!Connection.Segment->IsSplineSelected())
+						{
+							SelectSegment(Connection.Segment);
+						}
+					}
+				});
+			}
+		}
+	}
+
+	void SelectAllControlPoints(const ULandscapeInfo& InLandscapeInfo)
+	{
+		TArray<TScriptInterface<ILandscapeSplineInterface>> SplineActors(InLandscapeInfo.GetSplineActors());
+		for (TScriptInterface<ILandscapeSplineInterface> SplineActor : SplineActors)
+		{
+			if (ULandscapeSplinesComponent* SplineComponent = SplineActor->GetSplinesComponent())
+			{
+				SplineComponent->ForEachControlPoint([this](ULandscapeSplineControlPoint* ControlPoint) { SelectControlPoint(ControlPoint); });
+			}
+		}
+	}
+
 	void SelectAdjacentControlPoints()
 	{
 		for (ULandscapeSplineSegment* Segment : SelectedSplineSegments)
@@ -1030,7 +1063,7 @@ public:
 
 	virtual void EndTool(FEditorViewportClient* ViewportClient) override
 	{
-		LandscapeInfo = NULL;
+		LandscapeInfo = nullptr;
 	}
 
 	virtual bool MouseMove(FEditorViewportClient* ViewportClient, FViewport* Viewport, int32 x, int32 y) override
@@ -1079,8 +1112,8 @@ public:
 
 		if (HitProxy)
 		{
-			ULandscapeSplineControlPoint* ClickedControlPoint = NULL;
-			ULandscapeSplineSegment* ClickedSplineSegment = NULL;
+			ULandscapeSplineControlPoint* ClickedControlPoint = nullptr;
+			ULandscapeSplineSegment* ClickedSplineSegment = nullptr;
 
 			if (HitProxy->IsA(HLandscapeSplineProxy_ControlPoint::StaticGetType()))
 			{
@@ -1118,7 +1151,7 @@ public:
 				}
 			}
 
-			if (ClickedControlPoint != NULL)
+			if (ClickedControlPoint != nullptr)
 			{
 				if (Click.IsShiftDown() && ClickedControlPoint->IsSplineSelected())
 				{
@@ -1135,7 +1168,7 @@ public:
 				GUnrealEd->RedrawLevelEditingViewports();
 				return true;
 			}
-			else if (ClickedSplineSegment != NULL)
+			else if (ClickedSplineSegment != nullptr)
 			{
 				// save info about what we grabbed
 				if (Click.IsShiftDown() && ClickedSplineSegment->IsSplineSelected())
@@ -1273,9 +1306,9 @@ public:
 				int32 HitX = InViewport->GetMouseX();
 				int32 HitY = InViewport->GetMouseY();
 				HHitProxy* HitProxy = InViewport->GetHitProxy(HitX, HitY);
-				if (HitProxy != NULL)
+				if (HitProxy != nullptr)
 				{
-					ULandscapeSplineControlPoint* ClickedControlPoint = NULL;
+					ULandscapeSplineControlPoint* ClickedControlPoint = nullptr;
 
 					if (HitProxy->IsA(HLandscapeSplineProxy_ControlPoint::StaticGetType()))
 					{
@@ -1304,7 +1337,7 @@ public:
 						}
 					}
 
-					if (ClickedControlPoint != NULL)
+					if (ClickedControlPoint != nullptr)
 					{
 						FScopedTransaction Transaction(LOCTEXT("LandscapeSpline_AddSegment", "Add Landscape Spline Segment"));
 
@@ -1332,7 +1365,7 @@ public:
 				HHitProxy* HitProxy = InViewport->GetHitProxy(HitX, HitY);
 				if (HitProxy)
 				{
-					ULandscapeSplineSegment* ClickedSplineSegment = NULL;
+					ULandscapeSplineSegment* ClickedSplineSegment = nullptr;
 					FTransform LandscapeToSpline;
 
 					if (HitProxy->IsA(HLandscapeSplineProxy_Segment::StaticGetType()))
@@ -1370,7 +1403,7 @@ public:
 						}
 					}
 
-					if (ClickedSplineSegment != NULL)
+					if (ClickedSplineSegment != nullptr)
 					{
 						FVector HitLocation;
 						if (EdMode->LandscapeMouseTrace(InViewportClient, HitLocation))
@@ -1469,7 +1502,7 @@ public:
 				{
 					DraggingTangent_Segment->UpdateSplinePoints(true);
 
-					DraggingTangent_Segment = NULL;
+					DraggingTangent_Segment = nullptr;
 
 					InViewportClient->SetWidgetCoordSystemSpace(DraggingTangent_CacheCoordSpace);
 					InViewportClient->SetRequiredCursorOverride(false);
@@ -1851,7 +1884,7 @@ public:
 						PDI->DrawLine(StartPos, HandlePos, TangentColor, SDPG_Foreground);
 						if (PDI->IsHitTesting()) PDI->SetHitProxy(new HLandscapeSplineProxy_Tangent(Connection.Segment, Connection.End));
 						PDI->DrawPoint(HandlePos, TangentColor, 10.0f, SDPG_Foreground);
-						if (PDI->IsHitTesting()) PDI->SetHitProxy(NULL);
+						if (PDI->IsHitTesting()) PDI->SetHitProxy(nullptr);
 					}
 				}
 			}
@@ -1875,7 +1908,7 @@ public:
 						PDI->DrawLine(EndPos, EndHandlePos, TangentColor, SDPG_Foreground);
 						if (PDI->IsHitTesting()) PDI->SetHitProxy(new HLandscapeSplineProxy_Tangent(Segment, !!End));
 						PDI->DrawPoint(EndHandlePos, TangentColor, 10.0f, SDPG_Foreground);
-						if (PDI->IsHitTesting()) PDI->SetHitProxy(NULL);
+						if (PDI->IsHitTesting()) PDI->SetHitProxy(nullptr);
 					}
 				}
 			}
@@ -2114,7 +2147,7 @@ public:
 		}
 	}
 
-	void InternalProcessEditCopy(FString* OutData = NULL)
+	void InternalProcessEditCopy(FString* OutData = nullptr)
 	{
 		bool bFirstSplineLocation = true;
 		FVector SplineLocation;
@@ -2160,11 +2193,11 @@ public:
 			Ar.Logf(TEXT("%s%s\r\n"), *FLandscapeSplineTextObjectFactory::SplineLocationTag, *SplineLocation.ToString());
 			for (UObject* Object : Objects)
 			{
-				UExporter::ExportToOutputDevice(&Context, Object, NULL, Ar, TEXT("copy"), 3, PPF_ExportsNotFullyQualified | PPF_Copy | PPF_Delimited, false, Object->GetOuter());
+				UExporter::ExportToOutputDevice(&Context, Object, nullptr, Ar, TEXT("copy"), 3, PPF_ExportsNotFullyQualified | PPF_Copy | PPF_Delimited, false, Object->GetOuter());
 			}
 			Ar.Logf(TEXT("%s\r\n"), *FLandscapeSplineTextObjectFactory::SplineEndTag);
 
-			if (OutData != NULL)
+			if (OutData != nullptr)
 			{
 				*OutData = MoveTemp(Ar);
 			}
@@ -2175,7 +2208,7 @@ public:
 		}
 	}
 
-	void InternalProcessEditPaste(FString* InData = NULL, bool bOffset = false)
+	void InternalProcessEditPaste(FString* InData = nullptr, bool bOffset = false)
 	{
 		FScopedTransaction Transaction(LOCTEXT("LandscapeSpline_Paste", "Paste Landscape Splines"));
 
@@ -2208,9 +2241,9 @@ public:
 		
 		SplineComponent->Modify();
 
-		const TCHAR* Data = NULL;
+		const TCHAR* Data = nullptr;
 		FString PasteString;
-		if (InData != NULL)
+		if (InData != nullptr)
 		{
 			Data = **InData;
 		}
@@ -2233,7 +2266,7 @@ public:
 			for (UObject* Object : OutObjects)
 			{
 				ULandscapeSplineControlPoint* ControlPoint = Cast<ULandscapeSplineControlPoint>(Object);
-				if (ControlPoint != NULL)
+				if (ControlPoint != nullptr)
 				{
 					ControlPoint->Location += FVector(500, 500, 0);
 					ControlPoint->UpdateSplinePoints();
@@ -2338,6 +2371,31 @@ void FEdModeLandscape::GetSelectedSplineOwners(TSet<AActor*>& SelectedSplineOwne
 	}
 }
 
+void FEdModeLandscape::SelectAllSplineControlPoints()
+{
+	ULandscapeInfo* CurrentLandscapeInfo = CurrentToolTarget.LandscapeInfo.Get();
+	if (SplinesTool && CurrentLandscapeInfo)
+	{
+		SplinesTool->ClearSelectedSegments();
+		SplinesTool->SelectAllControlPoints(*CurrentLandscapeInfo);
+
+		SplinesTool->UpdatePropertiesWindows();
+		GUnrealEd->RedrawLevelEditingViewports();
+	}
+}
+
+void FEdModeLandscape::SelectAllSplineSegments()
+{
+	ULandscapeInfo* CurrentLandscapeInfo = CurrentToolTarget.LandscapeInfo.Get();
+	if (SplinesTool && CurrentLandscapeInfo)
+	{
+		SplinesTool->ClearSelectedControlPoints();
+		SplinesTool->SelectAllSplineSegments(*CurrentLandscapeInfo);
+
+		SplinesTool->UpdatePropertiesWindows();
+		GUnrealEd->RedrawLevelEditingViewports();
+	}
+}
 
 void FEdModeLandscape::SelectAllConnectedSplineControlPoints()
 {
diff --git a/Engine/Source/Editor/LandscapeEditor/Private/LandscapeEditorDetailCustomization_MiscTools.cpp b/Engine/Source/Editor/LandscapeEditor/Private/LandscapeEditorDetailCustomization_MiscTools.cpp
--- a/Engine/Source/Editor/LandscapeEditor/Private/LandscapeEditorDetailCustomization_MiscTools.cpp
+++ b/Engine/Source/Editor/LandscapeEditor/Private/LandscapeEditorDetailCustomization_MiscTools.cpp
@@ -137,6 +137,39 @@ void FLandscapeEditorDetailCustomization_MiscTools::CustomizeDetails(IDetailLayo
 				.OnClicked_Static(&FLandscapeEditorDetailCustomization_MiscTools::OnApplySelectedSplinesButtonClicked)
 			]
 		];
+
+		ToolsCategory.AddCustomRow(LOCTEXT("SelectAllLabel", "Select all"))
+		[
+			SNew(SHorizontalBox)
+			+ SHorizontalBox::Slot()
+			.Padding(0, 6, 0, 0)
+			[
+				SNew(STextBlock)
+				.Font(DetailBuilder.GetDetailFont())
+				.ShadowOffset(FVector2D::UnitVector)
+				.Text(LOCTEXT("Spline.SelectAll", "Select All:"))
+			]
+		];
+		ToolsCategory.AddCustomRow(LOCTEXT("SelectAllLabel", "Select all"))
+		[
+			SNew(SHorizontalBox)
+			+ SHorizontalBox::Slot()
+			[
+				SNew(SButton)
+				.ToolTipText(LOCTEXT("Spline.ControlPoints.All.Tooltip", "Selects all landscape spline control points in the map."))
+				.Text(LOCTEXT("Spline.ControlPoints", "Control Points"))
+				.HAlign(HAlign_Center)
+				.OnClicked_Static(&FLandscapeEditorDetailCustomization_MiscTools::OnSelectAllControlPointsButtonClicked)
+			]
+			+ SHorizontalBox::Slot()
+			[
+				SNew(SButton)
+				.ToolTipText(LOCTEXT("Spline.ControlPoints.All.Tooltip", "Selects all landscape spline segments in the map."))
+				.Text(LOCTEXT("Spline.Segments", "Segments"))
+				.HAlign(HAlign_Center)
+				.OnClicked_Static(&FLandscapeEditorDetailCustomization_MiscTools::OnSelectAllSegmentsButtonClicked)
+			]
+		];
 		ToolsCategory.AddCustomRow(LOCTEXT("Spline.bUseAutoRotateControlPoint.Selected", "Use Auto Rotate Control Point"))
 		[
 			SNew(SHorizontalBox)
@@ -342,6 +375,28 @@ FReply FLandscapeEditorDetailCustomization_MiscTools::OnApplySelectedSplinesButt
 	return FReply::Handled();
 }
 
+FReply FLandscapeEditorDetailCustomization_MiscTools::OnSelectAllControlPointsButtonClicked()
+{
+	FEdModeLandscape* LandscapeEdMode = GetEditorMode();
+	if (LandscapeEdMode && LandscapeEdMode->CurrentToolTarget.LandscapeInfo.IsValid())
+	{
+		LandscapeEdMode->SelectAllSplineControlPoints();
+	}
+
+	return FReply::Handled();
+}
+
+FReply FLandscapeEditorDetailCustomization_MiscTools::OnSelectAllSegmentsButtonClicked()
+{
+	FEdModeLandscape* LandscapeEdMode = GetEditorMode();
+	if (LandscapeEdMode && LandscapeEdMode->CurrentToolTarget.LandscapeInfo.IsValid())
+	{
+		LandscapeEdMode->SelectAllSplineSegments();
+	}
+
+	return FReply::Handled();
+}
+
 void FLandscapeEditorDetailCustomization_MiscTools::OnbUseAutoRotateControlPointChanged(ECheckBoxState NewState)
 {
 	FEdModeLandscape* LandscapeEdMode = GetEditorMode();
diff --git a/Engine/Source/Editor/LandscapeEditor/Private/LandscapeEditorDetailCustomization_MiscTools.h b/Engine/Source/Editor/LandscapeEditor/Private/LandscapeEditorDetailCustomization_MiscTools.h
--- a/Engine/Source/Editor/LandscapeEditor/Private/LandscapeEditorDetailCustomization_MiscTools.h
+++ b/Engine/Source/Editor/LandscapeEditor/Private/LandscapeEditorDetailCustomization_MiscTools.h
@@ -35,6 +35,8 @@ protected:
 	// Splines Tool
 	static FReply OnApplyAllSplinesButtonClicked();
 	static FReply OnApplySelectedSplinesButtonClicked();
+	static FReply OnSelectAllControlPointsButtonClicked();
+	static FReply OnSelectAllSegmentsButtonClicked();
 	void OnbUseAutoRotateControlPointChanged(ECheckBoxState NewState);
 	ECheckBoxState GetbUseAutoRotateControlPoint() const;
 
diff --git a/Engine/Source/Editor/LandscapeEditor/Private/LandscapeEditorObject.cpp b/Engine/Source/Editor/LandscapeEditor/Private/LandscapeEditorObject.cpp
--- a/Engine/Source/Editor/LandscapeEditor/Private/LandscapeEditorObject.cpp
+++ b/Engine/Source/Editor/LandscapeEditor/Private/LandscapeEditorObject.cpp
@@ -662,10 +662,6 @@ void ULandscapeEditorObject::ImportLandscapeData()
 
 ELandscapeImportResult ULandscapeEditorObject::CreateImportLayersInfo(TArray<FLandscapeImportLayerInfo>& OutImportLayerInfos)
 {
-	const int32 QuadsPerComponent = NewLandscape_SectionsPerComponent * NewLandscape_QuadsPerSection;
-	const int32 SizeX = NewLandscape_ComponentCount.X * QuadsPerComponent + 1;
-	const int32 SizeY = NewLandscape_ComponentCount.Y * QuadsPerComponent + 1;
-
 	const uint32 ImportSizeX = ImportLandscape_Width;
 	const uint32 ImportSizeY = ImportLandscape_Height;
 
