Adds Default Landscape Material to be used in LandscapeEditor when creating a new landscape

diff --git a/Engine/Source/Editor/LandscapeEditor/Private/LandscapeEditorObject.cpp b/Engine/Source/Editor/LandscapeEditor/Private/LandscapeEditorObject.cpp
--- a/Engine/Source/Editor/LandscapeEditor/Private/LandscapeEditorObject.cpp
+++ b/Engine/Source/Editor/LandscapeEditor/Private/LandscapeEditorObject.cpp
@@ -8,6 +8,7 @@
 #include "UObject/ConstructorHelpers.h"
 #include "LandscapeEditorModule.h"
 #include "LandscapeRender.h"
+#include "LandscapeSettings.h"
 #include "LandscapeImportHelper.h"
 #include "LandscapeMaterialInstanceConstant.h"
 #include "Misc/ConfigCacheIni.h"
@@ -316,9 +317,33 @@ void ULandscapeEditorObject::Load()
 		}
 	}
 
-	FString NewLandscapeMaterialName = (NewLandscape_Material != NULL) ? NewLandscape_Material->GetPathName() : FString();
-	GConfig->GetString(TEXT("LandscapeEdit"), TEXT("NewLandscapeMaterialName"), NewLandscapeMaterialName, GEditorPerProjectIni);
-	if(NewLandscapeMaterialName != TEXT(""))
+	FString NewLandscapeMaterialName;
+
+	// If NewLandscape_Material is not null, we will try to use it
+	if (!NewLandscape_Material.IsExplicitlyNull())
+	{
+		NewLandscapeMaterialName = NewLandscape_Material->GetPathName();
+	}
+	else
+	{
+		// If this project already has a saved NewLandscapeMaterialName, we use it
+		GConfig->GetString(TEXT("LandscapeEdit"), TEXT("NewLandscapeMaterialName"), NewLandscapeMaterialName, GEditorPerProjectIni);
+
+		if (NewLandscapeMaterialName.IsEmpty())
+		{
+			/* Project does not have a saved NewLandscapeMaterialNameand and NewLandscape_Material is not already assigned;
+			 * we fallback to the DefaultLandscapeMaterial for the project, if set */
+			const ULandscapeSettings* Settings = GetDefault<ULandscapeSettings>();
+			TSoftObjectPtr<UMaterialInterface> DefaultMaterial = Settings->GetDefaultLandscapeMaterial();
+
+			if (!DefaultMaterial.IsNull())
+			{
+				NewLandscapeMaterialName = DefaultMaterial.ToString();
+			}
+		}
+	}
+	
+	if (!NewLandscapeMaterialName.IsEmpty())
 	{
 		NewLandscape_Material = LoadObject<UMaterialInterface>(NULL, *NewLandscapeMaterialName, NULL, LOAD_NoWarn);
 	}
diff --git a/Engine/Source/Runtime/Landscape/Public/LandscapeSettings.h b/Engine/Source/Runtime/Landscape/Public/LandscapeSettings.h
--- a/Engine/Source/Runtime/Landscape/Public/LandscapeSettings.h
+++ b/Engine/Source/Runtime/Landscape/Public/LandscapeSettings.h
@@ -23,6 +23,9 @@ public:
 
 	int32 GetSideResolutionLimit() const { return SideResolutionLimit; }
 
+	/** Returns the default landscape material that should be used when creating a new landscape. */
+	TSoftObjectPtr<UMaterialInterface> GetDefaultLandscapeMaterial() const { return DefaultLandscapeMaterial; }
+
 public:
 	UPROPERTY(config, EditAnywhere, Category = "Layers", meta=(UIMin = "1", UIMax = "32", ClampMin = "1", ClampMax = "32", ToolTip = "This option controls the maximum editing layers that can be added to a Landscape"))
 	int32 MaxNumberOfLayers = 8;
@@ -31,6 +34,9 @@ protected:
 	UPROPERTY(config)
 	int32 SideResolutionLimit = 2048;
 
+	UPROPERTY(EditAnywhere, config, Category = "Materials", meta = (ToolTip = "Default Landscape Material will be prefilled when creating a new landscape."))
+	TSoftObjectPtr<UMaterialInterface> DefaultLandscapeMaterial;
+
 	UPROPERTY(transient)
 	bool bRestrictiveMode = false;
 };
