Enables EditLayer by default in LandscapeEditor when creating a new landscape.

diff --git a/Engine/Source/Editor/LandscapeEditor/Public/LandscapeEditorObject.h b/Engine/Source/Editor/LandscapeEditor/Public/LandscapeEditorObject.h
--- a/Engine/Source/Editor/LandscapeEditor/Public/LandscapeEditorObject.h
+++ b/Engine/Source/Editor/LandscapeEditor/Public/LandscapeEditorObject.h
@@ -569,7 +569,7 @@ private:
 	TArray<uint16> ImportLandscape_Data;
 public:
 	UPROPERTY(Category = "New Landscape", EditAnywhere, NonTransactional, meta = (DisplayName= "Enable Edit Layers", ToolTip="Enable support for landscape edit layers.", ShowForTools= "NewLandscape"))
-	bool bCanHaveLayersContent = false;
+	bool bCanHaveLayersContent = true;
 
 	UPROPERTY(Category = "New Landscape", EditAnywhere, NonTransactional, meta = (DisplayName = "Flip Y Axis", ToolTip = "Whether to flip Y coordinate of imported files.", ShowForTools = "NewLandscape,ImportExport"))
 	bool bFlipYAxis = false;
