Filters out cooked textures from AlphaBrush browser.

diff --git a/Engine/Source/Editor/LandscapeEditor/Private/LandscapeEditorDetailCustomization_AlphaBrush.cpp b/Engine/Source/Editor/LandscapeEditor/Private/LandscapeEditorDetailCustomization_AlphaBrush.cpp
--- a/Engine/Source/Editor/LandscapeEditor/Private/LandscapeEditorDetailCustomization_AlphaBrush.cpp
+++ b/Engine/Source/Editor/LandscapeEditor/Private/LandscapeEditorDetailCustomization_AlphaBrush.cpp
@@ -224,6 +224,11 @@ void FLandscapeEditorDetailCustomization_AlphaBrush::CustomizeDetails(IDetailLay
 				.IsEnabled(true)
 				.PropertyHandle(PropertyHandle_AlphaTexture)
 				.AllowedClass(UTexture2D::StaticClass())
+				.OnShouldFilterAsset_Lambda([](const FAssetData& AssetData) -> bool
+				{
+					// We cannot use cooked texture as parameter for now.
+					return (AssetData.PackageFlags & PKG_Cooked) != 0;
+				})
 				.AllowClear(false)
 			]
 		]
