Adds source validity check when filtering texture usable in alpha brush.

diff --git a/Engine/Source/Editor/LandscapeEditor/Private/LandscapeEditorDetailCustomization_AlphaBrush.cpp b/Engine/Source/Editor/LandscapeEditor/Private/LandscapeEditorDetailCustomization_AlphaBrush.cpp
--- a/Engine/Source/Editor/LandscapeEditor/Private/LandscapeEditorDetailCustomization_AlphaBrush.cpp
+++ b/Engine/Source/Editor/LandscapeEditor/Private/LandscapeEditorDetailCustomization_AlphaBrush.cpp
@@ -227,7 +227,14 @@ void FLandscapeEditorDetailCustomization_AlphaBrush::CustomizeDetails(IDetailLay
 				.OnShouldFilterAsset_Lambda([](const FAssetData& AssetData) -> bool
 				{
 					// We cannot use cooked texture as parameter for now.
-					return (AssetData.PackageFlags & PKG_Cooked) != 0;
+					if ((AssetData.PackageFlags & PKG_Cooked) != 0)
+					{
+						return true;
+					}
+
+					const UTexture2D* Texture = Cast<UTexture2D>(AssetData.GetAsset());
+
+					return (Texture == nullptr) || !Texture->Source.IsValid();
 				})
 				.AllowClear(false)
 			]
