Fixed "Use target value" in the landscape UX being displayed for sculpt tools while it's only valid in paint mode

diff --git a/Engine/Source/Editor/LandscapeEditor/Public/LandscapeEditorObject.h b/Engine/Source/Editor/LandscapeEditor/Public/LandscapeEditorObject.h
--- a/Engine/Source/Editor/LandscapeEditor/Public/LandscapeEditorObject.h
+++ b/Engine/Source/Editor/LandscapeEditor/Public/LandscapeEditorObject.h
@@ -281,7 +281,7 @@ class ULandscapeEditorObject : public UObject
 	bool bUseWeightTargetValue;
 
 	// Enable to make tools blend towards a target value
-	UPROPERTY(Category="Tool Settings", EditAnywhere, NonTransactional, meta=(DisplayName="Use Target Value", EditCondition="bUseWeightTargetValue", ShowForTools="Paint,Sculpt,Noise", ClampMin="0", ClampMax="10", UIMin="0", UIMax="1"))
+	UPROPERTY(Category="Tool Settings", EditAnywhere, NonTransactional, meta=(DisplayName="Use Target Value", EditCondition="bUseWeightTargetValue", ShowForTools="Paint,Sculpt,Noise", ShowForTargetTypes = "Weightmap", ClampMin="0", ClampMax="10", UIMin="0", UIMax="1"))
 	float WeightTargetValue;
 
 	// I have no idea what this is for but it's used by the noise and erosion tools, and isn't exposed to the UI
