Fixed landscape tools using ShowForMode not diplaying properties when in weightmap mode with no active paint layer (Use Target Value and Painting Restriction)

diff --git a/Engine/Source/Editor/LandscapeEditor/Private/SLandscapeEditor.cpp b/Engine/Source/Editor/LandscapeEditor/Private/SLandscapeEditor.cpp
--- a/Engine/Source/Editor/LandscapeEditor/Private/SLandscapeEditor.cpp
+++ b/Engine/Source/Editor/LandscapeEditor/Private/SLandscapeEditor.cpp
@@ -609,8 +609,16 @@ bool FLandscapeToolKit::GetIsPropertyVisibleFromProperty(const FProperty& Proper
 			Property.GetMetaData("ShowForTargetTypes").ParseIntoArray(ShowForTargetTypes, TEXT(","), true);
 
 			const ELandscapeToolTargetType::Type CurrentTargetType = LandscapeEdMode->CurrentToolTarget.TargetType;
-			if (CurrentTargetType == ELandscapeToolTargetType::Invalid ||
-				ShowForTargetTypes.FindByKey(TargetTypeNames[CurrentTargetType]) == nullptr)
+			// ELandscapeToolTargetType::Invalid means "weightmap with no valid paint layer" so we still want to display that property if it has been marked to be displayed in Weightmap target type, to be consistent 
+			//  with other paint brush properties (that don't use ShowForTargetTypes), which are still displayed in that case, even if they are ineffective :
+			if ((CurrentTargetType == ELandscapeToolTargetType::Invalid) 
+				&& (ShowForTargetTypes.FindByKey(TargetTypeNames[ELandscapeToolTargetType::Weightmap]) != nullptr))
+			{ 
+				return true;
+			}
+			// Otherwise, hide it, if ShowForTargetTypes was used on this property but doesn't correspond to the current target type :
+			else if ((CurrentTargetType == ELandscapeToolTargetType::Invalid)
+				|| (ShowForTargetTypes.FindByKey(TargetTypeNames[CurrentTargetType]) == nullptr))
 			{
 				return false;
 			}
