* Fixed data caching function used in Smooth and flatten tool when using combined layers
* Added visual logs for when the cache data is accessed and when the cache is expanded
* These operations being expensive, align the cache area on whole landscape to avoid expanding the cache when expanding the cached area within the limits of an already sampled landscape component

diff --git a/Engine/Source/Editor/LandscapeEditor/Private/LandscapeEdMode.cpp b/Engine/Source/Editor/LandscapeEditor/Private/LandscapeEdMode.cpp
--- a/Engine/Source/Editor/LandscapeEditor/Private/LandscapeEdMode.cpp
+++ b/Engine/Source/Editor/LandscapeEditor/Private/LandscapeEdMode.cpp
@@ -3053,7 +3053,7 @@ void FEdModeLandscape::ImportHeightData(ULandscapeInfo* LandscapeInfo, const FGu
 		FScopedTransaction Transaction(LOCTEXT("Undo_ImportHeightmap", "Importing Landscape Heightmap"));
 
 		FHeightmapAccessor<false> HeightmapAccessor(LandscapeInfo);
-		HeightmapAccessor.SetData(MinX, MinY, MaxX, MaxY, Data.GetData(), PaintRestriction);
+		HeightmapAccessor.SetData(MinX, MinY, MaxX, MaxY, Data.GetData());
 	});
 }
 
diff --git a/Engine/Source/Editor/LandscapeEditor/Private/LandscapeEdModeTools.h b/Engine/Source/Editor/LandscapeEditor/Private/LandscapeEdModeTools.h
--- a/Engine/Source/Editor/LandscapeEditor/Private/LandscapeEdModeTools.h
+++ b/Engine/Source/Editor/LandscapeEditor/Private/LandscapeEdModeTools.h
@@ -16,6 +16,8 @@
 #include "AI/NavigationSystemBase.h"
 #include "Landscape.h"
 #include "Logging/LogMacros.h"
+#include "VisualLogger/VisualLogger.h"
+
 
 DECLARE_LOG_CATEGORY_EXTERN(LogLandscapeTools, Log, All);
 
@@ -25,9 +27,9 @@ DECLARE_LOG_CATEGORY_EXTERN(LogLandscapeTools, Log, All);
 //
 struct FNoiseParameter
 {
-	float	Base,
-		NoiseScale,
-		NoiseAmount;
+	float Base;
+	float NoiseScale;
+	float NoiseAmount;
 
 	// Constructors.
 
@@ -288,14 +290,16 @@ public:
 
 	TLandscapeEditCache(const FLandscapeToolTarget& InTarget)
 		: DataAccess(InTarget)
-		, Valid(false)
+		, LandscapeInfo(InTarget.LandscapeInfo)
 	{
+		check(LandscapeInfo != nullptr);
 	}
 
 	// X2/Y2 Coordinates are "inclusive" max values
+	// Note that this should maybe be called "ExtendDataCache" because the region here will be combined with the existing cached region, not loaded independently, giving a cached region that is the bounding box of previous and new
 	void CacheData(int32 X1, int32 Y1, int32 X2, int32 Y2, bool bCacheOriginalData = false)
 	{
-		if (!Valid)
+		if (!bIsValid)
 		{
 			if (Accessor::bUseInterp)
 			{
@@ -307,7 +311,7 @@ public:
 				DataAccess.GetData(ValidX1, ValidY1, ValidX2, ValidY2, CachedData);
 				if (!ensureMsgf(ValidX1 <= ValidX2 && ValidY1 <= ValidY2, TEXT("Invalid cache area: X(%d-%d), Y(%d-%d) from region X(%d-%d), Y(%d-%d)"), ValidX1, ValidX2, ValidY1, ValidY2, X1, X2, Y1, Y2))
 				{
-					Valid = false;
+					bIsValid = false;
 					return;
 				}
 			}
@@ -321,14 +325,19 @@ public:
 				DataAccess.GetDataFast(CachedX1, CachedY1, CachedX2, CachedY2, CachedData);
 			}
 
+			// Drop a visual log to indicate the area covered by this cache region extension :
+			VisualizeLandscapeRegion(CachedX1, CachedY1, CachedX2, CachedY2, FColor::Red, TEXT("Cache Data"));
+
 			if (bCacheOriginalData)
 			{
 				OriginalData = CachedData;
 			}
-			Valid = true;
+			bIsValid = true;
 		}
 		else
 		{
+			bool bCacheExtended = false;
+
 			// Extend the cache area if needed
 			if (X1 < CachedX1)
 			{
@@ -352,6 +361,8 @@ public:
 					CacheOriginalData(X1, CachedY1, CachedX1 - 1, CachedY2);
 				}
 				CachedX1 = X1;
+
+				bCacheExtended = true;
 			}
 
 			if (X2 > CachedX2)
@@ -375,6 +386,8 @@ public:
 					CacheOriginalData(CachedX2 + 1, CachedY1, X2, CachedY2);
 				}
 				CachedX2 = X2;
+
+				bCacheExtended = true;
 			}
 
 			if (Y1 < CachedY1)
@@ -398,6 +411,8 @@ public:
 					CacheOriginalData(CachedX1, Y1, CachedX2, CachedY1 - 1);
 				}
 				CachedY1 = Y1;
+
+				bCacheExtended = true;
 			}
 
 			if (Y2 > CachedY2)
@@ -421,6 +436,14 @@ public:
 					CacheOriginalData(CachedX1, CachedY2 + 1, CachedX2, Y2);
 				}
 				CachedY2 = Y2;
+
+				bCacheExtended = true;
+			}
+
+			if (bCacheExtended)
+			{
+				// Drop a visual log to indicate the area covered by this cache region extension :
+				VisualizeLandscapeRegion(CachedX1, CachedY1, CachedX2, CachedY2, FColor::Red, TEXT("Cache Data"));
 			}
 		}
 	}
@@ -501,17 +524,26 @@ public:
 
 	bool HasCachedData(int32 X1, int32 Y1, int32 X2, int32 Y2) const
 	{
-		return (Valid && X1 >= CachedX1 && Y1 >= CachedY1 && X2 <= CachedX2 && Y2 <= CachedY2);
+		return (bIsValid && X1 >= CachedX1 && Y1 >= CachedY1 && X2 <= CachedX2 && Y2 <= CachedY2);
 	}
-
-	template<typename TGetCacheRegionFunction>
-	bool GetDataAndCache(int32 X1, int32 Y1, int32 X2, int32 Y2, TArray<AccessorType>& OutData, TGetCacheRegionFunction GetCacheRegion)
+	
+	using FPrepareRegionForCachingFunction = TFunction<FIntRect(const FIntRect& NewCacheBounds)>;
+	bool GetDataAndCache(int32 X1, int32 Y1, int32 X2, int32 Y2, TArray<AccessorType>& OutData, FPrepareRegionForCachingFunction PrepareRegionForCaching)
 	{
+		// Drop a visual log to indicate the area requested by this data access :
+		VisualizeLandscapeRegion(X1, Y1, X2, Y2, FColor::Blue, TEXT("CacheDataRequest"));
+
 		if (!HasCachedData(X1, Y1, X2, Y2))
 		{
-			FIntRect Bounds = GetCacheRegion();
-			check((Bounds.Min.X <= X1) && (Bounds.Min.Y <= Y1) && (Bounds.Max.X >= X2) && (Bounds.Max.Y >= Y2));
-			CacheData(Bounds.Min.X, Bounds.Min.Y, Bounds.Max.X, Bounds.Max.Y);
+			// The cache needs to be expanded, compute the new bounds : 
+			// The bounds we calculate here need to be what would be the result of calling CacheData with this region, meaning that they should include the previous bounds. This will let us pass the correct region of interest to PrepareRegionForCaching
+			FIntRect NewCacheBounds(bIsValid ? FMath::Min(X1, CachedX1) : X1, bIsValid ? FMath::Min(Y1, CachedY1) : Y1, bIsValid ? FMath::Max(X2, CachedX2) : X2, bIsValid ? FMath::Max(Y2, CachedY2) : Y2);
+
+			// The caller might request a cache region that is actually larger than the data he wants to sample for this particular read (e.g. to avoid re-caching when doing expensive samples like layer-collapsing ones) : 
+			NewCacheBounds = PrepareRegionForCaching(NewCacheBounds);
+			check((NewCacheBounds.Min.X <= X1) && (NewCacheBounds.Min.Y <= Y1) && (NewCacheBounds.Max.X >= X2) && (NewCacheBounds.Max.Y >= Y2));
+
+			CacheData(NewCacheBounds.Min.X, NewCacheBounds.Min.Y, NewCacheBounds.Max.X, NewCacheBounds.Max.Y);
 		}
 		ensure(HasCachedData(X1, Y1, X2, Y2));
 		return GetCachedData(X1, Y1, X2, Y2, OutData);
@@ -618,18 +650,32 @@ private:
 		}
 	}
 
+	void VisualizeLandscapeRegion(int32 InX1, int32 InY1, int32 InX2, int32 InY2, const FColor& InColor, const FString& InDescription)
+	{
+		check(LandscapeInfo != nullptr);
+		const FTransform& LandscapeTransform = LandscapeInfo->LandscapeActor->GetTransform();
+		FVector Min = LandscapeTransform.TransformPosition(FVector(InX1, InY1, 0));
+		FVector Max = LandscapeTransform.TransformPosition(FVector(InX2, InY2, 0));
+		UE_VLOG_BOX(LandscapeInfo->LandscapeActor.Get(), LogLandscapeTools, Log, FBox(Min, Max), InColor, TEXT("%s"), *InDescription);
+	}
+
 	TMap<FIntPoint, AccessorType> CachedData;
 	TMap<FIntPoint, AccessorType> OriginalData;
+	// Keep the landscape info for visual logging purposes :
+	TWeakObjectPtr<ULandscapeInfo> LandscapeInfo;
 
-	bool Valid;
+	bool bIsValid = false;
 
-	int32 CachedX1;
-	int32 CachedY1;
-	int32 CachedX2;
-	int32 CachedY2;
+	int32 CachedX1 = INDEX_NONE;
+	int32 CachedY1 = INDEX_NONE;
+	int32 CachedX2 = INDEX_NONE;
+	int32 CachedY2 = INDEX_NONE;
 
 	// To store valid region....
-	int32 ValidX1, ValidX2, ValidY1, ValidY2;
+	int32 ValidX1 = INDEX_NONE;
+	int32 ValidX2 = INDEX_NONE;
+	int32 ValidY1 = INDEX_NONE;
+	int32 ValidY2 = INDEX_NONE;
 };
 
 template<bool bInUseInterp>
@@ -910,40 +956,37 @@ public:
 				NewLayerVisibility.Add((i > EditingLayerIndex) ? false : CurrentLayer->bVisible);
 			}
 
-			FIntRect Bounds;
-			const float BoundStep = LandscapeInfo->ComponentSizeQuads / 4;
-			Bounds.Min.X = (FMath::FloorToInt((float)X1 / BoundStep) - 1) * BoundStep;
-			Bounds.Min.Y = (FMath::FloorToInt((float)Y1 / BoundStep) - 1) * BoundStep;
-			Bounds.Max.X = (FMath::CeilToInt((float)X2 / BoundStep) + 1) * BoundStep;
-			Bounds.Max.Y = (FMath::CeilToInt((float)Y2 / BoundStep) + 1) * BoundStep;
-
-			FScopedSetLandscapeEditingLayer Scope(Landscape, FGuid());
-			CacheUpToEditingLayer.GetDataAndCache(X1, Y1, X2, Y2, Data, [&]() -> FIntRect
+			auto OnCacheUpdating = [&](const FIntRect& NewCacheBounds) -> FIntRect
 			{
+				// This function is triggered when the cache needs expanding. We'll ask for a larger area so that we don't need to re-cache (a GPU-synchronizing operation) every time
+				//  we need an additional row/column. Aligning the cache region on components borders is appropriate since this is what gets rendered. This avoids
+				//  re-caching when we've already sampled a component: 
+				FIntRect DesiredCacheBounds;
+				DesiredCacheBounds.Min.X = (FMath::FloorToInt((float)NewCacheBounds.Min.X / LandscapeInfo->ComponentSizeQuads)) * LandscapeInfo->ComponentSizeQuads;
+				DesiredCacheBounds.Min.Y = (FMath::FloorToInt((float)NewCacheBounds.Min.Y / LandscapeInfo->ComponentSizeQuads)) * LandscapeInfo->ComponentSizeQuads;
+				DesiredCacheBounds.Max.X = (FMath::CeilToInt((float)NewCacheBounds.Max.X / LandscapeInfo->ComponentSizeQuads)) * LandscapeInfo->ComponentSizeQuads;
+				DesiredCacheBounds.Max.Y = (FMath::CeilToInt((float)NewCacheBounds.Max.Y / LandscapeInfo->ComponentSizeQuads)) * LandscapeInfo->ComponentSizeQuads;
+
 				TSet<ULandscapeComponent*> AffectedComponents;
-				LandscapeInfo->GetComponentsInRegion(Bounds.Min.X, Bounds.Min.Y, Bounds.Max.X, Bounds.Max.Y, AffectedComponents);
+				LandscapeInfo->GetComponentsInRegion(DesiredCacheBounds.Min.X, DesiredCacheBounds.Min.Y, DesiredCacheBounds.Max.X, DesiredCacheBounds.Max.Y, AffectedComponents);
 				SynchronousUpdateComponentVisibilityForHeight(AffectedComponents, NewLayerVisibility);
+
 				if (bVisibilityChanged)
 				{
 					VisibilityChangedComponents.Append(AffectedComponents);
 				}
-				return Bounds;
-			});
+
+				return DesiredCacheBounds;
+			};
+
+			FScopedSetLandscapeEditingLayer Scope(Landscape, FGuid());
+			CacheUpToEditingLayer.GetDataAndCache(X1, Y1, X2, Y2, Data, OnCacheUpdating);
 			// Release Texture Mips that will be Locked by the next SynchronousUpdateComponentVisibilityForHeight
 			CacheUpToEditingLayer.DataAccess.Flush();
 
-			CacheBottomLayers.GetDataAndCache(X1, Y1, X2, Y2, BottomLayersData, [&]() -> FIntRect
-			{
-				NewLayerVisibility[EditingLayerIndex] = false;
-				TSet<ULandscapeComponent*> AffectedComponents;
-				LandscapeInfo->GetComponentsInRegion(Bounds.Min.X, Bounds.Min.Y, Bounds.Max.X, Bounds.Max.Y, AffectedComponents);
-				SynchronousUpdateComponentVisibilityForHeight(AffectedComponents, NewLayerVisibility);
-				if (bVisibilityChanged)
-				{
-					VisibilityChangedComponents.Append(AffectedComponents);
-				}
-				return Bounds;
-			});
+			// Now turn off visibility on the current layer in order to have the data of all bottom layers except the current one
+			NewLayerVisibility[EditingLayerIndex] = false;
+			CacheBottomLayers.GetDataAndCache(X1, Y1, X2, Y2, BottomLayersData, OnCacheUpdating);
 			// Do the same here for consistency
 			CacheBottomLayers.DataAccess.Flush();
 		}
diff --git a/Engine/Source/Runtime/Landscape/Private/LandscapeEdit.cpp b/Engine/Source/Runtime/Landscape/Private/LandscapeEdit.cpp
--- a/Engine/Source/Runtime/Landscape/Private/LandscapeEdit.cpp
+++ b/Engine/Source/Runtime/Landscape/Private/LandscapeEdit.cpp
@@ -2402,8 +2402,6 @@ bool ULandscapeInfo::AreAllComponentsRegistered() const
 	return true;
 }
 
-#define MAX_LANDSCAPE_SUBSECTIONS 2
-
 bool ULandscapeInfo::HasUnloadedComponentsInRegion(int32 X1, int32 Y1, int32 X2, int32 Y2) const
 {
 	bool bResult = false;
