* Made landscape edit layers read-only, removed related context menu and prevent layer reordering in non edit layers- related landscape tools (Component Add/Delete/Move/Select,Mask,New/Resize Landscape)
* Don't use a modal error message box when the currently selected edit layer is locked in those tools (since they're not edit layers-related)

diff --git a/Engine/Source/Editor/LandscapeEditor/Private/LandscapeEdMode.cpp b/Engine/Source/Editor/LandscapeEditor/Private/LandscapeEdMode.cpp
--- a/Engine/Source/Editor/LandscapeEditor/Private/LandscapeEdMode.cpp
+++ b/Engine/Source/Editor/LandscapeEditor/Private/LandscapeEdMode.cpp
@@ -179,14 +179,14 @@ FEdModeLandscape::~FEdModeLandscape()
 
 	// Clean up Debug Materials
 	FlushRenderingCommands();
-	GLayerDebugColorMaterial = NULL;
-	GSelectionColorMaterial = NULL;
-	GSelectionRegionMaterial = NULL;
-	GMaskRegionMaterial = NULL;
-	GColorMaskRegionMaterial = NULL;
-	GLandscapeBlackTexture = NULL;
-	GLandscapeLayerUsageMaterial = NULL;
-	GLandscapeDirtyMaterial = NULL;
+	GLayerDebugColorMaterial = nullptr;
+	GSelectionColorMaterial = nullptr;
+	GSelectionRegionMaterial = nullptr;
+	GMaskRegionMaterial = nullptr;
+	GColorMaskRegionMaterial = nullptr;
+	GLandscapeBlackTexture = nullptr;
+	GLandscapeLayerUsageMaterial = nullptr;
+	GLandscapeDirtyMaterial = nullptr;
 
 	InteractorPainting = nullptr;
 }
@@ -391,6 +391,11 @@ bool FEdModeLandscape::IsLandscapeResolutionCompliant() const
 	return true;
 }
 
+bool FEdModeLandscape::DoesCurrentToolAffectEditLayers() const
+{
+	return (CurrentTool != nullptr) ? CurrentTool->AffectsEditLayers() : false;
+}
+
 FText FEdModeLandscape::GetLandscapeResolutionErrorText() const
 {
 	const TObjectPtr<const ULandscapeSettings> Settings = GetDefault<ULandscapeSettings>();
@@ -686,7 +691,7 @@ void FEdModeLandscape::Exit()
 		CurrentTool->PreviousBrushIndex = CurrentBrushSetIndex;
 		CurrentTool->ExitTool();
 	}
-	CurrentTool = NULL;
+	CurrentTool = nullptr;
 	// Leave CurrentToolIndex set so we can restore the active tool on re-opening the landscape editor
 
 	LandscapeList.Empty();
@@ -699,7 +704,7 @@ void FEdModeLandscape::Exit()
 	GLandscapeEditRenderMode = ELandscapeEditRenderMode::None;
 	GLandscapeEditModeActive = false;
 
-	CurrentGizmoActor = NULL;
+	CurrentGizmoActor = nullptr;
 
 	GEditor->SelectNone(false, true);
 
@@ -795,7 +800,7 @@ void FEdModeLandscape::Tick(FEditorViewportClient* ViewportClient, float DeltaTi
 		{
 			ALandscapeProxy* LandscapeProxy = CurrentToolTarget.LandscapeInfo->GetLandscapeProxy();
 			
-			if (LandscapeProxy == NULL ||
+			if (LandscapeProxy == nullptr ||
 				LandscapeProxy->GetLandscapeMaterial() != CachedLandscapeMaterial)
 			{
 				UpdateTargetList();
@@ -1172,7 +1177,7 @@ EEditAction::Type FEdModeLandscape::GetActionEditDuplicate()
 
 	if (NewLandscapePreviewMode == ENewLandscapePreviewMode::None)
 	{
-		if (CurrentTool != NULL)
+		if (CurrentTool != nullptr)
 		{
 			Result = CurrentTool->GetActionEditDuplicate();
 		}
@@ -1187,7 +1192,7 @@ EEditAction::Type FEdModeLandscape::GetActionEditDelete()
 
 	if (NewLandscapePreviewMode == ENewLandscapePreviewMode::None)
 	{
-		if (CurrentTool != NULL)
+		if (CurrentTool != nullptr)
 		{
 			Result = CurrentTool->GetActionEditDelete();
 		}
@@ -1222,7 +1227,7 @@ EEditAction::Type FEdModeLandscape::GetActionEditCut()
 
 	if (NewLandscapePreviewMode == ENewLandscapePreviewMode::None)
 	{
-		if (CurrentTool != NULL)
+		if (CurrentTool != nullptr)
 		{
 			Result = CurrentTool->GetActionEditCut();
 		}
@@ -1244,7 +1249,7 @@ EEditAction::Type FEdModeLandscape::GetActionEditCopy()
 
 	if (NewLandscapePreviewMode == ENewLandscapePreviewMode::None)
 	{
-		if (CurrentTool != NULL)
+		if (CurrentTool != nullptr)
 		{
 			Result = CurrentTool->GetActionEditCopy();
 		}
@@ -1270,7 +1275,7 @@ EEditAction::Type FEdModeLandscape::GetActionEditPaste()
 
 	if (NewLandscapePreviewMode == ENewLandscapePreviewMode::None)
 	{
-		if (CurrentTool != NULL)
+		if (CurrentTool != nullptr)
 		{
 			Result = CurrentTool->GetActionEditPaste();
 		}
@@ -1301,7 +1306,7 @@ bool FEdModeLandscape::ProcessEditDuplicate()
 
 	if (NewLandscapePreviewMode == ENewLandscapePreviewMode::None)
 	{
-		if (CurrentTool != NULL)
+		if (CurrentTool != nullptr)
 		{
 			Result = CurrentTool->ProcessEditDuplicate();
 		}
@@ -1321,7 +1326,7 @@ bool FEdModeLandscape::ProcessEditDelete()
 
 	if (NewLandscapePreviewMode == ENewLandscapePreviewMode::None)
 	{
-		if (CurrentTool != NULL)
+		if (CurrentTool != nullptr)
 		{
 			Result = CurrentTool->ProcessEditDelete();
 		}
@@ -1341,7 +1346,7 @@ bool FEdModeLandscape::ProcessEditCut()
 
 	if (NewLandscapePreviewMode == ENewLandscapePreviewMode::None)
 	{
-		if (CurrentTool != NULL)
+		if (CurrentTool != nullptr)
 		{
 			Result = CurrentTool->ProcessEditCut();
 		}
@@ -1361,7 +1366,7 @@ bool FEdModeLandscape::ProcessEditCopy()
 
 	if (NewLandscapePreviewMode == ENewLandscapePreviewMode::None)
 	{
-		if (CurrentTool != NULL)
+		if (CurrentTool != nullptr)
 		{
 			Result = CurrentTool->ProcessEditCopy();
 		}
@@ -1430,7 +1435,7 @@ bool FEdModeLandscape::ProcessEditPaste()
 	
 	if (NewLandscapePreviewMode == ENewLandscapePreviewMode::None)
 	{
-		if (CurrentTool != NULL)
+		if (CurrentTool != nullptr)
 		{
 			Result = CurrentTool->ProcessEditPaste();
 		}
@@ -1718,7 +1723,7 @@ bool FEdModeLandscape::InputKey(FEditorViewportClient* ViewportClient, FViewport
 					FVector HitLocation;
 					if (LandscapeMouseTrace(ViewportClient, HitLocation))
 					{
-						if (!CanEditLayer(&ErrorReasonOnMouseUp))
+						if (CurrentTool->AffectsEditLayers() && !CanEditLayer(&ErrorReasonOnMouseUp))
 						{
 							return true;
 						}
@@ -1816,7 +1821,7 @@ bool FEdModeLandscape::InputDelta(FEditorViewportClient* InViewportClient, FView
 
 void FEdModeLandscape::SetCurrentToolMode(FName ToolModeName, bool bRestoreCurrentTool /*= true*/)
 {
-	if (CurrentToolMode == NULL || ToolModeName != CurrentToolMode->ToolModeName)
+	if (CurrentToolMode == nullptr || ToolModeName != CurrentToolMode->ToolModeName)
 	{
 		for (int32 i = 0; i < LandscapeToolModes.Num(); ++i)
 		{
@@ -2077,7 +2082,7 @@ int32 FEdModeLandscape::UpdateLandscapeList()
 
 	if (!CurrentGizmoActor.IsValid())
 	{
-		ALandscapeGizmoActiveActor* GizmoActor = NULL;
+		ALandscapeGizmoActiveActor* GizmoActor = nullptr;
 		for (TActorIterator<ALandscapeGizmoActiveActor> It(GetWorld()); It; ++It)
 		{
 			GizmoActor = *It;
@@ -2295,8 +2300,8 @@ void FEdModeLandscape::UpdateTargetList()
 			}
 
 			// Add layers
-			UTexture2D* ThumbnailWeightmap = NULL;
-			UTexture2D* ThumbnailHeightmap = NULL;
+			UTexture2D* ThumbnailWeightmap = nullptr;
+			UTexture2D* ThumbnailHeightmap = nullptr;
 
 			TargetLayerStartingIndex = LandscapeTargetList.Num();
 
@@ -2321,15 +2326,15 @@ void FEdModeLandscape::UpdateTargetList()
 				}
 
 				// Ensure thumbnails are up valid
-				if (LayerSettings.ThumbnailMIC == NULL)
+				if (LayerSettings.ThumbnailMIC == nullptr)
 				{
-					if (ThumbnailWeightmap == NULL)
+					if (ThumbnailWeightmap == nullptr)
 					{
-						ThumbnailWeightmap = LoadObject<UTexture2D>(NULL, TEXT("/Engine/EditorLandscapeResources/LandscapeThumbnailWeightmap.LandscapeThumbnailWeightmap"), NULL, LOAD_None, NULL);
+						ThumbnailWeightmap = LoadObject<UTexture2D>(nullptr, TEXT("/Engine/EditorLandscapeResources/LandscapeThumbnailWeightmap.LandscapeThumbnailWeightmap"), nullptr, LOAD_None, nullptr);
 					}
-					if (ThumbnailHeightmap == NULL)
+					if (ThumbnailHeightmap == nullptr)
 					{
-						ThumbnailHeightmap = LoadObject<UTexture2D>(NULL, TEXT("/Engine/EditorLandscapeResources/LandscapeThumbnailHeightmap.LandscapeThumbnailHeightmap"), NULL, LOAD_None, NULL);
+						ThumbnailHeightmap = LoadObject<UTexture2D>(nullptr, TEXT("/Engine/EditorLandscapeResources/LandscapeThumbnailHeightmap.LandscapeThumbnailHeightmap"), nullptr, LOAD_None, nullptr);
 					}
 
 					// Construct Thumbnail MIC
@@ -2683,8 +2688,8 @@ void FEdModeLandscape::HandleLevelsChanged(bool ShouldExitMode)
 void FEdModeLandscape::OnMaterialCompilationFinished(UMaterialInterface* MaterialInterface)
 {
 	if (CurrentToolTarget.LandscapeInfo.IsValid() &&
-		CurrentToolTarget.LandscapeInfo->GetLandscapeProxy() != NULL &&
-		CurrentToolTarget.LandscapeInfo->GetLandscapeProxy()->GetLandscapeMaterial() != NULL &&
+		CurrentToolTarget.LandscapeInfo->GetLandscapeProxy() != nullptr &&
+		CurrentToolTarget.LandscapeInfo->GetLandscapeProxy()->GetLandscapeMaterial() != nullptr &&
 		CurrentToolTarget.LandscapeInfo->GetLandscapeProxy()->GetLandscapeMaterial()->IsDependent(MaterialInterface))
 	{
 		CurrentToolTarget.LandscapeInfo->UpdateLayerInfoMap();
@@ -3306,10 +3311,10 @@ ALandscape* FEdModeLandscape::ChangeComponentSetting(int32 NumComponentsX, int32
 
 	const int32 NewComponentSizeQuads = NumSubsections * SubsectionSizeQuads;
 
-	ALandscape* NewLandscape = NULL;
+	ALandscape* NewLandscape = nullptr;
 
 	ULandscapeInfo* LandscapeInfo = CurrentToolTarget.LandscapeInfo.Get();
-	if (ensure(LandscapeInfo != NULL))
+	if (ensure(LandscapeInfo != nullptr))
 	{
 		int32 OldMinX, OldMinY, OldMaxX, OldMaxY;
 		if (LandscapeInfo->GetLandscapeExtent(OldMinX, OldMinY, OldMaxX, OldMaxY))
@@ -3373,7 +3378,7 @@ ALandscape* FEdModeLandscape::ChangeComponentSetting(int32 NumComponentsX, int32
 						TArray<uint8> ResampledWeightData;
 						for (const FLandscapeInfoLayerSettings& LayerSettings : LandscapeInfo->Layers)
 						{
-							if (LayerSettings.LayerInfoObj != NULL)
+							if (LayerSettings.LayerInfoObj != nullptr)
 							{
 								auto ImportLayerInfo = new(OutImportMaterialLayerInfos) FLandscapeImportLayerInfo(LayerSettings);
 								ImportLayerInfo->LayerData.AddZeroed(OldVertsX * OldVertsY);
@@ -3411,7 +3416,7 @@ ALandscape* FEdModeLandscape::ChangeComponentSetting(int32 NumComponentsX, int32
 						TArray<uint8> ExpandedWeightData;
 						for (const FLandscapeInfoLayerSettings& LayerSettings : LandscapeInfo->Layers)
 						{
-							if (LayerSettings.LayerInfoObj != NULL)
+							if (LayerSettings.LayerInfoObj != nullptr)
 							{
 								auto ImportLayerInfo = new(OutImportMaterialLayerInfos) FLandscapeImportLayerInfo(LayerSettings);
 								ImportLayerInfo->LayerData.AddZeroed(NewVertsX * NewVertsY * sizeof(uint8));
@@ -3682,7 +3687,7 @@ ELandscapeEditingState FEdModeLandscape::GetEditingState() const
 	{
 		return ELandscapeEditingState::SIEWorld;
 	}
-	else if (GEditor->PlayWorld != NULL)
+	else if (GEditor->PlayWorld != nullptr)
 	{
 		return ELandscapeEditingState::PIEWorld;
 	}
@@ -4036,7 +4041,7 @@ bool FEdModeLandscape::CanEditLayer(FText* Reason /*=nullptr*/, FLandscapeLayer*
 		}
 	}
 
-	if (CurrentToolTarget.TargetType == ELandscapeToolTargetType::Weightmap && CurrentToolTarget.LayerInfo == NULL && CurrentTool->GetToolName() != FName("BlueprintBrush"))
+	if (CurrentToolTarget.TargetType == ELandscapeToolTargetType::Weightmap && CurrentToolTarget.LayerInfo == nullptr && CurrentTool->GetToolName() != FName("BlueprintBrush"))
 	{
 		if (Reason)
 		{
diff --git a/Engine/Source/Editor/LandscapeEditor/Private/LandscapeEdMode.h b/Engine/Source/Editor/LandscapeEditor/Private/LandscapeEdMode.h
--- a/Engine/Source/Editor/LandscapeEditor/Private/LandscapeEdMode.h
+++ b/Engine/Source/Editor/LandscapeEditor/Private/LandscapeEdMode.h
@@ -93,10 +93,10 @@ struct FLandscapeTargetListInfo
 		: TargetName(InTargetName)
 		, TargetType(InTargetType)
 		, LandscapeInfo(InLandscapeInfo)
-		, LayerInfoObj(NULL)
+		, LayerInfoObj(nullptr)
 		, LayerName(NAME_None)
-		, Owner(NULL)
-		, ThumbnailMIC(NULL)
+		, Owner(nullptr)
+		, ThumbnailMIC(nullptr)
 		, bValid(true)
 		, LayerIndex(InLayerIndex)
 	{
@@ -130,7 +130,7 @@ struct FLandscapeTargetListInfo
 			return &LandscapeInfo->Layers[Index];
 		}
 
-		return NULL;
+		return nullptr;
 	}
 
 	FLandscapeEditorLayerSettings* GetEditorLayerSettings() const
@@ -150,7 +150,7 @@ struct FLandscapeTargetListInfo
 				return &Proxy->EditorLayerSettings[Index];
 			}
 		}
-		return NULL;
+		return nullptr;
 	}
 
 	FName GetLayerName() const;
@@ -619,6 +619,9 @@ public:
 	/** Returns true if landscape resolution combined to the current tool action is still compliant to the currently applied limitations. */
 	bool IsLandscapeResolutionCompliant() const;
 
+	/** Returns true if the current landscape tool handles edit layers. */
+	bool DoesCurrentToolAffectEditLayers() const;
+
 	/** Returns the default Error Text when modifying or creating landscape would break the resolution limit. */
 	FText GetLandscapeResolutionErrorText() const;
 
diff --git a/Engine/Source/Editor/LandscapeEditor/Private/LandscapeEdModeComponentTools.cpp b/Engine/Source/Editor/LandscapeEditor/Private/LandscapeEdModeComponentTools.cpp
--- a/Engine/Source/Editor/LandscapeEditor/Private/LandscapeEdModeComponentTools.cpp
+++ b/Engine/Source/Editor/LandscapeEditor/Private/LandscapeEdModeComponentTools.cpp
@@ -131,6 +131,8 @@ public:
 	{
 	}
 
+	virtual bool AffectsEditLayers() const { return false; }
+
 	virtual ELandscapeLayerUpdateMode GetBeginToolContentUpdateFlag() const override
 	{
 		return ELandscapeLayerUpdateMode::Update_None;
@@ -249,6 +251,8 @@ public:
 	{
 	}
 
+	virtual bool AffectsEditLayers() const { return false; }
+
 	virtual const TCHAR* GetToolName() override { return TEXT("Mask"); }
 	virtual FText GetDisplayName() override { return NSLOCTEXT("UnrealEd", "LandscapeMode_Mask", "Region Selection"); };
 	virtual FText GetDisplayMessage() override { return NSLOCTEXT("UnrealEd", "LandscapeMode_Mask_Message", "Region Selection"); };
@@ -495,7 +499,7 @@ public:
 		: FLandscapeToolBase<FLandscapeToolStrokeMoveToLevel>(InEdMode)
 	{
 	}
-	virtual bool ShouldUpdateEditingLayer() const override { return false; }
+	virtual bool AffectsEditLayers() const override { return false; }
 
 	virtual const TCHAR* GetToolName() override { return TEXT("MoveToLevel"); }
 	virtual FText GetDisplayName() override { return NSLOCTEXT("UnrealEd", "LandscapeMode_MoveToLevel", "Move to Streaming Level"); };
@@ -786,7 +790,7 @@ public:
 		, bIsToolActionResolutionCompliant(true)
 	{
 	}
-	virtual bool ShouldUpdateEditingLayer() const override { return false; }
+	virtual bool AffectsEditLayers() const override { return false; }
 
 	virtual const TCHAR* GetToolName() override { return TEXT("AddComponent"); }
 	virtual FText GetDisplayName() override { return NSLOCTEXT("UnrealEd", "LandscapeMode_AddComponent", "Add New Landscape Component"); };
@@ -993,7 +997,7 @@ public:
 	{
 	}
 
-	virtual bool ShouldUpdateEditingLayer() const override { return false; }
+	virtual bool AffectsEditLayers() const override { return false; }
 
 	virtual const TCHAR* GetToolName() override { return TEXT("DeleteComponent"); }
 	virtual FText GetDisplayName() override { return NSLOCTEXT("UnrealEd", "LandscapeMode_DeleteComponent", "Delete Landscape Components"); };
@@ -1858,7 +1862,7 @@ public:
 		, DraggingEdge_Remainder(0.0f)
 	{
 	}
-
+	virtual bool AffectsEditLayers() const { return false; }
 	virtual const TCHAR* GetToolName() override { return TEXT("NewLandscape"); }
 	virtual FText GetDisplayName() override { return NSLOCTEXT("UnrealEd", "LandscapeMode_NewLandscape", "New Landscape"); };
 	virtual FText GetDisplayMessage() override { return NSLOCTEXT("UnrealEd", "LandscapeMode_NewLandscape_Message", "Create or import a new heightmap.  Assign a material and configure the components.  When you are ready to create your new Landscape, press the Create button in the lower-right corner of this panel. "); };
@@ -2316,6 +2320,7 @@ public:
 	{
 	}
 
+	virtual bool AffectsEditLayers() const { return false; }
 	virtual const TCHAR* GetToolName() override { return TEXT("ResizeLandscape"); }
 	virtual FText GetDisplayName() override { return LOCTEXT("LandscapeMode_ResizeLandscape", "Change Landscape Component Size"); };
 	virtual FText GetDisplayMessage() override { return LOCTEXT("LandscapeMode_ResizeLandscape_Message", "Change Landscape Component Size"); };
diff --git a/Engine/Source/Editor/LandscapeEditor/Private/LandscapeEdModeTools.h b/Engine/Source/Editor/LandscapeEditor/Private/LandscapeEdModeTools.h
--- a/Engine/Source/Editor/LandscapeEditor/Private/LandscapeEdModeTools.h
+++ b/Engine/Source/Editor/LandscapeEditor/Private/LandscapeEdModeTools.h
@@ -1381,9 +1381,9 @@ public:
 	}
 
 protected:
-	FEdModeLandscape* EdMode;
+	FEdModeLandscape* EdMode = nullptr;
 	const FLandscapeToolTarget& Target;
-	ULandscapeInfo* LandscapeInfo;
+	ULandscapeInfo* LandscapeInfo = nullptr;
 };
 
 
@@ -1405,7 +1405,10 @@ public:
 	{
 	}
 
-	virtual bool ShouldUpdateEditingLayer() const { return EdMode->CanHaveLandscapeLayersContent(); }
+	virtual bool ShouldUpdateEditingLayer() const 
+	{ 
+		return AffectsEditLayers() && EdMode->CanHaveLandscapeLayersContent();
+	}
 
 	virtual ELandscapeLayerUpdateMode GetBeginToolContentUpdateFlag() const
 	{
diff --git a/Engine/Source/Editor/LandscapeEditor/Private/LandscapeEditorDetailCustomization_Layers.cpp b/Engine/Source/Editor/LandscapeEditor/Private/LandscapeEditorDetailCustomization_Layers.cpp
--- a/Engine/Source/Editor/LandscapeEditor/Private/LandscapeEditorDetailCustomization_Layers.cpp
+++ b/Engine/Source/Editor/LandscapeEditor/Private/LandscapeEditorDetailCustomization_Layers.cpp
@@ -85,7 +85,7 @@ END_SLATE_FUNCTION_BUILD_OPTIMIZATION
 bool FLandscapeEditorDetailCustomization_Layers::ShoudShowLayersErrorMessageTip()
 {
 	FEdModeLandscape* LandscapeEdMode = GetEditorMode();
-	if (LandscapeEdMode)
+	if (LandscapeEdMode && LandscapeEdMode->DoesCurrentToolAffectEditLayers())
 	{
 		return !LandscapeEdMode->CanEditLayer();
 	}
@@ -424,11 +424,11 @@ TSharedPtr<SWidget> FLandscapeEditorCustomNodeBuilder_Layers::OnLayerContextMenu
 {
 	FEdModeLandscape* LandscapeEdMode = GetEditorMode();
 	ALandscape* Landscape = LandscapeEdMode ? LandscapeEdMode->GetLandscape() : nullptr;
-	if (LandscapeEdMode && Landscape)
+	if (Landscape && LandscapeEdMode && LandscapeEdMode->DoesCurrentToolAffectEditLayers())
 	{
 		FLandscapeLayer* Layer = LandscapeEdMode->GetLayer(InLayerIndex);
 		TSharedRef<FLandscapeEditorCustomNodeBuilder_Layers> SharedThis = AsShared();
-		FMenuBuilder MenuBuilder(true, NULL);
+		FMenuBuilder MenuBuilder(true, nullptr);
 		MenuBuilder.BeginSection("LandscapeEditorLayerActions", LOCTEXT("LandscapeEditorLayerActions.Heading", "Edit Layers"));
 		{
 			// Create Layer
@@ -852,7 +852,7 @@ bool FLandscapeEditorCustomNodeBuilder_Layers::IsLayerEditionEnabled(int32 InLay
 	ALandscape* Landscape = LandscapeEdMode ? LandscapeEdMode->GetLandscape() : nullptr;
 	const FLandscapeLayer* Layer = LandscapeEdMode ? LandscapeEdMode->GetLayer(InLayerIndex) : nullptr;
 	const FLandscapeLayer* LayerReservedForSplines = Landscape ? Landscape->GetLandscapeSplinesReservedLayer() : nullptr;
-	return Layer && !Layer->bLocked && (Layer != LayerReservedForSplines);
+	return Layer && !Layer->bLocked && (Layer != LayerReservedForSplines) && LandscapeEdMode->DoesCurrentToolAffectEditLayers();
 }
 
 const FSlateBrush* FLandscapeEditorCustomNodeBuilder_Layers::GetLockBrushForLayer(int32 InLayerIndex) const
@@ -878,7 +878,7 @@ int32 FLandscapeEditorCustomNodeBuilder_Layers::SlotIndexToLayerIndex(int32 Slot
 FReply FLandscapeEditorCustomNodeBuilder_Layers::HandleDragDetected(const FGeometry& MyGeometry, const FPointerEvent& MouseEvent, int32 SlotIndex, SVerticalBox::FSlot* Slot)
 {
 	FEdModeLandscape* LandscapeEdMode = GetEditorMode();
-	if (LandscapeEdMode)
+	if (LandscapeEdMode && LandscapeEdMode->DoesCurrentToolAffectEditLayers())
 	{
 		int32 LayerIndex = SlotIndexToLayerIndex(SlotIndex);
 		FLandscapeLayer* Layer = LandscapeEdMode->GetLayer(LayerIndex);
diff --git a/Engine/Source/Editor/LandscapeEditor/Public/LandscapeToolInterface.h b/Engine/Source/Editor/LandscapeEditor/Public/LandscapeToolInterface.h
--- a/Engine/Source/Editor/LandscapeEditor/Public/LandscapeToolInterface.h
+++ b/Engine/Source/Editor/LandscapeEditor/Public/LandscapeToolInterface.h
@@ -204,7 +204,7 @@ public:
 	virtual bool InputDelta(FEditorViewportClient* InViewportClient, FViewport* InViewport, FVector& InDrag, FRotator& InRot, FVector& InScale) { return false; }
 	virtual bool GetCursor(EMouseCursor::Type& OutCursor) const { return false;  }
 
-	FLandscapeTool() : PreviousBrushIndex(-1) {}
+	FLandscapeTool() {}
 	virtual ~FLandscapeTool() {}
 	virtual const TCHAR* GetToolName() = 0;
 	virtual FText GetDisplayName() = 0;
@@ -218,6 +218,7 @@ public:
 	virtual bool IsSelectionAllowed(AActor* InActor, bool bInSelection) const { return false; }
 	virtual bool UsesTransformWidget() const { return false; }
 	virtual EAxisList::Type GetWidgetAxisToDraw(UE::Widget::EWidgetMode InWidgetMode) const { return EAxisList::All; }
+	virtual bool AffectsEditLayers() const { return true; };
 
 	virtual bool OverrideWidgetLocation() const { return true; }
 	virtual bool OverrideWidgetRotation() const { return true; }
@@ -264,8 +265,8 @@ public:
 	}
 
 public:
-	int32					PreviousBrushIndex;
-	TArray<FName>			ValidBrushes;
+	int32 PreviousBrushIndex = INDEX_NONE;
+	TArray<FName> ValidBrushes;
 };
 
 namespace LandscapeTool
