Fixed Painting Restriction not appearing when there's no active paint layer

diff --git a/Engine/Source/Editor/LandscapeEditor/Private/LandscapeEditorDetailCustomization_TargetLayers.cpp b/Engine/Source/Editor/LandscapeEditor/Private/LandscapeEditorDetailCustomization_TargetLayers.cpp
--- a/Engine/Source/Editor/LandscapeEditor/Private/LandscapeEditorDetailCustomization_TargetLayers.cpp
+++ b/Engine/Source/Editor/LandscapeEditor/Private/LandscapeEditorDetailCustomization_TargetLayers.cpp
@@ -121,8 +121,10 @@ bool FLandscapeEditorDetailCustomization_TargetLayers::ShouldShowPaintingRestric
 	{
 		const FName CurrentToolName = LandscapeEdMode->CurrentTool->GetToolName();
 
+		// Tool target type "Invalid" means Weightmap with no valid paint layer, so technically, it is weightmap and we therefore choose to show PaintingRestriction : 
 		if ((LandscapeEdMode->CurrentToolTarget.TargetType == ELandscapeToolTargetType::Weightmap && CurrentToolName != TEXT("BlueprintBrush"))
-			|| LandscapeEdMode->CurrentToolTarget.TargetType == ELandscapeToolTargetType::Visibility)
+			|| (LandscapeEdMode->CurrentToolTarget.TargetType == ELandscapeToolTargetType::Invalid)
+			|| (LandscapeEdMode->CurrentToolTarget.TargetType == ELandscapeToolTargetType::Visibility))
 		{
 			return true;
 		}
