Hides "Enable Edit Layers" when creating a new landscape and from Landcape property details when in Restrictive Mode.

diff --git a/Engine/Source/Editor/DetailCustomizations/Private/LandscapeUIDetails.cpp b/Engine/Source/Editor/DetailCustomizations/Private/LandscapeUIDetails.cpp
--- a/Engine/Source/Editor/DetailCustomizations/Private/LandscapeUIDetails.cpp
+++ b/Engine/Source/Editor/DetailCustomizations/Private/LandscapeUIDetails.cpp
@@ -16,6 +16,7 @@
 #include "Editor.h"
 #include "LandscapeInfo.h"
 #include "LandscapeStreamingProxy.h"
+#include "LandscapeSettings.h"
 
 #define LOCTEXT_NAMESPACE "FLandscapeUIDetails"
 
@@ -61,6 +62,11 @@ void FLandscapeUIDetails::CustomizeDetails( IDetailLayoutBuilder& DetailBuilder
 		[
 			CanHaveLayersPropertyHandle->CreatePropertyNameWidget(DisplayAndFilterText, ToolTipText)
 		]
+		.Visibility(MakeAttributeLambda([]()
+		{
+			const ULandscapeSettings* Settings = GetDefault<ULandscapeSettings>();
+			return Settings->InRestrictiveMode() ? EVisibility::Hidden : EVisibility::Visible;
+		}))
 		.ValueContent()
 		[
 			SNew(SCheckBox)
diff --git a/Engine/Source/Editor/LandscapeEditor/Private/LandscapeEditorDetailCustomization_NewLandscape.cpp b/Engine/Source/Editor/LandscapeEditor/Private/LandscapeEditorDetailCustomization_NewLandscape.cpp
--- a/Engine/Source/Editor/LandscapeEditor/Private/LandscapeEditorDetailCustomization_NewLandscape.cpp
+++ b/Engine/Source/Editor/LandscapeEditor/Private/LandscapeEditorDetailCustomization_NewLandscape.cpp
@@ -98,7 +98,11 @@ void FLandscapeEditorDetailCustomization_NewLandscape::CustomizeDetails(IDetailL
 	];
 
 	TSharedRef<IPropertyHandle> PropertyHandle_CanHaveLayersContent = DetailBuilder.GetProperty(GET_MEMBER_NAME_CHECKED(ULandscapeEditorObject, bCanHaveLayersContent));
-	NewLandscapeCategory.AddProperty(PropertyHandle_CanHaveLayersContent);
+	NewLandscapeCategory.AddProperty(PropertyHandle_CanHaveLayersContent).Visibility(MakeAttributeLambda([]()
+	{
+		const ULandscapeSettings* Settings = GetDefault<ULandscapeSettings>();
+		return Settings->InRestrictiveMode() ? EVisibility::Hidden : EVisibility::Visible;
+	}));
 
 	TSharedRef<IPropertyHandle> PropertyHandle_FlipYAxis = DetailBuilder.GetProperty(GET_MEMBER_NAME_CHECKED(ULandscapeEditorObject, bFlipYAxis));
 	NewLandscapeCategory.AddProperty(PropertyHandle_FlipYAxis).Visibility(MakeAttributeLambda([]()
