* Fixed transaction localization for landscape resolution change
* Fixed New Landscape - Fill World not being transacted

diff --git a/Engine/Source/Editor/LandscapeEditor/Private/LandscapeEditorDetailCustomization_NewLandscape.cpp b/Engine/Source/Editor/LandscapeEditor/Private/LandscapeEditorDetailCustomization_NewLandscape.cpp
--- a/Engine/Source/Editor/LandscapeEditor/Private/LandscapeEditorDetailCustomization_NewLandscape.cpp
+++ b/Engine/Source/Editor/LandscapeEditor/Private/LandscapeEditorDetailCustomization_NewLandscape.cpp
@@ -753,7 +753,7 @@ void FLandscapeEditorDetailCustomization_NewLandscape::OnChangeLandscapeResoluti
 			return;
 		}
 
-		FScopedTransaction Transaction(LOCTEXT("ChangeResolutionX_Transaction", "Change Landscape Resolution X"), !bUsingSlider && bCommit);
+		FScopedTransaction Transaction(LOCTEXT("ChangeResolutionY_Transaction", "Change Landscape Resolution Y"), !bUsingSlider && bCommit);
 	
 		LandscapeEdMode->UISettings->Modify();
 		LandscapeEdMode->UISettings->NewLandscape_ComponentCount.Y = NewComponentCountY;
@@ -978,10 +978,11 @@ FReply FLandscapeEditorDetailCustomization_NewLandscape::OnFillWorldButtonClicke
 	FEdModeLandscape* LandscapeEdMode = GetEditorMode();
 	if (LandscapeEdMode != nullptr)
 	{
-		FVector& NewLandscapeLocation = LandscapeEdMode->UISettings->NewLandscape_Location;
-		NewLandscapeLocation.X = 0;
-		NewLandscapeLocation.Y = 0;
+		FScopedTransaction Transaction(LOCTEXT("FillWorld_Transaction", "Landscape Fill World"));
+
+		LandscapeEdMode->UISettings->Modify();
 
+		LandscapeEdMode->UISettings->NewLandscape_Location = FVector::ZeroVector;
 		const int32 QuadsPerComponent = LandscapeEdMode->UISettings->NewLandscape_SectionsPerComponent * LandscapeEdMode->UISettings->NewLandscape_QuadsPerSection;
 		LandscapeEdMode->UISettings->NewLandscape_ComponentCount.X = FMath::CeilToInt(WORLD_MAX / QuadsPerComponent / LandscapeEdMode->UISettings->NewLandscape_Scale.X);
 		LandscapeEdMode->UISettings->NewLandscape_ComponentCount.Y = FMath::CeilToInt(WORLD_MAX / QuadsPerComponent / LandscapeEdMode->UISettings->NewLandscape_Scale.Y);
