Landscape: Make the Blueprint Brush dropdown be able to select runtime module blueprint brushes.

diff --git a/Engine/Source/Editor/LandscapeEditor/Public/LandscapeEditorObject.h b/Engine/Source/Editor/LandscapeEditor/Public/LandscapeEditorObject.h
--- a/Engine/Source/Editor/LandscapeEditor/Public/LandscapeEditorObject.h
+++ b/Engine/Source/Editor/LandscapeEditor/Public/LandscapeEditorObject.h
@@ -8,7 +8,7 @@
 #include "LandscapeProxy.h"
 #include "Editor/LandscapeEditor/Private/LandscapeEdMode.h"
 #include "LandscapeFileFormatInterface.h"
-#include "LandscapeBlueprintBrush.h"
+#include "LandscapeBlueprintBrushBase.h"
 #include "LandscapeImportHelper.h"
 
 #include "LandscapeEditorObject.generated.h"
@@ -476,7 +476,7 @@ class ULandscapeEditorObject : public UObject
 	// Blueprint Brush Tool
 
 	UPROPERTY(Category = "Tool Settings", EditAnywhere, Transient, meta = (DisplayName = "Blueprint Brush", ShowForTools = "BlueprintBrush"))
-	TSubclassOf<ALandscapeBlueprintBrush> BlueprintBrush;
+	TSubclassOf<ALandscapeBlueprintBrushBase> BlueprintBrush;
 
 	// Resize Landscape Tool
 
