* Unified landscape brush command shortcuts with other modeling tools
* Added command to change both the brush size (left-right mouse move) and falloff (up-down mouse move). Default is ctrl+alt+right click drag
* Fixed the "change landscape brush size/falloff/strength" commands that couldn't work with default chord values nor secondary key bindings
* Added DoesCommandSupportInput function to editor viewport client in order to centralize copy/pasted logic
* Fixed "cmd" modifier not being taken into account all the time in said logic

diff --git a/Engine/Source/Editor/LandscapeEditor/Private/LandscapeEdMode.cpp b/Engine/Source/Editor/LandscapeEditor/Private/LandscapeEdMode.cpp
--- a/Engine/Source/Editor/LandscapeEditor/Private/LandscapeEdMode.cpp
+++ b/Engine/Source/Editor/LandscapeEditor/Private/LandscapeEdMode.cpp
@@ -777,7 +777,7 @@ void FEdModeLandscape::Tick(FEditorViewportClient* ViewportClient, float DeltaTi
 			(LandscapeEditorControlType == ELandscapeFoliageEditorControlType::RequireCtrl && !IsCtrlDown(Viewport)))
 		{
 			// Don't end the current tool if we are just modifying it
-			if (!IsAdjustingBrush(Viewport) && CurrentTool->IsToolActive())
+			if (!IsAdjustingBrush(ViewportClient) && CurrentTool->IsToolActive())
 			{
 				CurrentTool->EndTool(ViewportClient);
 				Viewport->CaptureMouse(false);
@@ -849,29 +849,42 @@ bool FEdModeLandscape::MouseMove(FEditorViewportClient* InViewportClient, FViewp
 
 		if (FMath::Abs(MouseXDelta) > 0 || FMath::Abs(MouseYDelta) > 0)
 		{
-			const bool bSizeChange = FMath::Abs(MouseXDelta) > FMath::Abs(MouseYDelta) ?
+			const bool bHorizontallyDominant = FMath::Abs(MouseXDelta) >= FMath::Abs(MouseYDelta);
+			const bool bIncrease = bHorizontallyDominant ?
 				MouseXDelta > 0 : 
 				MouseYDelta < 0; // The way y position is stored here is inverted relative to expected mouse movement to change brush size
+
 			// Are we altering something about the brush?
-			FInputChord CompareChord;
-			FInputBindingManager::Get().GetUserDefinedChord(FLandscapeEditorCommands::LandscapeContext, TEXT("DragBrushSize"), EMultipleKeyBindingIndex::Primary, CompareChord);
-			if (InViewport->KeyState(CompareChord.Key))
+			const FLandscapeEditorCommands& LandscapeActions = FLandscapeEditorCommands::Get();
+			if (InViewportClient->IsCommandChordPressed(LandscapeActions.DragBrushSizeAndFalloff))
+			{
+				// If a chord controlling both Size and Falloff is specified, control Size when moving mostly left-right and Falloff when moving mostly up/down
+				if (bHorizontallyDominant)
+				{
+					ChangeBrushSize(bIncrease);
+				}
+				else
+				{
+					ChangeBrushFalloff(bIncrease);
+				}
+				return true;
+			}
+
+			if (InViewportClient->IsCommandChordPressed(LandscapeActions.DragBrushSize))
 			{
-				ChangeBrushSize(bSizeChange);
+				ChangeBrushSize(bIncrease);
 				return true;
 			}
 
-			FInputBindingManager::Get().GetUserDefinedChord(FLandscapeEditorCommands::LandscapeContext, TEXT("DragBrushStrength"), EMultipleKeyBindingIndex::Primary, CompareChord);
-			if (InViewport->KeyState(CompareChord.Key))
+			if (InViewportClient->IsCommandChordPressed(LandscapeActions.DragBrushFalloff))
 			{
-				ChangeBrushStrength(bSizeChange);
+				ChangeBrushFalloff(bIncrease);
 				return true;
 			}
 
-			FInputBindingManager::Get().GetUserDefinedChord(FLandscapeEditorCommands::LandscapeContext, TEXT("DragBrushFalloff"), EMultipleKeyBindingIndex::Primary, CompareChord);
-			if (InViewport->KeyState(CompareChord.Key))
+			if (InViewportClient->IsCommandChordPressed(LandscapeActions.DragBrushStrength))
 			{
-				ChangeBrushFalloff(bSizeChange);
+				ChangeBrushStrength(bIncrease);
 				return true;
 			}
 		}
@@ -1495,21 +1508,22 @@ bool FEdModeLandscape::HandleClick(FEditorViewportClient* InViewportClient, HHit
 	return false;
 }
 
-bool FEdModeLandscape::IsAdjustingBrush(FViewport* InViewport) const
+bool FEdModeLandscape::IsAdjustingBrush(FEditorViewportClient* InViewportClient) const
 {
-	FInputChord CompareChord;
-	FInputBindingManager::Get().GetUserDefinedChord(FLandscapeEditorCommands::LandscapeContext, TEXT("DragBrushSize"), EMultipleKeyBindingIndex::Primary, CompareChord);
-	if (InViewport->KeyState(CompareChord.Key))
+	const FLandscapeEditorCommands& LandscapeActions = FLandscapeEditorCommands::Get();
+	if (InViewportClient->IsCommandChordPressed(LandscapeActions.DragBrushSizeAndFalloff))
+	{
+		return true;
+	}
+	if (InViewportClient->IsCommandChordPressed(LandscapeActions.DragBrushSize))
 	{
 		return true;
 	}
-	FInputBindingManager::Get().GetUserDefinedChord(FLandscapeEditorCommands::LandscapeContext, TEXT("DragBrushFalloff"), EMultipleKeyBindingIndex::Primary, CompareChord);
-	if (InViewport->KeyState(CompareChord.Key))
+	if (InViewportClient->IsCommandChordPressed(LandscapeActions.DragBrushFalloff))
 	{
 		return true;
 	}
-	FInputBindingManager::Get().GetUserDefinedChord(FLandscapeEditorCommands::LandscapeContext, TEXT("DragBrushStrength"), EMultipleKeyBindingIndex::Primary, CompareChord);
-	if (InViewport->KeyState(CompareChord.Key))
+	if (InViewportClient->IsCommandChordPressed(LandscapeActions.DragBrushStrength))
 	{
 		return true;
 	}
@@ -1643,7 +1657,7 @@ bool FEdModeLandscape::InputKey(FEditorViewportClient* ViewportClient, FViewport
 	}
 
 
-	if(IsAdjustingBrush(Viewport))
+	if(IsAdjustingBrush(ViewportClient))
 	{
 		ToolActiveViewport = Viewport;
 		return false; // false to let FEditorViewportClient.InputKey start mouse tracking and enable InputDelta() so we can use it
diff --git a/Engine/Source/Editor/LandscapeEditor/Private/LandscapeEdMode.h b/Engine/Source/Editor/LandscapeEditor/Private/LandscapeEdMode.h
--- a/Engine/Source/Editor/LandscapeEditor/Private/LandscapeEdMode.h
+++ b/Engine/Source/Editor/LandscapeEditor/Private/LandscapeEdMode.h
@@ -404,7 +404,7 @@ public:
 	virtual bool HandleClick(FEditorViewportClient* InViewportClient, HHitProxy* HitProxy, const FViewportClick& Click) override;
 
 	/** True if we are interactively changing the brush size, falloff, or strength */
-	bool IsAdjustingBrush(FViewport* InViewport) const;
+	bool IsAdjustingBrush(FEditorViewportClient* InViewportClient) const;
 	void ChangeBrushSize(bool bIncrease);
 	void ChangeBrushFalloff(bool bIncrease);
 	void ChangeBrushStrength(bool bIncrease);
diff --git a/Engine/Source/Editor/LandscapeEditor/Private/LandscapeEdModeBrushes.cpp b/Engine/Source/Editor/LandscapeEditor/Private/LandscapeEdModeBrushes.cpp
--- a/Engine/Source/Editor/LandscapeEditor/Private/LandscapeEdModeBrushes.cpp
+++ b/Engine/Source/Editor/LandscapeEditor/Private/LandscapeEdModeBrushes.cpp
@@ -142,7 +142,7 @@ public:
 		TSet<ULandscapeComponent*> NewComponents;
 		bool bHasUnloadedComponents = false;
 		// Adjusting the brush may use the same keybind as moving the camera as they can be user-set, so we need this second check.
-		if (!ViewportClient->IsMovingCamera() || EdMode->IsAdjustingBrush(ViewportClient->Viewport))
+		if (!ViewportClient->IsMovingCamera() || EdMode->IsAdjustingBrush(ViewportClient))
 		{
 			// GetComponentsInRegion expects an inclusive max
 			LandscapeInfo->GetComponentsInRegion(Bounds.Min.X, Bounds.Min.Y, Bounds.Max.X - 1, Bounds.Max.Y - 1, NewComponents);
@@ -418,7 +418,7 @@ public:
 		TSet<ULandscapeComponent*> NewComponents;
 
 		// Adjusting the brush may use the same keybind as moving the camera as they can be user-set, so we need this second check.
-		if (!ViewportClient->IsMovingCamera() || EdMode->IsAdjustingBrush(ViewportClient->Viewport))
+		if (!ViewportClient->IsMovingCamera() || EdMode->IsAdjustingBrush(ViewportClient))
 		{
 			ULandscapeInfo* LandscapeInfo = EdMode->CurrentToolTarget.LandscapeInfo.Get();
 			if (LandscapeInfo && LandscapeInfo->ComponentSizeQuads > 0)
diff --git a/Engine/Source/Editor/LandscapeEditor/Private/LandscapeEdModeTools.h b/Engine/Source/Editor/LandscapeEditor/Private/LandscapeEdModeTools.h
--- a/Engine/Source/Editor/LandscapeEditor/Private/LandscapeEdModeTools.h
+++ b/Engine/Source/Editor/LandscapeEditor/Private/LandscapeEdModeTools.h
@@ -1527,7 +1527,7 @@ public:
 			if (EdMode->LandscapeMouseTrace(ViewportClient, x, y, HitLocation))
 			{
 				// If we are moving the mouse to adjust the brush size, don't move the brush
-				if (EdMode->CurrentBrush && !EdMode->IsAdjustingBrush(Viewport))
+				if (EdMode->CurrentBrush && !EdMode->IsAdjustingBrush(ViewportClient))
 				{
 					// Inform the brush of the current location, to update the cursor
 					EdMode->CurrentBrush->MouseMove(HitLocation.X, HitLocation.Y);
diff --git a/Engine/Source/Editor/LandscapeEditor/Private/LandscapeEditorCommands.cpp b/Engine/Source/Editor/LandscapeEditor/Private/LandscapeEditorCommands.cpp
--- a/Engine/Source/Editor/LandscapeEditor/Private/LandscapeEditorCommands.cpp
+++ b/Engine/Source/Editor/LandscapeEditor/Private/LandscapeEditorCommands.cpp
@@ -126,16 +126,17 @@ void FLandscapeEditorCommands::RegisterCommands()
 
 	UI_COMMAND(IncreaseBrushSize, "Increase Brush Size", "Press this key to increase brush size by a fixed increment.", EUserInterfaceActionType::RadioButton, FInputChord(EKeys::RightBracket));
 	UI_COMMAND(DecreaseBrushSize, "Decrease Brush Size", "Press this key to decrease brush size by a fixed increment.", EUserInterfaceActionType::RadioButton, FInputChord(EKeys::LeftBracket));
-	UI_COMMAND(IncreaseBrushFalloff, "Increase Brush Falloff", "Press this key to increase brush falloff by a fixed increment.", EUserInterfaceActionType::RadioButton, FInputChord());
-	UI_COMMAND(DecreaseBrushFalloff, "Decrease Brush Falloff", "Press this key to decrease brush falloff by a fixed increment.", EUserInterfaceActionType::RadioButton, FInputChord());
-	UI_COMMAND(IncreaseBrushStrength, "Increase Brush Strength", "Press this key to increase brush strength by a fixed increment.", EUserInterfaceActionType::RadioButton, FInputChord());
-	UI_COMMAND(DecreaseBrushStrength, "Decrease Brush Strength", "Press this key to decrease brush strength by a fixed increment.", EUserInterfaceActionType::RadioButton, FInputChord());
+	UI_COMMAND(IncreaseBrushFalloff, "Increase Brush Falloff", "Press this key to increase brush falloff by a fixed increment.", EUserInterfaceActionType::RadioButton, FInputChord(EModifierKey::Control | EModifierKey::Shift, EKeys::RightBracket));
+	UI_COMMAND(DecreaseBrushFalloff, "Decrease Brush Falloff", "Press this key to decrease brush falloff by a fixed increment.", EUserInterfaceActionType::RadioButton, FInputChord(EModifierKey::Control | EModifierKey::Shift, EKeys::LeftBracket));
+	UI_COMMAND(IncreaseBrushStrength, "Increase Brush Strength", "Press this key to increase brush strength by a fixed increment.", EUserInterfaceActionType::RadioButton, FInputChord(EModifierKey::Control, EKeys::RightBracket));
+	UI_COMMAND(DecreaseBrushStrength, "Decrease Brush Strength", "Press this key to decrease brush strength by a fixed increment.", EUserInterfaceActionType::RadioButton, FInputChord(EModifierKey::Control, EKeys::LeftBracket));
 	UI_COMMAND(IncreaseAlphaBrushRotation, "Increase Alpha Brush Rotation", "Press this key to increase alpha brush rotation by a fixed increment.", EUserInterfaceActionType::Button, FInputChord());
 	UI_COMMAND(DecreaseAlphaBrushRotation, "Decrease Alpha Brush Rotation", "Press this key to decrease alpha brush rotation by a fixed increment.", EUserInterfaceActionType::Button, FInputChord());
 
-	UI_COMMAND(DragBrushSize, "Change Brush Size", "Hold this key and then click and drag to increase or decrease brush size.", EUserInterfaceActionType::RadioButton, FInputChord());
-	UI_COMMAND(DragBrushFalloff, "Change Brush Falloff", "Hold this key and then click and drag to increase or decrease brush falloff.", EUserInterfaceActionType::RadioButton, FInputChord());
-	UI_COMMAND(DragBrushStrength, "Change Brush Strength", "Hold this key and then click and drag to increase or decrease brush strength.", EUserInterfaceActionType::RadioButton, FInputChord());
+	UI_COMMAND(DragBrushSizeAndFalloff, "Change Brush Size And Falloff", "Hold this key and then drag to increase or decrease brush size (right/left) or falloff (up/down).", EUserInterfaceActionType::RadioButton, FInputChord(EModifierKey::Control | EModifierKey::Alt, EKeys::RightMouseButton));
+	UI_COMMAND(DragBrushSize, "Change Brush Size", "Hold this key and then drag to increase or decrease brush size.", EUserInterfaceActionType::RadioButton, FInputChord());
+	UI_COMMAND(DragBrushFalloff, "Change Brush Falloff", "Hold this key and then drag to increase or decrease brush falloff.", EUserInterfaceActionType::RadioButton, FInputChord());
+	UI_COMMAND(DragBrushStrength, "Change Brush Strength", "Hold this key and then drag to increase or decrease brush strength.", EUserInterfaceActionType::RadioButton, FInputChord());
 }
 
 #undef LOCTEXT_NAMESPACE
diff --git a/Engine/Source/Editor/LandscapeEditor/Private/LandscapeEditorCommands.h b/Engine/Source/Editor/LandscapeEditor/Private/LandscapeEditorCommands.h
--- a/Engine/Source/Editor/LandscapeEditor/Private/LandscapeEditorCommands.h
+++ b/Engine/Source/Editor/LandscapeEditor/Private/LandscapeEditorCommands.h
@@ -88,6 +88,7 @@ public:
 	TSharedPtr<FUICommandInfo> IncreaseAlphaBrushRotation;
 	TSharedPtr<FUICommandInfo> DecreaseAlphaBrushRotation;
 
+	TSharedPtr<FUICommandInfo> DragBrushSizeAndFalloff;
 	TSharedPtr<FUICommandInfo> DragBrushSize;
 	TSharedPtr<FUICommandInfo> DragBrushFalloff;
 	TSharedPtr<FUICommandInfo> DragBrushStrength;
diff --git a/Engine/Source/Editor/UnrealEd/Private/EditorViewportClient.cpp b/Engine/Source/Editor/UnrealEd/Private/EditorViewportClient.cpp
--- a/Engine/Source/Editor/UnrealEd/Private/EditorViewportClient.cpp
+++ b/Engine/Source/Editor/UnrealEd/Private/EditorViewportClient.cpp
@@ -3195,6 +3195,23 @@ bool FEditorViewportClient::IsCmdPressed() const
 	return Viewport->KeyState(EKeys::LeftCommand) || Viewport->KeyState(EKeys::RightCommand);
 }
 
+bool FEditorViewportClient::IsCommandChordPressed(const TSharedPtr<FUICommandInfo> InCommand, FKey InOptionalKey) const
+{
+	bool bIsChordPressed = false;
+	// Check each bound chord
+	for (uint32 i = 0; i < static_cast<uint32>(EMultipleKeyBindingIndex::NumChords); ++i)
+	{
+		EMultipleKeyBindingIndex ChordIndex = static_cast<EMultipleKeyBindingIndex> (i);
+		const FInputChord& Chord = *InCommand->GetActiveChord(ChordIndex);
+
+		bIsChordPressed |= (Chord.NeedsControl() == IsCtrlPressed())
+			&& (Chord.NeedsAlt() == IsAltPressed())
+			&& (Chord.NeedsShift() == IsShiftPressed())
+			&& (Chord.NeedsCommand() == IsCmdPressed())
+			&& Viewport->KeyState(InOptionalKey.IsValid() ? InOptionalKey : Chord.Key) == true;
+	}
+	return bIsChordPressed;
+}
 
 void FEditorViewportClient::ProcessDoubleClickInViewport( const struct FInputEventState& InputState, FSceneView& View )
 {
diff --git a/Engine/Source/Editor/UnrealEd/Private/LevelEditorViewport.cpp b/Engine/Source/Editor/UnrealEd/Private/LevelEditorViewport.cpp
--- a/Engine/Source/Editor/UnrealEd/Private/LevelEditorViewport.cpp
+++ b/Engine/Source/Editor/UnrealEd/Private/LevelEditorViewport.cpp
@@ -2758,25 +2758,6 @@ TSharedPtr<FDragTool> FLevelEditorViewportClient::MakeDragTool( EDragTool::Type
 	return DragTool;
 }
 
-static bool CommandAcceptsInput( FLevelEditorViewportClient& ViewportClient, FKey Key, const TSharedPtr<FUICommandInfo> Command )
-{
-	bool bAccepted = false;
-	for (uint32 i = 0; i < static_cast<uint8>(EMultipleKeyBindingIndex::NumChords); ++i)
-	{
-		// check each bound chord
-		EMultipleKeyBindingIndex ChordIndex = static_cast<EMultipleKeyBindingIndex>(i);
-		const FInputChord& Chord = *Command->GetActiveChord(ChordIndex);
-
-		bAccepted |= Chord.IsValidChord()
-			&& (!Chord.NeedsControl() || ViewportClient.IsCtrlPressed())
-			&& (!Chord.NeedsAlt() || ViewportClient.IsAltPressed())
-			&& (!Chord.NeedsShift() || ViewportClient.IsShiftPressed())
-			&& (!Chord.NeedsCommand() || ViewportClient.IsCmdPressed())
-			&& Chord.Key == Key;
-	}
-	return bAccepted;
-}
-
 static const FLevelViewportCommands& GetLevelViewportCommands()
 {
 	static FName LevelEditorName("LevelEditor");
@@ -2909,7 +2890,7 @@ bool FLevelEditorViewportClient::InputKey(const FInputKeyEventArgs& InEventArgs)
 	bool bHandled = FEditorViewportClient::InputKey(InEventArgs);
 
 	// Handle input for the player height preview mode. 
-	if (!InputState.IsMouseButtonEvent() && CommandAcceptsInput(*this, InEventArgs.Key, GetLevelViewportCommands().EnablePreviewMesh))
+	if (!InputState.IsMouseButtonEvent() && IsCommandChordPressed(GetLevelViewportCommands().EnablePreviewMesh, InEventArgs.Key))
 	{
 		// Holding down the backslash buttons turns on the mode. 
 		if (InEventArgs.Event == IE_Pressed)
@@ -2917,7 +2898,7 @@ bool FLevelEditorViewportClient::InputKey(const FInputKeyEventArgs& InEventArgs)
 			GEditor->SetPreviewMeshMode(true);
 
 			// If shift down, cycle between the preview meshes
-			if (CommandAcceptsInput(*this, InEventArgs.Key, GetLevelViewportCommands().CyclePreviewMesh))
+			if (IsCommandChordPressed(GetLevelViewportCommands().CyclePreviewMesh, InEventArgs.Key))
 			{
 				GEditor->CyclePreviewMesh();
 			}
diff --git a/Engine/Source/Editor/UnrealEd/Private/SnappingUtils.cpp b/Engine/Source/Editor/UnrealEd/Private/SnappingUtils.cpp
--- a/Engine/Source/Editor/UnrealEd/Private/SnappingUtils.cpp
+++ b/Engine/Source/Editor/UnrealEd/Private/SnappingUtils.cpp
@@ -132,19 +132,8 @@ bool FEditorViewportSnapping::IsSnapToVertexEnabled(bool bIsPivot)
 	else if( GCurrentLevelEditingViewportClient )
 	{
 		FLevelEditorModule& LevelEditor = FModuleManager::GetModuleChecked<FLevelEditorModule>( TEXT("LevelEditor") );
-		const FLevelEditorCommands& Commands = LevelEditor.GetLevelEditorCommands();
-		bool bIsChordPressed = false;
-		for (uint32 i = 0; i < static_cast<uint8>(EMultipleKeyBindingIndex::NumChords); ++i)
-		{
-			EMultipleKeyBindingIndex ChordIndex = static_cast<EMultipleKeyBindingIndex> (i);
-			const FInputChord& Chord = bIsPivot ? (*Commands.HoldToEnablePivotVertexSnapping->GetActiveChord(ChordIndex)) : (*Commands.HoldToEnableVertexSnapping->GetActiveChord(ChordIndex));
-
-			bIsChordPressed |= (Chord.NeedsControl() == GCurrentLevelEditingViewportClient->IsCtrlPressed())
-				&& (Chord.NeedsAlt() == GCurrentLevelEditingViewportClient->IsAltPressed())
-				&& (Chord.NeedsShift() == GCurrentLevelEditingViewportClient->IsShiftPressed())
-				&& GCurrentLevelEditingViewportClient->Viewport->KeyState(Chord.Key) == true;
-		}
-		return bIsChordPressed;
+		TSharedPtr<FUICommandInfo> Command = bIsPivot ? LevelEditor.GetLevelEditorCommands().HoldToEnablePivotVertexSnapping : LevelEditor.GetLevelEditorCommands().HoldToEnableVertexSnapping;
+		return GCurrentLevelEditingViewportClient->IsCommandChordPressed(Command);
 	}
 	else
 	{
diff --git a/Engine/Source/Editor/UnrealEd/Public/EditorViewportClient.h b/Engine/Source/Editor/UnrealEd/Public/EditorViewportClient.h
--- a/Engine/Source/Editor/UnrealEd/Public/EditorViewportClient.h
+++ b/Engine/Source/Editor/UnrealEd/Public/EditorViewportClient.h
@@ -1050,6 +1050,14 @@ public:
 	bool IsCtrlPressed() const;
 	bool IsShiftPressed() const;
 	bool IsCmdPressed() const;
+	
+	/**
+	 * Utility function to return whether the command accepts the key states
+	 * @param InCommand The command being checked
+	 * @param InOptionalKey (Optional) input key being tested against. If not specified, the current viewport's key state's key will be used
+	 * @return True if one of the command's chords accepts the input :
+	 */
+	bool IsCommandChordPressed(const TSharedPtr<FUICommandInfo> InCommand, FKey InOptionalKey = FKey()) const;
 
 	/** @return True if the window is in an immersive viewport */
 	bool IsInImmersiveViewport() const;
