Cleans Landscape Visibility Layer behavior in Landscape Editor:
* Fixes Visibility Layer Info Object being created when a default Layer Info Object is set
* Fixes Visibility Layer being displayed twice in the visibility tool
* Removes Visibility Layer from New Landscape Tool
* Removes Visibility Layer from Paint tool
* Adds Visibility Layer in Import/Export tool so it can be imported/exported

diff --git a/Engine/Source/Editor/LandscapeEditor/Private/LandscapeEditorDetailCustomization_NewLandscape.cpp b/Engine/Source/Editor/LandscapeEditor/Private/LandscapeEditorDetailCustomization_NewLandscape.cpp
--- a/Engine/Source/Editor/LandscapeEditor/Private/LandscapeEditorDetailCustomization_NewLandscape.cpp
+++ b/Engine/Source/Editor/LandscapeEditor/Private/LandscapeEditorDetailCustomization_NewLandscape.cpp
@@ -239,7 +239,10 @@ void FLandscapeEditorDetailCustomization_NewLandscape::CustomizeDetails(IDetailL
 	.Visibility(TAttribute<EVisibility>::Create(TAttribute<EVisibility>::FGetter::CreateStatic(&GetVisibilityOnlyInNewLandscapeMode, ENewLandscapePreviewMode::ImportLandscape)));
 
 	TSharedRef<IPropertyHandle> PropertyHandle_Layers = DetailBuilder.GetProperty(GET_MEMBER_NAME_CHECKED(ULandscapeEditorObject, ImportLandscape_Layers));
-	NewLandscapeCategory.AddProperty(PropertyHandle_Layers);
+
+	TSharedRef<FDetailArrayBuilder> ArrayBuilder = MakeShared<FDetailArrayBuilder>(PropertyHandle_Layers, /*InGenerateHeader*/ true, /*InDisplayResetToDefault*/ false, /*InDisplayElementNum*/ false);
+	ArrayBuilder->OnGenerateArrayElementWidget(FOnGenerateArrayElementWidget::CreateSP(this, &FLandscapeEditorDetailCustomization_NewLandscape::GenerateLayersArrayElementWidget));
+	NewLandscapeCategory.AddCustomBuilder(ArrayBuilder, /*bForAdvanced*/ false);
 
 	TSharedRef<IPropertyHandle> PropertyHandle_Location = DetailBuilder.GetProperty(GET_MEMBER_NAME_CHECKED(ULandscapeEditorObject, NewLandscape_Location));
 	TSharedRef<IPropertyHandle> PropertyHandle_Location_X = PropertyHandle_Location->GetChildHandle("X").ToSharedRef();
@@ -721,11 +724,11 @@ void FLandscapeEditorDetailCustomization_NewLandscape::OnChangeLandscapeResoluti
 		int32 NewComponentCountX = LandscapeEdMode->UISettings->CalcComponentsCount(NewValue);
 		if (NewComponentCountX == LandscapeEdMode->UISettings->NewLandscape_ComponentCount.X)
 		{
-			return;
+return;
 		}
 
 		FScopedTransaction Transaction(LOCTEXT("ChangeResolutionX_Transaction", "Change Landscape Resolution X"), !bUsingSlider && bCommit);
-		
+
 		LandscapeEdMode->UISettings->Modify();
 		LandscapeEdMode->UISettings->NewLandscape_ComponentCount.X = NewComponentCountX;
 	}
@@ -754,7 +757,7 @@ void FLandscapeEditorDetailCustomization_NewLandscape::OnChangeLandscapeResoluti
 		}
 
 		FScopedTransaction Transaction(LOCTEXT("ChangeResolutionY_Transaction", "Change Landscape Resolution Y"), !bUsingSlider && bCommit);
-	
+
 		LandscapeEdMode->UISettings->Modify();
 		LandscapeEdMode->UISettings->NewLandscape_ComponentCount.Y = NewComponentCountY;
 	}
@@ -811,6 +814,31 @@ EVisibility FLandscapeEditorDetailCustomization_NewLandscape::GetVisibilityOnlyI
 	return EVisibility::Collapsed;
 }
 
+void FLandscapeEditorDetailCustomization_NewLandscape::GenerateLayersArrayElementWidget(TSharedRef<IPropertyHandle> InPropertyHandle, int32 InArrayIndex, IDetailChildrenBuilder& InChildrenBuilder)
+{
+	InChildrenBuilder.AddProperty(InPropertyHandle)
+	.Visibility(TAttribute<EVisibility>::Create(TAttribute<EVisibility>::FGetter::CreateSP(this, &FLandscapeEditorDetailCustomization_NewLandscape::GetLayerVisibility, InPropertyHandle)));
+}
+
+EVisibility FLandscapeEditorDetailCustomization_NewLandscape::GetLayerVisibility(TSharedRef<IPropertyHandle> InPropertyHandle) const
+{
+	const FEdModeLandscape* EdMode = GetEditorMode();
+	TArray<void*> RawData;
+	InPropertyHandle->AccessRawData(RawData);
+
+	if ((EdMode != nullptr) && !RawData.IsEmpty() && (RawData[0] != nullptr))
+	{
+		FLandscapeImportLayer* ImportLayer = reinterpret_cast<FLandscapeImportLayer*>(RawData[0]);
+
+		if ((EdMode->NewLandscapePreviewMode != ENewLandscapePreviewMode::ImportLandscape) && (ImportLayer->LayerName == ALandscapeProxy::VisibilityLayer->LayerName))
+		{
+			return EVisibility::Hidden;
+		}
+	}
+
+	return EVisibility::Visible;
+}
+
 bool FLandscapeEditorDetailCustomization_NewLandscape::IsCreateButtonEnabled() const
 {
 	const FEdModeLandscape* EdMode = GetEditorMode();
@@ -921,8 +949,8 @@ FReply FLandscapeEditorDetailCustomization_NewLandscape::OnCreateButtonClicked()
 			ULandscapeLayerInfoObject* LayerInfo = ImportLandscapeLayersList[i].LayerInfo;
 			FName LayerName = ImportLandscapeLayersList[i].LayerName;
 
-			// If DefaultLayerInfoObject is set and LayerInfo does not exist, we will try to create the new LayerInfo by cloning DefaultLayerInfoObject.
-			if (bIsDefaultLayerInfoObjectSet && (LayerInfo == nullptr))
+			// If DefaultLayerInfoObject is set and LayerInfo does not exist, we will try to create the new LayerInfo by cloning DefaultLayerInfoObject. Except for VisibilityLayer which doesn't require an asset.
+			if (bIsDefaultLayerInfoObjectSet && (LayerInfo == nullptr) && (LayerName != ALandscapeProxy::VisibilityLayer->LayerName))
 			{
 				LayerInfo = Landscape->CreateLayerInfo(*LayerName.ToString(), DefaultLayerInfoObject.LoadSynchronous());
 
diff --git a/Engine/Source/Editor/LandscapeEditor/Private/LandscapeEditorDetailCustomization_NewLandscape.h b/Engine/Source/Editor/LandscapeEditor/Private/LandscapeEditorDetailCustomization_NewLandscape.h
--- a/Engine/Source/Editor/LandscapeEditor/Private/LandscapeEditorDetailCustomization_NewLandscape.h
+++ b/Engine/Source/Editor/LandscapeEditor/Private/LandscapeEditorDetailCustomization_NewLandscape.h
@@ -66,6 +66,10 @@ public:
 
 	static EVisibility GetVisibilityOnlyInNewLandscapeMode(ENewLandscapePreviewMode value);
 
+	/** Called to generate ImportLayer children */
+	void GenerateLayersArrayElementWidget(TSharedRef<IPropertyHandle> InPropertyHandle, int32 InArrayIndex, IDetailChildrenBuilder& InChildrenBuilder);
+	EVisibility GetLayerVisibility(TSharedRef<IPropertyHandle> InPropertyHandle) const;
+
 	// Import
 	static EVisibility GetHeightmapErrorVisibility(TSharedRef<IPropertyHandle> PropertyHandle_HeightmapImportResult);
 	static FSlateColor GetHeightmapErrorColor(TSharedRef<IPropertyHandle> PropertyHandle_HeightmapImportResult);
diff --git a/Engine/Source/Editor/LandscapeEditor/Private/LandscapeEditorObject.cpp b/Engine/Source/Editor/LandscapeEditor/Private/LandscapeEditorObject.cpp
--- a/Engine/Source/Editor/LandscapeEditor/Private/LandscapeEditorObject.cpp
+++ b/Engine/Source/Editor/LandscapeEditor/Private/LandscapeEditorObject.cpp
@@ -629,6 +629,11 @@ void ULandscapeEditorObject::RefreshLayerImport(FLandscapeImportLayer& ImportLay
 	ImportLayer.ErrorMessage = FText();
 	ImportLayer.ImportResult = ELandscapeImportResult::Success;
 
+	if (ImportLayer.LayerName == ALandscapeProxy::VisibilityLayer->LayerName)
+	{
+		ImportLayer.LayerInfo = ALandscapeProxy::VisibilityLayer;
+	}
+
 	if (!ImportLayer.SourceFilePath.IsEmpty())
 	{
 		if (!ImportLayer.LayerInfo)
@@ -837,7 +842,7 @@ void ULandscapeEditorObject::RefreshImportLayersList(bool bRefreshFromTarget)
 		Material = ParentMode->GetTargetLandscapeMaterial();
 		for (const TSharedRef<FLandscapeTargetListInfo>& TargetListInfo : ParentMode->GetTargetList())
 		{
-			if (TargetListInfo->TargetType != ELandscapeToolTargetType::Weightmap)
+			if ((TargetListInfo->TargetType != ELandscapeToolTargetType::Weightmap) && (TargetListInfo->TargetType != ELandscapeToolTargetType::Visibility))
 			{
 				continue;
 			}
