Landscape: Allow painting on the edges of landscape components

diff --git a/Engine/Source/Editor/LandscapeEditor/Private/LandscapeEdMode.cpp b/Engine/Source/Editor/LandscapeEditor/Private/LandscapeEdMode.cpp
--- a/Engine/Source/Editor/LandscapeEditor/Private/LandscapeEdMode.cpp
+++ b/Engine/Source/Editor/LandscapeEditor/Private/LandscapeEdMode.cpp
@@ -58,6 +58,7 @@
 #include "Engine/TextureRenderTarget2D.h"
 #include "Settings/EditorExperimentalSettings.h"
 #include "ComponentRecreateRenderStateContext.h"
+#include "VisualLogger/VisualLogger.h"
 
 #define LOCTEXT_NAMESPACE "Landscape"
 
@@ -1095,10 +1096,40 @@ bool FEdModeLandscape::LandscapeMouseTrace(FEditorViewportClient* ViewportClient
 		Start -= WORLD_MAX * MouseViewportRayDirection;
 	}
 
-	return LandscapeTrace(Start, End, OutHitLocation);
+	return LandscapeTrace(Start, End, MouseViewportRayDirection,OutHitLocation);
 }
 
-bool FEdModeLandscape::LandscapeTrace(const FVector& InRayOrigin, const FVector& InRayEnd, FVector& OutHitLocation)
+struct FEdModeLandscape::FProcessLandscapeTraceHitsResult
+{
+	FVector HitLocation;
+	ULandscapeHeightfieldCollisionComponent* HeightfieldComponent;
+	ALandscapeProxy* LandscapeProxy;
+};
+
+bool FEdModeLandscape::ProcessLandscapeTraceHits(const TArray<FHitResult>& InResults, FProcessLandscapeTraceHitsResult& OutLandscapeTraceHitsResult)
+{
+	for (int32 i = 0; i < InResults.Num(); i++)
+	{
+		const FHitResult& Hit = InResults[i];
+		ULandscapeHeightfieldCollisionComponent* CollisionComponent = Cast<ULandscapeHeightfieldCollisionComponent>(Hit.Component.Get());
+		if (CollisionComponent)
+		{
+			ALandscapeProxy* HitLandscape = CollisionComponent->GetLandscapeProxy();
+			if (HitLandscape &&
+				CurrentToolTarget.LandscapeInfo.IsValid() &&
+				CurrentToolTarget.LandscapeInfo->LandscapeGuid == HitLandscape->GetLandscapeGuid())
+			{
+				OutLandscapeTraceHitsResult.HeightfieldComponent = CollisionComponent;
+				OutLandscapeTraceHitsResult.HitLocation = Hit.Location;
+				OutLandscapeTraceHitsResult.LandscapeProxy = HitLandscape;
+				return true;
+			}
+		}
+	}
+	return false;
+}
+
+bool FEdModeLandscape::LandscapeTrace(const FVector& InRayOrigin, const FVector& InRayEnd, const FVector& InDirection, FVector& OutHitLocation)
 {
 	if (!CurrentTool || !CurrentToolTarget.LandscapeInfo.IsValid())
 	{
@@ -1120,22 +1151,75 @@ bool FEdModeLandscape::LandscapeTrace(const FVector& InRayOrigin, const FVector&
 	TArray<FHitResult> Results;
 	// Each landscape component has 2 collision shapes, 1 of them is specific to landscape editor
 	// Trace only ECC_Visibility channel, so we do hit only Editor specific shape
-	World->LineTraceMultiByObjectType(Results, Start, End, FCollisionObjectQueryParams(ECollisionChannel::ECC_Visibility), FCollisionQueryParams(SCENE_QUERY_STAT(LandscapeTrace), true));
+	if (World->LineTraceMultiByObjectType(Results, Start, End, FCollisionObjectQueryParams(ECollisionChannel::ECC_Visibility), FCollisionQueryParams(SCENE_QUERY_STAT(LandscapeTrace), true)))
+	{
+		if (FProcessLandscapeTraceHitsResult ProcessResult; ProcessLandscapeTraceHits(Results, ProcessResult))
+		{
+			OutHitLocation = ProcessResult.LandscapeProxy->LandscapeActorToWorld().InverseTransformPosition(ProcessResult.HitLocation);
+			
+			UE_VLOG_SEGMENT_THICK(World, LogLandscapeEdMode, VeryVerbose, InRayOrigin, InRayOrigin + 20000.0f * InDirection, FColor(100,255,100), 4, TEXT("landscape:ray-hit"));
+			UE_VLOG_LOCATION(World, LogLandscapeEdMode, VeryVerbose,  ProcessResult.LandscapeProxy->LandscapeActorToWorld().TransformPosition(OutHitLocation), 20.0,  FColor(100,100,255), TEXT("landscape:point-hit"));
+			return true;
+		}
+	}
 
-	for (int32 i = 0; i < Results.Num(); i++)
+	UE_VLOG_SEGMENT_THICK(World, LogLandscapeEdMode, VeryVerbose, InRayOrigin, InRayOrigin + 20000.0f * InDirection, FColor(255,100,100), 4, TEXT("landscape:ray-miss"));
+		
+	// If there is no landscape directly under the mouse search for a landscape collision
+	// under the shape of the brush.
+	FCollisionShape SphereShape;
+	SphereShape.SetSphere(UISettings->BrushRadius);
+	if ( World->SweepMultiByObjectType(Results, Start, End, FQuat::Identity, FCollisionObjectQueryParams(ECollisionChannel::ECC_Visibility), SphereShape, FCollisionQueryParams(SCENE_QUERY_STAT(LandscapeTrace), true)))
 	{
-		const FHitResult& Hit = Results[i];
-		ULandscapeHeightfieldCollisionComponent* CollisionComponent = Cast<ULandscapeHeightfieldCollisionComponent>(Hit.Component.Get());
-		if (CollisionComponent)
+		if (FProcessLandscapeTraceHitsResult SweepProcessResult; ProcessLandscapeTraceHits(Results, SweepProcessResult))
 		{
-			ALandscapeProxy* HitLandscape = CollisionComponent->GetLandscapeProxy();
-			if (HitLandscape &&
-				CurrentToolTarget.LandscapeInfo.IsValid() &&
-				CurrentToolTarget.LandscapeInfo->LandscapeGuid == HitLandscape->GetLandscapeGuid())
+			UE_VLOG_LOCATION(World, LogLandscapeEdMode, VeryVerbose,  SweepProcessResult.HitLocation, UISettings->BrushRadius,  FColor(255,100,100), TEXT("landscape:sweep-hit-location"));
+			FSphere HitSphere(SweepProcessResult.HitLocation, UISettings->BrushRadius);
+
+			float MeanHeight = 0.0f;
+			int32 Count = 0;
+			const int32 NumHeightSamples = 16;
+			
+			for (int32 Y = 0; Y < NumHeightSamples; ++Y)
 			{
-				OutHitLocation = HitLandscape->LandscapeActorToWorld().InverseTransformPosition(Hit.Location);
-				return true;
+				float HY = (Y / (float) NumHeightSamples) * HitSphere.W * 2.0f +  HitSphere.Center.Y - HitSphere.W;
+				for (int32 X = 0; X < NumHeightSamples; ++X)
+				{
+					float HX = (X / (float) NumHeightSamples) * HitSphere.W * 2.0f +  HitSphere.Center.X - HitSphere.W;
+
+					FVector HeightSampleLocation(HX, HY, HitSphere.Center.Z);
+
+					TOptional<float> Height = SweepProcessResult.LandscapeProxy->GetHeightAtLocation(HeightSampleLocation, EHeightfieldSource::Editor);
+
+					if (!Height.IsSet())
+					{
+						continue;
+					}
+
+					UE_VLOG_LOCATION(World, LogLandscapeEdMode, VeryVerbose,  FVector(HX, HY, Height.GetValue()), 2.0f,  FColor(100,100,255), TEXT(""));
+					
+					MeanHeight += Height.GetValue();
+					Count++;
+				}
 			}
+
+			FVector PointOnPlane ( HitSphere.Center.X, HitSphere.Center.Y, HitSphere.Center.Z );
+			
+			if  ( Count > 0)
+			{
+				MeanHeight /= (float) Count;
+				PointOnPlane.Z = MeanHeight;
+			}
+
+			UE_VLOG_LOCATION(World, LogLandscapeEdMode, VeryVerbose,  PointOnPlane, 10.0,  FColor(100,100,255), TEXT("landscape:point-on-plane"));
+			
+			const FPlane Plane( PointOnPlane, FVector::ZAxisVector);
+			FVector EstimatedHitLocation = FMath::RayPlaneIntersection(Start, InDirection, Plane);
+
+			UE_VLOG_LOCATION(World, LogLandscapeEdMode, VeryVerbose,  EstimatedHitLocation, 10.0,  FColor(100,100,255), TEXT("landscape:estimated-hit-location"));
+			
+			OutHitLocation = SweepProcessResult.LandscapeProxy->LandscapeActorToWorld().InverseTransformPosition(EstimatedHitLocation);
+			return true;
 		}
 	}
 	
diff --git a/Engine/Source/Editor/LandscapeEditor/Private/LandscapeEdMode.h b/Engine/Source/Editor/LandscapeEditor/Private/LandscapeEdMode.h
--- a/Engine/Source/Editor/LandscapeEditor/Private/LandscapeEdMode.h
+++ b/Engine/Source/Editor/LandscapeEditor/Private/LandscapeEdMode.h
@@ -488,7 +488,12 @@ public:
 	bool LandscapePlaneTrace(FEditorViewportClient* ViewportClient, int32 MouseX, int32 MouseY, const FPlane& Plane, FVector& OutHitLocation);
 
 	/** Trace under the specified laser start and direction and return the landscape hit and the hit location (in landscape quad space) */
-	bool LandscapeTrace(const FVector& InRayOrigin, const FVector& InRayEnd, FVector& OutHitLocation);
+	bool LandscapeTrace(const FVector& InRayOrigin, const FVector& InRayEnd, const FVector& InDirection, FVector& OutHitLocation);
+
+	struct FProcessLandscapeTraceHitsResult;
+	
+	/** Check if we've collided with the currently edited landscape, OutHitLocation is in the space of the Landscape Actor */   
+	bool ProcessLandscapeTraceHits(const TArray<FHitResult>& InResults, FProcessLandscapeTraceHitsResult& OutLandscapeTraceHitsResult);
 
 	void SetCurrentToolMode(FName ToolModeName, bool bRestoreCurrentTool = true);
 
diff --git a/Engine/Source/Runtime/Landscape/Classes/LandscapeHeightfieldCollisionComponent.h b/Engine/Source/Runtime/Landscape/Classes/LandscapeHeightfieldCollisionComponent.h
--- a/Engine/Source/Runtime/Landscape/Classes/LandscapeHeightfieldCollisionComponent.h
+++ b/Engine/Source/Runtime/Landscape/Classes/LandscapeHeightfieldCollisionComponent.h
@@ -28,6 +28,13 @@ namespace Chaos
 	class FHeightField;
 }
 
+enum class EHeightfieldSource
+{
+	Simple = 0,
+	Complex = 1,
+	Editor = 2
+};
+
 UCLASS(MinimalAPI, Within=LandscapeProxy)
 class ULandscapeHeightfieldCollisionComponent : public UPrimitiveComponent
 {
@@ -260,7 +267,7 @@ public:
 #endif
 
 public:
-	TOptional<float> GetHeight(float X, float Y);
+	TOptional<float> GetHeight(float X, float Y, EHeightfieldSource HeightFieldSource);
 
 	/**
 	 * Populates a supplied array with the heights from the heightfield.  Samples are placed
diff --git a/Engine/Source/Runtime/Landscape/Classes/LandscapeProxy.h b/Engine/Source/Runtime/Landscape/Classes/LandscapeProxy.h
--- a/Engine/Source/Runtime/Landscape/Classes/LandscapeProxy.h
+++ b/Engine/Source/Runtime/Landscape/Classes/LandscapeProxy.h
@@ -14,6 +14,7 @@
 #include "LandscapeComponent.h"
 #include "LandscapeNaniteComponent.h"
 #include "LandscapeWeightmapUsage.h"
+#include "LandscapeHeightfieldCollisionComponent.h"
 #include "VT/RuntimeVirtualTextureEnum.h"
 #include "ActorPartition/PartitionActor.h"
 #include "ILandscapeSplineInterface.h"
@@ -382,7 +383,7 @@ protected:
 #endif // WITH_EDITORONLY_DATA
 
 public:
-	LANDSCAPE_API TOptional<float> GetHeightAtLocation(FVector Location) const;
+	LANDSCAPE_API TOptional<float> GetHeightAtLocation(FVector Location, EHeightfieldSource HeightFieldSource = EHeightfieldSource::Complex) const;
 
 	/** Fills an array with height values **/
 	LANDSCAPE_API void GetHeightValues(int32& SizeX, int32& SizeY, TArray<float>& ArrayValue) const;
diff --git a/Engine/Source/Runtime/Landscape/Private/LandscapeCollision.cpp b/Engine/Source/Runtime/Landscape/Private/LandscapeCollision.cpp
--- a/Engine/Source/Runtime/Landscape/Private/LandscapeCollision.cpp
+++ b/Engine/Source/Runtime/Landscape/Private/LandscapeCollision.cpp
@@ -609,13 +609,6 @@ FPrimitiveSceneProxy* ULandscapeHeightfieldCollisionComponent::CreateSceneProxy(
 		TUniquePtr<FColoredMaterialRenderProxy> WireframeMaterialInstance = nullptr;
 	};
 
-	enum class EHeightfieldSource
-	{
-		Simple = 0,
-		Complex = 1,
-		Editor = 2
-	};
-
 	FLandscapeHeightfieldCollisionComponentSceneProxy* Proxy = nullptr;
 
 	if (HeightfieldRef.IsValid() && IsValidRef(HeightfieldRef))
@@ -2596,14 +2589,36 @@ ULandscapeComponent* ULandscapeHeightfieldCollisionComponent::GetRenderComponent
 	return RenderComponent.Get();
 }
 
-TOptional<float> ULandscapeHeightfieldCollisionComponent::GetHeight(float X, float Y)
+TOptional<float> ULandscapeHeightfieldCollisionComponent::GetHeight(float X, float Y, EHeightfieldSource HeightFieldSource)
 {
 	TOptional<float> Height;
 	const float ZScale = GetComponentTransform().GetScale3D().Z * LANDSCAPE_ZSCALE;
 
-	if (IsValidRef(HeightfieldRef) && HeightfieldRef->Heightfield.Get())
+	if (!IsValidRef(HeightfieldRef))
+	{
+		return Height;
+	}
+	
+	Chaos::FHeightField* HeightField = nullptr;
+	
+	switch(HeightFieldSource)
+	{
+	case EHeightfieldSource::Simple:
+		HeightField = HeightfieldRef->HeightfieldSimple.Get(); 
+		break;
+	case EHeightfieldSource::Complex:
+		HeightField = HeightfieldRef->Heightfield.Get(); 
+		break;
+#if WITH_EDITORONLY_DATA		
+	case EHeightfieldSource::Editor:
+		HeightField = HeightfieldRef->EditorHeightfield.Get();
+		break;
+#endif 
+	}
+	
+	if (HeightField)
 	{
-		Height = HeightfieldRef->Heightfield->GetHeightAt({ X, Y });
+		Height = HeightField->GetHeightAt({ X, Y });
 	}
 
 	return Height;
@@ -2722,7 +2737,7 @@ void ULandscapeMeshCollisionComponent::GetResourceSizeEx(FResourceSizeEx& Cumula
 }
 
 
-TOptional<float> ALandscapeProxy::GetHeightAtLocation(FVector Location) const
+TOptional<float> ALandscapeProxy::GetHeightAtLocation(FVector Location, EHeightfieldSource HeightFieldSource) const
 {
 	TOptional<float> Height;
 	if (ULandscapeInfo* Info = GetLandscapeInfo())
@@ -2733,7 +2748,7 @@ TOptional<float> ALandscapeProxy::GetHeightAtLocation(FVector Location) const
 		if (Component)
 		{
 			const FVector ComponentSpaceLocation = Component->GetComponentToWorld().InverseTransformPosition(Location);
-			const TOptional<float> LocalHeight = Component->GetHeight(ComponentSpaceLocation.X, ComponentSpaceLocation.Y);
+			const TOptional<float> LocalHeight = Component->GetHeight(ComponentSpaceLocation.X, ComponentSpaceLocation.Y, HeightFieldSource);
 			if (LocalHeight.IsSet())
 			{
 				Height = Component->GetComponentToWorld().TransformPositionNoScale(FVector(0, 0, LocalHeight.GetValue())).Z;
