FLandscapeToolFlatten: null out component ptr when exiting tool (Fix GC Leak)

diff --git a/Engine/Source/Editor/LandscapeEditor/Private/LandscapeEdModePaintTools.cpp b/Engine/Source/Editor/LandscapeEditor/Private/LandscapeEdModePaintTools.cpp
--- a/Engine/Source/Editor/LandscapeEditor/Private/LandscapeEdModePaintTools.cpp
+++ b/Engine/Source/Editor/LandscapeEditor/Private/LandscapeEdModePaintTools.cpp
@@ -1187,6 +1187,7 @@ public:
 		{
 			HeightmapFlattenPreviewComponent->DetachFromComponent(FDetachmentTransformRules::KeepRelativeTransform);
 			HeightmapFlattenPreviewComponent->DestroyComponent();
+			HeightmapFlattenPreviewComponent = nullptr;
 		}
 	}
 };
