Landscape Edit Layers: Add edit layer insertion modal dialog and ILandscapeEditorServices. Implement high level GetOrCreateEditLayer(FName InEditLayerName, ALandscape* InTargetLandscape) which allows for custom layer creation logic. Default logic returns layer index if matching layer is found, otherwise creates the layer and prompts the user with a dialog to choose how to insert the new layer.

diff --git a/Engine/Plugins/Experimental/LandscapePatch/Source/LandscapePatch/Private/LandscapePatchManager.cpp b/Engine/Plugins/Experimental/LandscapePatch/Source/LandscapePatch/Private/LandscapePatchManager.cpp
--- a/Engine/Plugins/Experimental/LandscapePatch/Source/LandscapePatch/Private/LandscapePatchManager.cpp
+++ b/Engine/Plugins/Experimental/LandscapePatch/Source/LandscapePatch/Private/LandscapePatchManager.cpp
@@ -2,12 +2,15 @@
 
 #include "LandscapePatchManager.h"
 
+#include "Modules/ModuleManager.h"
 #include "Engine/TextureRenderTarget2D.h"
 #include "Landscape.h"
 #include "LandscapeDataAccess.h"
 #include "LandscapePatchComponent.h"
 #include "LandscapePatchLogging.h"
 #include "RenderGraph.h" // RDG_EVENT_NAME
+#include "LandscapeModule.h"
+#include "LandscapeEditorServices.h"
 
 #define LOCTEXT_NAMESPACE "LandscapePatchManager"
 
@@ -131,41 +134,12 @@ void ALandscapePatchManager::SetTargetLandscape(ALandscape* InTargetLandscape)
 
 		if (InTargetLandscape && ensure(InTargetLandscape->CanHaveLayersContent()))
 		{
-			FName LayerName = FName("LandscapePatches");
-			int32 ExistingPatchLayerIndex = InTargetLandscape->GetLayerIndex(LayerName);
-			if (ExistingPatchLayerIndex == INDEX_NONE)
-			{
-				ExistingPatchLayerIndex = InTargetLandscape->CreateLayer(LayerName);
-
-				// If there's a layer with a water brush manager, put the layer we created under that
-				// TODO: It's uncertain whether we want this approach long term. We might want to have a popup ask where to
-				// insert the layer, for instance. However, this is what the first artists to test the module suggested,
-				// and it's fairly easy to swap layers around, so it is easy for users to change if they are unhappy. 
-				int32 LayerIndex = 0;
-				for (; LayerIndex < ExistingPatchLayerIndex; ++LayerIndex)
-				{
-					TArray<ALandscapeBlueprintBrushBase*> LayerBrushes = InTargetLandscape->GetBrushesForLayer(LayerIndex);
-					if (LayerBrushes.ContainsByPredicate([](ALandscapeBlueprintBrushBase* Brush) 
-						{ 
-							// It would be preferable to do a Cast<ALandscapeBlueprintBrushBase> to be robust to subclassing, etc.,
-							// but that doesn't seem worth having to add a dependency on the water plugin, hence the comparison by name.
-							// Note that the "A" prefix is removed in GetName.
-							FString BrushManagerName = Brush->GetClass()->GetName();
-							return Brush->GetClass()->GetName() == TEXT("WaterBrushManager");
-						}))
-					{
-						// Found a water brush manager.
-						break;
-					}
-				}
-
-				if (LayerIndex < ExistingPatchLayerIndex)
-				{
-					InTargetLandscape->ReorderLayer(ExistingPatchLayerIndex, LayerIndex);
-					ExistingPatchLayerIndex = LayerIndex;
-				}
-			}
-			InTargetLandscape->AddBrushToLayer(ExistingPatchLayerIndex, this);
+			static const FName PatchLayerName = FName("LandscapePatches");
+
+			ILandscapeModule& LandscapeModule = FModuleManager::GetModuleChecked<ILandscapeModule>("Landscape");
+			int32 PatchLayerIndex = LandscapeModule.GetLandscapeEditorServices()->GetOrCreateEditLayer(PatchLayerName, InTargetLandscape);
+			
+			InTargetLandscape->AddBrushToLayer(PatchLayerIndex, this);
 		}
 	}
 #endif
diff --git a/Engine/Plugins/Experimental/Water/Source/Editor/Private/WaterLandscapeBrush.cpp b/Engine/Plugins/Experimental/Water/Source/Editor/Private/WaterLandscapeBrush.cpp
--- a/Engine/Plugins/Experimental/Water/Source/Editor/Private/WaterLandscapeBrush.cpp
+++ b/Engine/Plugins/Experimental/Water/Source/Editor/Private/WaterLandscapeBrush.cpp
@@ -26,6 +26,8 @@
 #include "Components/BillboardComponent.h"
 #include "Modules/ModuleManager.h"
 #include "WaterModule.h"
+#include "LandscapeModule.h"
+#include "LandscapeEditorServices.h"
 
 #define LOCTEXT_NAMESPACE "WaterLandscapeBrush"
 
@@ -457,13 +459,12 @@ void AWaterLandscapeBrush::SetTargetLandscape(ALandscape* InTargetLandscape)
 
 		if (InTargetLandscape && InTargetLandscape->CanHaveLayersContent())
 		{
-			FName WaterLayerName = FName("Water");
-			int32 ExistingWaterLayerIndex = InTargetLandscape->GetLayerIndex(WaterLayerName);
-			if (ExistingWaterLayerIndex == INDEX_NONE)
-			{
-				ExistingWaterLayerIndex = InTargetLandscape->CreateLayer(WaterLayerName);
-			}
-			InTargetLandscape->AddBrushToLayer(ExistingWaterLayerIndex, this);
+			static const FName WaterLayerName = FName("Water");
+			
+			ILandscapeModule& LandscapeModule = FModuleManager::GetModuleChecked<ILandscapeModule>("Landscape");
+			int32 WaterLayerIndex = LandscapeModule.GetLandscapeEditorServices()->GetOrCreateEditLayer(WaterLayerName, InTargetLandscape);
+			
+			InTargetLandscape->AddBrushToLayer(WaterLayerIndex, this);
 		}
 	}
 
diff --git a/Engine/Source/Editor/LandscapeEditor/Private/LandscapeEditorModule.cpp b/Engine/Source/Editor/LandscapeEditor/Private/LandscapeEditorModule.cpp
--- a/Engine/Source/Editor/LandscapeEditor/Private/LandscapeEditorModule.cpp
+++ b/Engine/Source/Editor/LandscapeEditor/Private/LandscapeEditorModule.cpp
@@ -19,6 +19,7 @@
 #include "LandscapeFileFormatPng.h"
 #include "LandscapeFileFormatRaw.h"
 #include "Settings/EditorExperimentalSettings.h"
+#include "LandscapeEditorServices.h"
 
 #include "Framework/MultiBox/MultiBoxBuilder.h"
 #include "PropertyEditorModule.h"
@@ -27,6 +28,7 @@
 #include "LandscapeEditorDetailCustomization_CopyPaste.h"
 #include "LandscapeEditorDetailCustomization_ImportLayers.h"
 #include "LandscapeSplineDetails.h"
+#include "LandscapeModule.h"
 
 #include "LevelEditor.h"
 #include "ToolMenus.h"
@@ -134,6 +136,10 @@ public:
 			FUIAction ActionBuildNanite(FExecuteAction::CreateStatic(&BuildNanite), FCanExecuteAction());
 			Section.AddMenuEntry(NAME_None, LOCTEXT("BuildNaniteOnly", "Build Nanite Only"), LOCTEXT("BuildLandscapeNanite", "Build Nanite representation"), TAttribute<FSlateIcon>(), ActionBuildNanite, EUserInterfaceActionType::Button);
 		}
+		
+		ILandscapeModule& LandscapeModule = FModuleManager::GetModuleChecked<ILandscapeModule>("Landscape");
+		LandscapeEditorServices.Reset(new FLandscapeEditorServices);
+		LandscapeModule.SetLandscapeEditorServices(LandscapeEditorServices.Get());
 	}
 
 	/**
@@ -164,6 +170,13 @@ public:
 		// remove actor factories
 		// TODO - this crashes on shutdown
 		// GEditor->ActorFactories.RemoveAll([](const UActorFactory* ActorFactory) { return ActorFactory->IsA<UActorFactoryLandscape>(); });
+
+		ILandscapeModule& LandscapeModule = FModuleManager::GetModuleChecked<ILandscapeModule>("Landscape");
+		if (LandscapeModule.GetLandscapeEditorServices() == LandscapeEditorServices.Get())
+		{
+			LandscapeModule.SetLandscapeEditorServices(nullptr);
+		}
+		LandscapeEditorServices.Reset();
 	}
 
 	static void ConstructLandscapeViewportMenu(FMenuBuilder& MenuBuilder)
@@ -313,7 +326,7 @@ public:
 	virtual const ILandscapeWeightmapFileFormat* GetWeightmapFormatByExtension(const TCHAR* Extension) const override;
 
 	virtual TSharedPtr<FUICommandList> GetLandscapeLevelViewportCommandList() const override;
-		
+
 protected:
 	TSharedPtr<FExtender> ViewportMenuExtender;
 	TSharedPtr<FUICommandList> GlobalUICommandList;
@@ -323,6 +336,7 @@ protected:
 	mutable FString WeightmapImportDialogTypeString;
 	mutable FString HeightmapExportDialogTypeString;
 	mutable FString WeightmapExportDialogTypeString;
+	TUniquePtr<ILandscapeEditorServices> LandscapeEditorServices;
 };
 
 IMPLEMENT_MODULE(FLandscapeEditorModule, LandscapeEditor);
diff --git a/Engine/Source/Editor/LandscapeEditor/Private/LandscapeEditorServices.cpp b/Engine/Source/Editor/LandscapeEditor/Private/LandscapeEditorServices.cpp
new file mode 100644
--- /dev/null
+++ b/Engine/Source/Editor/LandscapeEditor/Private/LandscapeEditorServices.cpp
@@ -0,0 +1,19 @@
+// Copyright Epic Games, Inc. All Rights Reserved.
+
+#include "LandscapeEditorServices.h"
+#include "SLandscapeLayerListDialog.h"
+
+#if WITH_EDITOR
+int32 FLandscapeEditorServices::GetOrCreateEditLayer(FName InEditLayerName, ALandscape* InTargetLandscape)
+{
+	int32 ExistingLayerIndex = InTargetLandscape->GetLayerIndex(InEditLayerName);
+	if (ExistingLayerIndex == INDEX_NONE)
+	{
+		InTargetLandscape->CreateLayer(InEditLayerName);
+		TSharedPtr<SLandscapeLayerListDialog> Dialog = SNew(SLandscapeLayerListDialog, InTargetLandscape->LandscapeLayers);
+		Dialog->ShowModal();
+		ExistingLayerIndex = Dialog->InsertedLayerIndex;
+	}
+	return ExistingLayerIndex;
+}
+#endif // WITH_EDITOR
\ No newline at end of file
diff --git a/Engine/Source/Editor/LandscapeEditor/Private/SLandscapeLayerListDialog.cpp b/Engine/Source/Editor/LandscapeEditor/Private/SLandscapeLayerListDialog.cpp
new file mode 100644
--- /dev/null
+++ b/Engine/Source/Editor/LandscapeEditor/Private/SLandscapeLayerListDialog.cpp
@@ -0,0 +1,284 @@
+// Copyright Epic Games, Inc. All Rights Reserved.
+
+#include "SLandscapeLayerListDialog.h"
+#include "Landscape.h"
+#include "Widgets/Text/STextBlock.h"
+#include "Widgets/Views/STableViewBase.h"
+#include "Widgets/Views/STableRow.h"
+#include "Input/DragAndDrop.h"
+#include "DragAndDrop/DecoratedDragDropOp.h"
+
+#define LOCTEXT_NAMESPACE "LandscapeEditor"
+
+/* File scope utility classes */
+
+struct FWidgetLayerListItem
+{
+	FWidgetLayerListItem(const FLandscapeLayer* InLayer, const TFunction<void(void)>& InOnLayerListUpdated, TArray<TSharedPtr<FWidgetLayerListItem>>* InWidgetLayerList)
+	: LayerName(InLayer->Name)
+	, LayerGuid(InLayer->Guid)
+	, OnLayerListUpdated(InOnLayerListUpdated)
+	, WidgetLayerList(InWidgetLayerList)
+	{}
+
+	// Duplicated Layer Info
+	FName LayerName;
+	FGuid LayerGuid;
+	
+	// Metadata for UI
+	bool bAllowedToDrag = false;
+	TFunction<void(void)> OnLayerListUpdated;
+	TArray<TSharedPtr<FWidgetLayerListItem>>* WidgetLayerList = nullptr;
+};
+
+class FWidgetLayerListDragDropOp : public FDragDropOperation
+{
+public:
+	DRAG_DROP_OPERATOR_TYPE(FWidgetLayerListDragDropOp , FDragDropOperation)
+
+	/** The template to create an instance */
+	TSharedPtr<FWidgetLayerListItem> ListItem;
+
+	/** Constructs the drag drop operation */
+	static TSharedRef<FWidgetLayerListDragDropOp> New(const TSharedPtr<FWidgetLayerListItem>& InListItem, FText InDragText)
+	{
+		TSharedRef<FWidgetLayerListDragDropOp> Operation = MakeShared<FWidgetLayerListDragDropOp>();
+		Operation->ListItem = InListItem;
+		Operation->Construct();
+
+		return Operation;
+	}
+};
+
+class SWidgetLayerListItem : public STableRow<TSharedPtr<FWidgetLayerListItem>>
+{
+public:
+	SLATE_BEGIN_ARGS( SWidgetLayerListItem ){}
+	SLATE_END_ARGS()
+
+	void Construct( const FArguments& InArgs, const TSharedRef<STableViewBase>& InOwnerTableView, TSharedPtr<FWidgetLayerListItem> InListItem )
+	{
+		ListItem = InListItem;
+		
+		STableRow<TSharedPtr<FWidgetLayerListItem>>::Construct(
+			STableRow<TSharedPtr<FWidgetLayerListItem>>::FArguments()
+			.OnDragDetected(this, &SWidgetLayerListItem::OnDragDetected)
+			.OnCanAcceptDrop(this, &SWidgetLayerListItem::OnCanAcceptDrop)
+			.OnAcceptDrop(this, &SWidgetLayerListItem::OnAcceptDrop)
+			.Padding(FMargin(0, 0, 30, 0))
+			.Content()
+			[
+				SNew(SBox)
+				.Padding(FMargin(0, 2))
+				[
+					SNew(SHorizontalBox)
+
+					+ SHorizontalBox::Slot()
+					.AutoWidth()
+					.Padding(5, 0, 0, 0)
+					[
+						SNew(SImage)
+						.Image(FCoreStyle::Get().GetBrush("VerticalBoxDragIndicatorShort"))
+						.Visibility_Lambda([this, InListItem]()
+						{
+							return IsHovered() && InListItem->bAllowedToDrag ? EVisibility::Visible : EVisibility::Hidden;
+						})
+					]
+
+					+ SHorizontalBox::Slot()
+					.Padding(FMargin(10, 0, 0, 0))
+					[
+						SAssignNew(TextBlock, STextBlock)
+						.Font(FCoreStyle::Get().GetFontStyle("NormalFont"))
+						.MinDesiredWidth(150)
+						.Text(this, &SWidgetLayerListItem::GetLayerName)
+						.Justification(ETextJustify::Left)
+						.ColorAndOpacity(InListItem->bAllowedToDrag ? FLinearColor::White : FLinearColor(0.25, 0.25, 0.25))
+					]
+				]
+			],
+			InOwnerTableView);
+	}
+
+private:
+	FText GetLayerName() const
+	{
+		return ListItem.IsValid() ? FText::FromName(ListItem.Pin()->LayerName) : FText::GetEmpty();
+	}
+
+	/** Called whenever a drag is detected by the tree view. */
+	FReply OnDragDetected(const FGeometry& InGeometry, const FPointerEvent& InPointerEvent)
+	{
+		const TSharedPtr<FWidgetLayerListItem> ListItemPinned = ListItem.Pin();
+		if (ListItemPinned.IsValid())
+		{
+			if (ListItemPinned->bAllowedToDrag == false)
+			{
+				return FReply::Unhandled();
+			}
+			
+			const FText DefaultText = FText::Format(LOCTEXT("DefaultDragDropText", "Move {0}"), GetLayerName());
+			return FReply::Handled().BeginDragDrop(FWidgetLayerListDragDropOp::New(ListItemPinned, DefaultText));
+		}
+		return FReply::Unhandled();
+	}
+
+	/** Called to determine whether a current drag operation is valid for this row. */
+	TOptional<EItemDropZone> OnCanAcceptDrop(const FDragDropEvent& DragDropEvent, EItemDropZone InItemDropZone, TSharedPtr<FWidgetLayerListItem> InListItem)
+	{
+		const TSharedPtr<FWidgetLayerListDragDropOp> DragDropOp = DragDropEvent.GetOperationAs<FWidgetLayerListDragDropOp>();
+		
+		if (DragDropOp.IsValid() && InListItem.Get()->LayerGuid != DragDropOp.Get()->ListItem.Get()->LayerGuid)
+		{
+			if (InItemDropZone == EItemDropZone::OntoItem)
+			{
+				return TOptional<EItemDropZone>();
+			}
+			
+			return InItemDropZone;
+		}
+		
+		return TOptional<EItemDropZone>();
+	}
+
+	/** Called to complete a drag and drop onto this drop. */
+	FReply OnAcceptDrop(const FDragDropEvent& DragDropEvent, EItemDropZone InItemDropZone, TSharedPtr<FWidgetLayerListItem> InListItem)
+	{
+		const TSharedPtr<FWidgetLayerListDragDropOp> DragDropOp = DragDropEvent.GetOperationAs<FWidgetLayerListDragDropOp>();
+		if (DragDropOp.IsValid() 
+			&& DragDropOp->ListItem.IsValid() && InListItem.IsValid()
+			&& DragDropOp->ListItem->LayerGuid != InListItem->LayerGuid)
+		{			
+			// Copy layer to local variable and remove from list
+			const TSharedPtr<FWidgetLayerListItem> LayerToMove = DragDropOp->ListItem;
+			DragDropOp->ListItem->WidgetLayerList->RemoveAt(InListItem->WidgetLayerList->IndexOfByKey(DragDropOp->ListItem));
+
+			// Determine new index based on drag+drop op
+			const int32 RelativeNewIndex = InListItem->WidgetLayerList->IndexOfByKey(InListItem) + (InItemDropZone == EItemDropZone::AboveItem ? 0 : 1);
+			
+			// Insert copied layer into list at new position and update pointer in DragDropOp->ListItem
+			DragDropOp->ListItem->WidgetLayerList->Insert(LayerToMove, RelativeNewIndex);
+
+			// Reconstruct and refresh ListView
+			InListItem->OnLayerListUpdated();
+			
+			return FReply::Handled();
+		}
+		return FReply::Unhandled();
+	}
+
+private:
+	TWeakPtr<FWidgetLayerListItem> ListItem;
+	TSharedPtr<STextBlock> TextBlock;
+};
+
+/* SLandscapeLayerListDialog implementation */
+
+void SLandscapeLayerListDialog::Construct(const FArguments& InArgs, TArray<FLandscapeLayer>& InLayers)
+{
+	LayerList = &InLayers;
+
+	for(int LayerIndex = 0; LayerIndex < LayerList->Num(); ++LayerIndex)
+	{
+		WidgetLayerList.Add(MakeShared<FWidgetLayerListItem>(&(*LayerList)[LayerList->Num() - 1 - LayerIndex], [this](){ OnLayerListUpdated(); }, &WidgetLayerList));
+	}
+
+	WidgetLayerList[0]->bAllowedToDrag = true;
+	InsertedLayerIndex = InLayers.Num() - 1;
+	
+	// Construct list view
+	SAssignNew(LayerListView, SWidgetLayerListView)
+	.ItemHeight(20.0f)
+	.SelectionMode(ESelectionMode::Single)
+	.OnGenerateRow(this, &SLandscapeLayerListDialog::OnGenerateRow)
+	.ListItemsSource(&WidgetLayerList);
+	
+	FButton CancelButton(LOCTEXT("Cancel", "Cancel"));
+	FButton AcceptButton(LOCTEXT("Accept", "Accept"));
+	AcceptButton.OnClicked.BindSP(this, &SLandscapeLayerListDialog::OnAccept);
+	
+	// Construct custom dialog with list view supporting drag + drop
+	SCustomDialog::Construct(SCustomDialog::FArguments()
+		.Title(FText(LOCTEXT("LandscapeLayerListDialogTitleText", "Insert New Landscape Edit Layer")))
+		.UseScrollBox(false)
+		.Content()
+		[
+			SNew(SBox)
+			// Right padding of 20 is to offset SCustomDialog NullWidget to the left of content
+			.Padding(FMargin(0, 0, 20, 0))
+			[
+				SNew(SVerticalBox)
+				+ SVerticalBox::Slot()
+				.AutoHeight()
+				[
+					SNew(STextBlock)
+					.Font(FCoreStyle::Get().GetFontStyle("NormalFont"))
+					.WrapTextAt(350)
+					.Text(FText::Format(LOCTEXT(
+						"LandscapeLayerListDialogInstructionText",
+						"Drag/drop the \"{0}\" layer in the list to choose where in the edit layers stack it should be inserted.\n"
+						"\n(Note: Closing the window or pressing cancel will result in the \"{0}\" layer going onto the top of the layer stack.)\n"),
+						FText::FromName(WidgetLayerList[0]->LayerName)))
+				]
+				+ SVerticalBox::Slot()
+				.MaxHeight(125)
+				.HAlign(HAlign_Center)
+				[
+					SNew(SBorder)
+					[
+						LayerListView.ToSharedRef()
+					]
+				]
+			]
+		]
+		.Buttons({
+			AcceptButton,
+			CancelButton
+		})
+	);
+}
+
+void SLandscapeLayerListDialog::OnLayerListUpdated()
+{
+	if (LayerListView.IsValid())
+	{
+		LayerListView->RequestListRefresh();
+	}
+}
+
+TSharedRef<ITableRow> SLandscapeLayerListDialog::OnGenerateRow(TSharedPtr<FWidgetLayerListItem> InListItem, const TSharedRef< STableViewBase >& InOwnerTableView) const
+{
+	return SNew( SWidgetLayerListItem, InOwnerTableView, InListItem );
+}
+
+void SLandscapeLayerListDialog::OnAccept()
+{
+	// Find the draggable layer in WidgetLayerList
+	for (int WidgetLayerListIndex = 0; WidgetLayerListIndex < WidgetLayerList.Num(); ++WidgetLayerListIndex)
+	{
+		const TSharedPtr<FWidgetLayerListItem>& WidgetLayer = WidgetLayerList[WidgetLayerList.Num() - 1 - WidgetLayerListIndex];
+		
+		if (WidgetLayer->bAllowedToDrag)
+		{
+			// Find actual edit layer corresponding to draggable layer
+			for (int LayerIndex = 0; LayerIndex < LayerList->Num(); ++LayerIndex)
+			{
+				const FLandscapeLayer& Layer = (*LayerList)[LayerIndex];
+
+				// Once found, remove and copy edit layer into correct spot in LayerList
+				if (WidgetLayer->LayerGuid == Layer.Guid)
+				{
+					const FLandscapeLayer LayerCopy = Layer;
+					
+					LayerList->RemoveAt(LayerIndex);
+					InsertedLayerIndex = WidgetLayerListIndex;
+					LayerList->Insert(LayerCopy, InsertedLayerIndex);
+					
+					return;
+				}
+			}
+		}
+	}
+}
+
+#undef LOCTEXT_NAMESPACE 
\ No newline at end of file
diff --git a/Engine/Source/Editor/LandscapeEditor/Public/SLandscapeLayerListDialog.h b/Engine/Source/Editor/LandscapeEditor/Public/SLandscapeLayerListDialog.h
new file mode 100644
--- /dev/null
+++ b/Engine/Source/Editor/LandscapeEditor/Public/SLandscapeLayerListDialog.h
@@ -0,0 +1,57 @@
+// Copyright Epic Games, Inc. All Rights Reserved.
+
+#pragma once
+
+#include "CoreMinimal.h"
+#include "Widgets/Views/SListView.h"	// SWidgetLayerListView
+#include "Dialog/SCustomDialog.h"		// SLandscapeLayerListDialog
+
+class FWidgetLayerListDragDropOp;
+struct FLandscapeLayer;
+struct FWidgetLayerListItem;
+typedef SListView<TSharedPtr<FWidgetLayerListItem> > SWidgetLayerListView;
+
+/**
+ *	Displays a list of layers with drag + drop support for layer reordering outside of landscape mode.
+ *	Can be directly constructed in editor code or can be constructed indirectly through ILandscapeEditorServices.
+ */
+class SLandscapeLayerListDialog : public SCustomDialog
+{
+public:
+	friend class FLandscapeEditorServices;
+	
+	virtual ~SLandscapeLayerListDialog() {}
+	
+	SLATE_BEGIN_ARGS( SLandscapeLayerListDialog	){}
+	SLATE_END_ARGS()
+
+	/**
+	 * Construct this widget. Called by the SNew() Slate macro.
+	 * @param InArgs Declaration used by the SNew() macro to construct this widget
+	 * @param InLayers The array of landscape layers which should be displayed
+	 */
+	void Construct(const FArguments& InArgs, TArray<FLandscapeLayer>& InLayers);
+
+private:
+	/* Called when LayerList is changed to update WidgetLayerList and the dialog window. */
+	void OnLayerListUpdated();
+	
+	/* Called when a new row is being generated. Converts FWidgetLayerListItem to ITableRow. */
+	TSharedRef<ITableRow> OnGenerateRow( TSharedPtr<FWidgetLayerListItem> InListItem, const TSharedRef< STableViewBase >& InOwnerTableView ) const;
+
+	/* Called when Accept button is pressed */
+	void OnAccept();
+
+	/* The list view of WidgetLayerList. */
+	TSharedPtr<SWidgetLayerListView> LayerListView;
+
+	/* Array of FWidgetLayerListItem, which contains FLandscapeLayer* in addition to intermediate data.
+	 * Note: This array is a UI reflection of LayerList but the order is the reverse order of LayerList */
+	TArray<TSharedPtr<FWidgetLayerListItem>> WidgetLayerList;
+
+	/* The original list of layers passed in Construct. */
+	TArray<FLandscapeLayer>* LayerList = nullptr;
+
+	/* Index of the layer which is being inserted into the layer stack */
+	int InsertedLayerIndex = -1;
+};
\ No newline at end of file
diff --git a/Engine/Source/Runtime/Landscape/Private/LandscapeModule.cpp b/Engine/Source/Runtime/Landscape/Private/LandscapeModule.cpp
--- a/Engine/Source/Runtime/Landscape/Private/LandscapeModule.cpp
+++ b/Engine/Source/Runtime/Landscape/Private/LandscapeModule.cpp
@@ -36,12 +36,16 @@ public:
 	void ShutdownModule() override;
 
 	virtual TSharedPtr<FLandscapeSceneViewExtension, ESPMode::ThreadSafe> GetLandscapeSceneViewExtension() const override { return SceneViewExtension; }
+
+	virtual void SetLandscapeEditorServices(ILandscapeEditorServices* InLandscapeEditorServices) override { LandscapeEditorServices = InLandscapeEditorServices; }
+	virtual ILandscapeEditorServices* GetLandscapeEditorServices() const override { return LandscapeEditorServices; }
 private:
 	void OnPostEngineInit();
 	void OnEnginePreExit();
 
 private:
 	TSharedPtr<FLandscapeSceneViewExtension, ESPMode::ThreadSafe> SceneViewExtension;
+	ILandscapeEditorServices* LandscapeEditorServices = nullptr;
 };
 
 /**
diff --git a/Engine/Source/Runtime/Landscape/Public/LandscapeEditorServices.h b/Engine/Source/Runtime/Landscape/Public/LandscapeEditorServices.h
new file mode 100644
--- /dev/null
+++ b/Engine/Source/Runtime/Landscape/Public/LandscapeEditorServices.h
@@ -0,0 +1,32 @@
+// Copyright Epic Games, Inc. All Rights Reserved.
+
+#pragma once
+
+#include "Landscape.h"
+
+class ILandscapeEditorServices
+{
+public:
+	virtual ~ILandscapeEditorServices() {}
+
+	/**
+	 * Attempts to find an edit layer named InEditLayerName in InTargetLandscape.
+	 * Creates the layer if it does not exist.
+	 * 
+	 * @param InEditLayerName The name of the layer to search for and possibly create
+	 * @param InTargetLandscape The landscape which should have a layer called InEditLayerName
+	 * @return The index at which the edit layer named InEditLayerName exists
+	 */
+	virtual int32 GetOrCreateEditLayer(FName InEditLayerName, ALandscape* InTargetLandscape) = 0;
+};
+
+#if WITH_EDITOR
+class FLandscapeEditorServices : public ILandscapeEditorServices
+{
+public:
+	virtual ~FLandscapeEditorServices() {}
+
+	// Implements pure virtual function. Insertion logic is left to the user through modal drag + drop dialog.
+	virtual int32 GetOrCreateEditLayer(FName InEditLayerName, ALandscape* InTargetLandscape);
+};
+#endif // WITH_EDITOR
\ No newline at end of file
diff --git a/Engine/Source/Runtime/Landscape/Public/LandscapeModule.h b/Engine/Source/Runtime/Landscape/Public/LandscapeModule.h
--- a/Engine/Source/Runtime/Landscape/Public/LandscapeModule.h
+++ b/Engine/Source/Runtime/Landscape/Public/LandscapeModule.h
@@ -4,6 +4,7 @@
 
 #include "CoreMinimal.h"
 #include "Modules/ModuleInterface.h"
+#include "LandscapeEditorServices.h"
 
 class FLandscapeSceneViewExtension;
 
@@ -14,4 +15,7 @@ class ILandscapeModule : public IModuleInterface
 {
 public:
 	virtual TSharedPtr<FLandscapeSceneViewExtension, ESPMode::ThreadSafe> GetLandscapeSceneViewExtension() const = 0;
+	
+	virtual LANDSCAPE_API void SetLandscapeEditorServices(ILandscapeEditorServices* InLandscapeEditorServices) = 0;
+	virtual LANDSCAPE_API ILandscapeEditorServices* GetLandscapeEditorServices() const = 0;
 };
