* Fixed single file export for heightmap/weightmaps
* Added option to export to single file for grid-based worlds
* Fixed a problem where you couldn't import a single heightmap/weightmap file if it was named with the same prefix as the grid-based ones located in the same folder (e.g. Heightmap.png with Heightmap_x0_y0.png, etc. being present in the same folder)
* fixed crash if file pattern resulted in no files
* set import gizmo to location of last export

diff --git a/Engine/Source/Editor/LandscapeEditor/Private/LandscapeEdModeComponentTools.cpp b/Engine/Source/Editor/LandscapeEditor/Private/LandscapeEdModeComponentTools.cpp
--- a/Engine/Source/Editor/LandscapeEditor/Private/LandscapeEdModeComponentTools.cpp
+++ b/Engine/Source/Editor/LandscapeEditor/Private/LandscapeEdModeComponentTools.cpp
@@ -1875,7 +1875,7 @@ public:
 		DraggingEdge = ELandscapeEdge::None;
 		DraggingEdge_Remainder = 0.0f;
 		EdMode->NewLandscapePreviewMode = NewLandscapePreviewMode;
-		EdMode->UISettings->ImportLandscapeData();
+		EdMode->UISettings->RefreshImports();
 	}
 
 	virtual void ExitTool() override
diff --git a/Engine/Source/Editor/LandscapeEditor/Private/LandscapeEditorDetailCustomization_ImportExport.cpp b/Engine/Source/Editor/LandscapeEditor/Private/LandscapeEditorDetailCustomization_ImportExport.cpp
--- a/Engine/Source/Editor/LandscapeEditor/Private/LandscapeEditorDetailCustomization_ImportExport.cpp
+++ b/Engine/Source/Editor/LandscapeEditor/Private/LandscapeEditorDetailCustomization_ImportExport.cpp
@@ -73,21 +73,24 @@ void FLandscapeEditorDetailCustomization_ImportExport::CustomizeDetails(IDetailL
 	
 	PropertyHandle_ImportHeightmapFilename->SetOnPropertyValueChanged(FSimpleDelegate::CreateLambda([PropertyHandle_ImportHeightmapFilename]()
 	{
-		FLandscapeEditorDetailCustomization_ImportExport::FormatFilename(PropertyHandle_ImportHeightmapFilename);
+		FLandscapeEditorDetailCustomization_ImportExport::FormatFilename(PropertyHandle_ImportHeightmapFilename, /*bForExport = */false);
 		FLandscapeEditorDetailCustomization_ImportExport::OnImportHeightmapFilenameChanged();
 	}));
 
 	PropertyHandle_ExportHeightmapFilename->SetOnPropertyValueChanged(FSimpleDelegate::CreateLambda([PropertyHandle_ExportHeightmapFilename]()
 	{
-		FLandscapeEditorDetailCustomization_ImportExport::FormatFilename(PropertyHandle_ExportHeightmapFilename);
+		FLandscapeEditorDetailCustomization_ImportExport::FormatFilename(PropertyHandle_ExportHeightmapFilename, /*bForExport = */true);
 	}));
-
+	
 	TSharedRef<IPropertyHandle> PropertyHandle_HeightmapImportResult = DetailBuilder.GetProperty(GET_MEMBER_NAME_CHECKED(ULandscapeEditorObject, ImportLandscape_HeightmapImportResult));
 	TSharedRef<IPropertyHandle> PropertyHandle_HeightmapErrorMessage = DetailBuilder.GetProperty(GET_MEMBER_NAME_CHECKED(ULandscapeEditorObject, ImportLandscape_HeightmapErrorMessage));
 	DetailBuilder.HideProperty(PropertyHandle_HeightmapImportResult);
 	DetailBuilder.HideProperty(PropertyHandle_HeightmapErrorMessage);
 		
 	TSharedRef<IPropertyHandle> PropertyHandle_ExportEditLayer = DetailBuilder.GetProperty(GET_MEMBER_NAME_CHECKED(ULandscapeEditorObject, bExportEditLayer));
+	TSharedRef<IPropertyHandle> PropertyHandle_ExportSingleFile = DetailBuilder.GetProperty(GET_MEMBER_NAME_CHECKED(ULandscapeEditorObject, bExportSingleFile));
+	DetailBuilder.HideProperty(PropertyHandle_ExportEditLayer);
+	DetailBuilder.HideProperty(PropertyHandle_ExportSingleFile);
 
 	auto AddHeightmapFileName = [&](TSharedRef<IPropertyHandle> PropertyHandle_HeightmapFilename, TAttribute<EVisibility> PropertyVisibility)
 	{
@@ -157,9 +160,39 @@ void FLandscapeEditorDetailCustomization_ImportExport::CustomizeDetails(IDetailL
 	ImportExportCategory.AddProperty(PropertyHandle_ImportType).Visibility(MakeAttributeLambda([]() { return FLandscapeEditorDetailCustomization_ImportExport::GetImportExportVisibility(true); }));
 
 	TSharedRef<IPropertyHandle> PropertyHandle_FlipYAxis = DetailBuilder.GetProperty(GET_MEMBER_NAME_CHECKED(ULandscapeEditorObject, bFlipYAxis));
-	PropertyHandle_FlipYAxis->SetOnPropertyValueChanged(FSimpleDelegate::CreateLambda([PropertyHandle_FlipYAxis]() { FLandscapeEditorDetailCustomization_ImportExport::OnImportHeightmapFilenameChanged(); }));
+	PropertyHandle_FlipYAxis->SetOnPropertyValueChanged(FSimpleDelegate::CreateLambda([]() { FLandscapeEditorDetailCustomization_ImportExport::OnImportHeightmapFilenameChanged(); }));
 	ImportExportCategory.AddProperty(PropertyHandle_FlipYAxis).Visibility(MakeAttributeLambda([]() { return (!GetEditorMode()->UseSingleFileImport()) ? EVisibility::Visible : EVisibility::Collapsed; }));
 
+	TSharedPtr<SToolTip> ExportSingleFileTooltip = SNew(SToolTip)
+		.Text(PropertyHandle_ExportSingleFile->GetToolTipText());
+
+	ImportExportCategory.AddProperty(PropertyHandle_ExportSingleFile)
+		.Visibility(MakeAttributeLambda([]() { return FLandscapeEditorDetailCustomization_ImportExport::GetImportExportVisibility(false); }))
+		.IsEnabled(MakeAttributeLambda([]() { return FLandscapeEditorDetailCustomization_ImportExport::GetExportSingleFileIsEnabled(); }))
+		.CustomWidget()
+		.NameContent()
+		[
+			SNew(SBox)
+			.VAlign(VAlign_Center)
+			.Padding(FMargin(0, 0, 2, 0))
+			[
+				SNew(STextBlock)
+				.Text(PropertyHandle_ExportSingleFile->GetPropertyDisplayName())
+				.Font(DetailBuilder.GetDetailFont())
+				.ToolTip(ExportSingleFileTooltip)
+			]
+		]
+		.ValueContent()
+		[
+			SNew(SCheckBox)
+			.IsChecked_Static(&FLandscapeEditorDetailCustomization_ImportExport::GetExportSingleFileCheckState)
+			.OnCheckStateChanged_Lambda([this, PropertyHandle_ExportHeightmapFilename](ECheckBoxState NewCheckState) 
+				{ 
+					FLandscapeEditorDetailCustomization_ImportExport::OnExportSingleFileCheckStateChanged(NewCheckState);
+					FLandscapeEditorDetailCustomization_ImportExport::FormatFilename(PropertyHandle_ExportHeightmapFilename, /*bForExport = */true);
+				})
+		];
+
 	ImportExportCategory.AddProperty(PropertyHandle_ExportEditLayer).Visibility(MakeAttributeLambda([this]()
 	{
 		if (ULandscapeInfo* LandscapeInfo = GetEditorMode()->CurrentToolTarget.LandscapeInfo.Get())
@@ -230,6 +263,35 @@ void FLandscapeEditorDetailCustomization_ImportExport::CustomizeDetails(IDetailL
 }
 END_SLATE_FUNCTION_BUILD_OPTIMIZATION
 
+bool FLandscapeEditorDetailCustomization_ImportExport::GetExportSingleFileIsEnabled()
+{
+	FEdModeLandscape* LandscapeEdMode = GetEditorMode();
+	check(LandscapeEdMode);
+	// The single file option only makes sense for landscapes where multiple file import is an option : 
+	return !LandscapeEdMode->UseSingleFileImport();
+}
+
+ECheckBoxState FLandscapeEditorDetailCustomization_ImportExport::GetExportSingleFileCheckState()
+{
+	FEdModeLandscape* LandscapeEdMode = GetEditorMode();
+	check(LandscapeEdMode);
+	return (LandscapeEdMode->UseSingleFileImport() || (LandscapeEdMode->UISettings &&  LandscapeEdMode->UISettings->bExportSingleFile)) ? ECheckBoxState::Checked : ECheckBoxState::Unchecked;
+}
+
+void FLandscapeEditorDetailCustomization_ImportExport::OnExportSingleFileCheckStateChanged(ECheckBoxState NewCheckState)
+{
+	FEdModeLandscape* LandscapeEdMode = GetEditorMode();
+	check(LandscapeEdMode);
+
+	if (LandscapeEdMode->UISettings)
+	{
+		if (NewCheckState != ECheckBoxState::Undetermined)
+		{
+			LandscapeEdMode->UISettings->bExportSingleFile = (NewCheckState == ECheckBoxState::Checked);
+		}
+	}
+}
+
 EVisibility FLandscapeEditorDetailCustomization_ImportExport::GetImportExportVisibility(bool bImport)
 {
 	return IsImporting() == bImport ? EVisibility::Visible : EVisibility::Collapsed;
@@ -252,15 +314,15 @@ ECheckBoxState FLandscapeEditorDetailCustomization_ImportExport::GetHeightmapSel
 	return IsHeightmapEnabled() ? ECheckBoxState::Checked : ECheckBoxState::Unchecked;
 }
 
-void FLandscapeEditorDetailCustomization_ImportExport::OnHeightmapSelectedCheckStateChanged(ECheckBoxState CheckState)
+void FLandscapeEditorDetailCustomization_ImportExport::OnHeightmapSelectedCheckStateChanged(ECheckBoxState NewCheckedState)
 {
 	FEdModeLandscape* LandscapeEdMode = GetEditorMode();
 	check(LandscapeEdMode != nullptr);
 	if (LandscapeEdMode->UISettings)
 	{
-		if (CheckState != ECheckBoxState::Undetermined)
+		if (NewCheckedState != ECheckBoxState::Undetermined)
 		{
-			LandscapeEdMode->UISettings->bHeightmapSelected = CheckState == ECheckBoxState::Checked;
+			LandscapeEdMode->UISettings->bHeightmapSelected = (NewCheckedState == ECheckBoxState::Checked);
 		}
 	}
 }
@@ -279,11 +341,30 @@ void FLandscapeEditorDetailCustomization_ImportExport::OnModeChanged(ECheckBoxSt
 		FEdModeLandscape* LandscapeEdMode = GetEditorMode();
 		check(LandscapeEdMode != nullptr);
 		LandscapeEdMode->ImportExportMode = Value;
+
+		if (Value == EImportExportMode::Import)
+		{
+			if (ULandscapeInfo* LandscapeInfo = LandscapeEdMode->CurrentToolTarget.LandscapeInfo.Get())
+			{
+				FVector LocalPosition(LandscapeEdMode->UISettings->ImportLandscape_GizmoLocalPosition.X, LandscapeEdMode->UISettings->ImportLandscape_GizmoLocalPosition.Y, 0.0f);
+				FVector GizmoPosition = LandscapeInfo->GetLandscapeProxy()->LandscapeActorToWorld().TransformPosition(LocalPosition);
+				if (ALandscapeGizmoActiveActor* GizmoActor = LandscapeEdMode->CurrentGizmoActor.Get())
+				{
+					GizmoActor->SetActorLocation(GizmoPosition);
+				}
+			}
+		}
 	}
 }
 
 EVisibility FLandscapeEditorDetailCustomization_ImportExport::GetImportResultErrorVisibility(TSharedRef<IPropertyHandle> PropertyHandle_Result)
 {
+	// Not relevant when exporting : 
+	if (!IsImporting())
+	{
+		return EVisibility::Collapsed;
+	}
+
 	ELandscapeImportResult ImportResult;
 	FPropertyAccess::Result Result = PropertyHandle_Result->GetValue((uint8&)ImportResult);
 
@@ -388,11 +469,12 @@ void FLandscapeEditorDetailCustomization_ImportExport::OnImportHeightmapFilename
 	LandscapeEdMode->UISettings->OnImportHeightmapFilenameChanged();
 }
 
-void FLandscapeEditorDetailCustomization_ImportExport::FormatFilename(TSharedRef<IPropertyHandle> PropertyHandle_Filename)
+void FLandscapeEditorDetailCustomization_ImportExport::FormatFilename(TSharedRef<IPropertyHandle> PropertyHandle_Filename, bool bForExport)
 {
 	FEdModeLandscape* LandscapeEdMode = GetEditorMode();
 	check(LandscapeEdMode != nullptr);
-	if (!LandscapeEdMode->UseSingleFileImport())
+	// No need to format the file name if single file import/export is enabled
+	if (!LandscapeEdMode->UseSingleFileImport() && (!bForExport || !LandscapeEdMode->UISettings->bExportSingleFile))
 	{
 		// If selected file as the _xN_yM suffix remove it
 		FString	FilePath;
@@ -480,7 +562,7 @@ FReply FLandscapeEditorDetailCustomization_ImportExport::OnImportExportButtonCli
 		}
 		else
 		{
-			auto BuildExportFileName = [](const FString& Filename, FIntPoint FileOffset, bool bUseOffset)
+			auto BuildExportFileName = [](const FString& Filename, const FIntPoint& FileOffset, bool bUseOffset)
 			{
 				if (bUseOffset)
 				{
@@ -492,47 +574,63 @@ FReply FLandscapeEditorDetailCustomization_ImportExport::OnImportExportButtonCli
 				return Filename;
 			};
 
+			auto PerformExport = [BuildExportFileName](ULandscapeEditorObject* InLandscapeEditorSettings, ULandscapeInfo* InLandscapeInfo, const FIntRect& InExportRegion, const FIntPoint& InFileOffset, bool bInUseOffset)
+			{
+				if (InLandscapeEditorSettings->bHeightmapSelected)
+				{
+					FString ExportFilename = BuildExportFileName(InLandscapeEditorSettings->HeightmapExportFilename, InFileOffset, bInUseOffset);
+					InLandscapeInfo->ExportHeightmap(ExportFilename, InExportRegion);
+				}
+
+				for (const FLandscapeImportLayer& ImportLayer : InLandscapeEditorSettings->ImportLandscape_Layers)
+				{
+					if (ImportLayer.bSelected)
+					{
+						FString ExportFilename = BuildExportFileName(ImportLayer.ExportFilePath, InFileOffset, bInUseOffset);
+						InLandscapeInfo->ExportLayer(ImportLayer.LayerInfo, ExportFilename, InExportRegion);
+					}
+				}
+			};
+
 			FScopedSetLandscapeEditingLayer Scope(LandscapeInfo->LandscapeActor.Get(), LandscapeEdMode->UISettings->bExportEditLayer ? LandscapeEdMode->GetCurrentLayerGuid() : FGuid());
 
-			LandscapeInfo->ForAllLandscapeProxies([LandscapeInfo, LandscapeEdMode, LandscapeExtent, BuildExportFileName](ALandscapeProxy* LandscapeProxy)
+			ALandscape* Landscape = LandscapeInfo->LandscapeActor.Get();
+			check(Landscape);
+			if (LandscapeEdMode->UseSingleFileImport() || LandscapeEdMode->UISettings->bExportSingleFile)
 			{
+				// For single file export, export the whole landscape (main actor) : 
 				FIntRect ExportRegion;
-				if (LandscapeInfo->GetLandscapeExtent(LandscapeProxy, ExportRegion))
+				if (LandscapeInfo->GetLandscapeExtent(ExportRegion))
+				{
+					PerformExport(LandscapeEdMode->UISettings, LandscapeInfo, ExportRegion, FIntPoint(), /*bInUseOffset = */false);
+				}
+			}
+			else
+			{
+				// For multiple file export, export each landscape proxy individually :
+				LandscapeInfo->ForAllLandscapeProxies([LandscapeInfo, Landscape, LandscapeEdMode, LandscapeExtent, BuildExportFileName, PerformExport](ALandscapeProxy* LandscapeProxy)
 				{
-					FIntPoint FileOffset{};
-					if (!LandscapeEdMode->UseSingleFileImport())
+					FIntRect ExportRegion;
+					if (LandscapeInfo->GetLandscapeExtent(LandscapeProxy, ExportRegion))
 					{
-						ALandscape* Landscape = LandscapeInfo->LandscapeActor.Get();
-						check(Landscape);
-
 						const int32 YCoord = LandscapeEdMode->UISettings->bFlipYAxis ? LandscapeExtent.Max.Y - ExportRegion.Max.Y : ExportRegion.Min.Y - LandscapeExtent.Min.Y;
-						FileOffset = FIntPoint((ExportRegion.Min.X - LandscapeExtent.Min.X) / Landscape->GridSize,
+						FIntPoint FileOffset = FIntPoint((ExportRegion.Min.X - LandscapeExtent.Min.X) / Landscape->GridSize,
 							YCoord / Landscape->GridSize);
 
 						// Remove the shared line/column that this proxy has with its neighbors because it 
 						// will be included by the neighbor or lost if there is none (that could become an option to avoid that loss)
 						ExportRegion.Max.X -= 1;
 						ExportRegion.Max.Y -= 1;
-					}
-											
 
-					if (LandscapeEdMode->UISettings->bHeightmapSelected)
-					{
-						FString ExportFilename = BuildExportFileName(LandscapeEdMode->UISettings->HeightmapExportFilename, FileOffset, !LandscapeEdMode->UseSingleFileImport());
-						LandscapeInfo->ExportHeightmap(ExportFilename, ExportRegion);
-					}
-					
-					for (const FLandscapeImportLayer& ImportLayer : LandscapeEdMode->UISettings->ImportLandscape_Layers)
-					{
-						if (ImportLayer.bSelected)
-						{
-							FString ExportFilename = BuildExportFileName(ImportLayer.ExportFilePath, FileOffset, !LandscapeEdMode->UseSingleFileImport());
-							LandscapeInfo->ExportLayer(ImportLayer.LayerInfo, ExportFilename, ExportRegion);
-						}
+						PerformExport(LandscapeEdMode->UISettings, LandscapeInfo, ExportRegion, FileOffset, /*bInUseOffset = */true);
 					}
-				}
-			});
-			
+				});
+			}
+
+			FIntRect Extent;
+			LandscapeInfo->GetLandscapeExtent(Extent);
+			// Set the import gizmo location to landscape minx/miny
+			LandscapeEdMode->UISettings->ImportLandscape_GizmoLocalPosition = Extent.Min;
 		}
 	}
 
diff --git a/Engine/Source/Editor/LandscapeEditor/Private/LandscapeEditorDetailCustomization_ImportExport.h b/Engine/Source/Editor/LandscapeEditor/Private/LandscapeEditorDetailCustomization_ImportExport.h
--- a/Engine/Source/Editor/LandscapeEditor/Private/LandscapeEditorDetailCustomization_ImportExport.h
+++ b/Engine/Source/Editor/LandscapeEditor/Private/LandscapeEditorDetailCustomization_ImportExport.h
@@ -30,8 +30,11 @@ public:
 	/** IDetailCustomization interface */
 	virtual void CustomizeDetails(IDetailLayoutBuilder& DetailBuilder) override;
 
-	static void FormatFilename(TSharedRef<IPropertyHandle> PropertyHandle_Filename);
+	static void FormatFilename(TSharedRef<IPropertyHandle> PropertyHandle_Filename, bool bForExport);
 private:
+	static bool GetExportSingleFileIsEnabled();
+	static ECheckBoxState GetExportSingleFileCheckState();
+	static void OnExportSingleFileCheckStateChanged(ECheckBoxState NewCheckState);
 	static EVisibility GetImportExportVisibility(bool bImport);
 	static bool IsHeightmapEnabled();
 	static ECheckBoxState GetHeightmapSelectedCheckState();
diff --git a/Engine/Source/Editor/LandscapeEditor/Private/LandscapeEditorDetailCustomization_ImportLayers.cpp b/Engine/Source/Editor/LandscapeEditor/Private/LandscapeEditorDetailCustomization_ImportLayers.cpp
--- a/Engine/Source/Editor/LandscapeEditor/Private/LandscapeEditorDetailCustomization_ImportLayers.cpp
+++ b/Engine/Source/Editor/LandscapeEditor/Private/LandscapeEditorDetailCustomization_ImportLayers.cpp
@@ -59,13 +59,13 @@ void FLandscapeEditorStructCustomization_FLandscapeImportLayer::CustomizeChildre
 	
 	PropertyHandle_SourceFilePath->SetOnPropertyValueChanged(FSimpleDelegate::CreateLambda([PropertyHandle_SourceFilePath]()
 	{
-		FLandscapeEditorDetailCustomization_ImportExport::FormatFilename(PropertyHandle_SourceFilePath);
+		FLandscapeEditorDetailCustomization_ImportExport::FormatFilename(PropertyHandle_SourceFilePath, /*bForExport = */false);
 		FLandscapeEditorStructCustomization_FLandscapeImportLayer::OnImportWeightmapFilenameChanged();
 	}));
 
 	PropertyHandle_ExportFilePath->SetOnPropertyValueChanged(FSimpleDelegate::CreateLambda([PropertyHandle_ExportFilePath]()
 	{
-		FLandscapeEditorDetailCustomization_ImportExport::FormatFilename(PropertyHandle_ExportFilePath);
+		FLandscapeEditorDetailCustomization_ImportExport::FormatFilename(PropertyHandle_ExportFilePath, /*bForExport = */true);
 	}));
 
 	FName LayerName;
diff --git a/Engine/Source/Editor/LandscapeEditor/Private/LandscapeEditorDetailCustomization_NewLandscape.cpp b/Engine/Source/Editor/LandscapeEditor/Private/LandscapeEditorDetailCustomization_NewLandscape.cpp
--- a/Engine/Source/Editor/LandscapeEditor/Private/LandscapeEditorDetailCustomization_NewLandscape.cpp
+++ b/Engine/Source/Editor/LandscapeEditor/Private/LandscapeEditorDetailCustomization_NewLandscape.cpp
@@ -139,7 +139,7 @@ void FLandscapeEditorDetailCustomization_NewLandscape::CustomizeDetails(IDetailL
 	DetailBuilder.HideProperty(PropertyHandle_HeightmapErrorMessage);
 	PropertyHandle_HeightmapFilename->SetOnPropertyValueChanged(FSimpleDelegate::CreateLambda([this, PropertyHandle_HeightmapFilename]()
 	{
-		FLandscapeEditorDetailCustomization_ImportExport::FormatFilename(PropertyHandle_HeightmapFilename);
+		FLandscapeEditorDetailCustomization_ImportExport::FormatFilename(PropertyHandle_HeightmapFilename, /*bForExport = */false);
 		OnImportHeightmapFilenameChanged();
 	}));
 	NewLandscapeCategory.AddProperty(PropertyHandle_HeightmapFilename)
diff --git a/Engine/Source/Editor/LandscapeEditor/Private/LandscapeImportHelper.cpp b/Engine/Source/Editor/LandscapeEditor/Private/LandscapeImportHelper.cpp
--- a/Engine/Source/Editor/LandscapeEditor/Private/LandscapeImportHelper.cpp
+++ b/Engine/Source/Editor/LandscapeEditor/Private/LandscapeImportHelper.cpp
@@ -10,7 +10,7 @@
 
 #define LOCTEXT_NAMESPACE "LandscapeImportHelper"
 
-bool FLandscapeImportHelper::ExtractCoordinates(FString BaseFilename, FIntPoint& OutCoord, FString& OutBaseFilePattern)
+bool FLandscapeImportHelper::ExtractCoordinates(const FString& BaseFilename, FIntPoint& OutCoord, FString& OutBaseFilePattern)
 {
 	//We expect file name in form: <tilename>_x<number>_y<number>
 	int32 XPos = BaseFilename.Find(TEXT("_x"), ESearchCase::IgnoreCase, ESearchDir::FromEnd);
@@ -34,12 +34,29 @@ bool FLandscapeImportHelper::ExtractCoordinates(FString BaseFilename, FIntPoint&
 
 void FLandscapeImportHelper::GetMatchingFiles(const FString& FilePathPattern, TArray<FString>& OutFileToImport)
 {
-	IFileManager::Get().IterateDirectoryRecursively(*FPaths::GetPath(FilePathPattern), [&OutFileToImport, &FilePathPattern](const TCHAR* FilenameOrDirectory, bool bIsDirectory)
+	FString BaseFilePathPattern = FPaths::GetBaseFilename(FilePathPattern);
+	IFileManager::Get().IterateDirectoryRecursively(*FPaths::GetPath(FilePathPattern), [&OutFileToImport, &FilePathPattern, &BaseFilePathPattern](const TCHAR* FilenameOrDirectory, bool bIsDirectory)
 	{
 		if (!bIsDirectory)
 		{
 			FString Filename(FilenameOrDirectory);
-			if (Filename.StartsWith(FilePathPattern))
+			FString BaseFilename = FPaths::GetBaseFilename(Filename);
+			bool bIsValidFile = false;
+			// The file name must match either exactly (e.g. MyHeightmap.png) :
+			if (BaseFilename == BaseFilePathPattern)
+			{
+				bIsValidFile = true;
+			}
+			// Or it must exactly match the pattern (e.g. MyHeightmap_x0_y0.png ok, MyHeightmap0.png not ok) :
+			else
+			{
+				FIntPoint Coord;
+				FString BaseFilePattern;
+				bool bIsValidPatternFileName = FLandscapeImportHelper::ExtractCoordinates(BaseFilename, Coord, BaseFilePattern);
+				bIsValidFile = bIsValidPatternFileName && (BaseFilePattern == BaseFilePathPattern);
+			}
+
+			if (bIsValidFile)
 			{
 				OutFileToImport.Add(Filename);
 			}
@@ -129,7 +146,7 @@ ELandscapeImportResult GetImportDescriptorInternal(const FString& FilePath, bool
 	FString FilePathPattern;
 	TArray<FString> OutFilesToImport;
 
-	// If we are handling multiple files handle the case where Filepath is in the _xN_yM.extenstion format or if its a pattern already
+	// If we are handling multiple files handle the case where Filepath is in the _xN_yM.extention format or if it's a pattern already
 	if (!bSingleFile)
 	{
 		if (FLandscapeImportHelper::ExtractCoordinates(FPaths::GetBaseFilename(FilePath), OutCoord, OutFileImportPattern))
@@ -142,8 +159,28 @@ ELandscapeImportResult GetImportDescriptorInternal(const FString& FilePath, bool
 		}
 
 		FLandscapeImportHelper::GetMatchingFiles(FilePathPattern, OutFilesToImport);
-		// We were expecting multiple files but only got one.
-		bSingleFile = OutFilesToImport.Num() == 1 && FilePath == OutFilesToImport[0];
+		if (OutFilesToImport.IsEmpty())
+		{
+			return ELandscapeImportResult::Error;
+		}
+		
+		if (OutFilesToImport.Contains(FilePath))
+		{
+			// If one of the files found has a name that exactly matches yet we have on or more numbered files matching the pattern, then we have an ambiguity and warn the user about it :
+			if (OutFilesToImport.Num() > 1)
+			{
+				OutMessage = FText::Format(LOCTEXT("Import_AmbiguousImportFile", 
+					"Ambiguous import files found :\nThere's a single '{0}.{3}' file and {1} {1}|plural(one=file,other=files) whose {1}|plural(one=name,other=names) {1}|plural(one=matches,other=match) the proper input pattern (ex: '{2}_x0_y0.{3}').\nPlease rename the single file."),
+					FText::FromString(FPaths::GetBaseFilename(FilePath)), OutFilesToImport.Num() - 1, FText::FromString(FPaths::GetBaseFilename(FilePathPattern)), FText::FromString(FPaths::GetExtension(FilePath)));
+				return ELandscapeImportResult::Error;
+			}
+			else
+			{
+				// We have a single file found and it's exactly matching the expected file name, consider we're in single file mode :
+				check(OutFilesToImport.Num() == 1);
+				bSingleFile = true;
+			}
+		}
 	}
 	else
 	{
@@ -177,7 +214,7 @@ ELandscapeImportResult GetImportDescriptorInternal(const FString& FilePath, bool
 			FString OutLocalFileImportPattern;
 			if (!bSingleFile && !FLandscapeImportHelper::ExtractCoordinates(FPaths::GetBaseFilename(ImportFilename), OutCoord, OutLocalFileImportPattern))
 			{
-				OutMessage = FText::Format(LOCTEXT("Import_InvalidFilename", "File '{0}' doesn't have proper pattern(ex: {1}_x0_y0.{2})"),FText::FromString(FPaths::GetBaseFilename(ImportFilename)), FText::FromString(OutFileImportPattern), FText::FromString(FPaths::GetExtension(FilePath)));
+				OutMessage = FText::Format(LOCTEXT("Import_InvalidFilename", "File '{0}' doesn't have the proper pattern(ex: '{1}_x0_y0.{2})'"),FText::FromString(FPaths::GetBaseFilename(ImportFilename)), FText::FromString(OutFileImportPattern), FText::FromString(FPaths::GetExtension(FilePath)));
 				return ELandscapeImportResult::Error;
 			}
 			MinCoord.X = FMath::Min(OutCoord.X, MinCoord.X);
diff --git a/Engine/Source/Editor/LandscapeEditor/Public/LandscapeEditorObject.h b/Engine/Source/Editor/LandscapeEditor/Public/LandscapeEditorObject.h
--- a/Engine/Source/Editor/LandscapeEditor/Public/LandscapeEditorObject.h
+++ b/Engine/Source/Editor/LandscapeEditor/Public/LandscapeEditorObject.h
@@ -556,9 +556,12 @@ class ULandscapeEditorObject : public UObject
 	UPROPERTY(NonTransactional)
 	bool bHeightmapSelected = false;
 	
-	UPROPERTY(Category = "Import / Export", EditAnywhere, NonTransactional, meta = (DisplayName="Export Edit Layer", ShowForTools = "ImportExport", ToolTip="When true exports the selected edit layer, if false exports the blend result"))
+	UPROPERTY(Category = "Import / Export", EditAnywhere, NonTransactional, meta = (DisplayName="Export Selected Edit Layer", ShowForTools = "ImportExport", ToolTip="When true exports the selected edit layer, if false exports the blended result"))
 	bool bExportEditLayer = true;
 
+	UPROPERTY(Category = "Import / Export", EditAnywhere, NonTransactional, meta = (DisplayName = "Export Single File", ShowForTools = "ImportExport", ToolTip = "(World Partition only) When true, exports the landscape as a single file, if false exports each grid tile individually."))
+	bool bExportSingleFile = false;
+
 	UPROPERTY(NonTransactional)
 	FLandscapeImportDescriptor HeightmapImportDescriptor;
 	
diff --git a/Engine/Source/Editor/LandscapeEditor/Public/LandscapeImportHelper.h b/Engine/Source/Editor/LandscapeEditor/Public/LandscapeImportHelper.h
--- a/Engine/Source/Editor/LandscapeEditor/Public/LandscapeImportHelper.h
+++ b/Engine/Source/Editor/LandscapeEditor/Public/LandscapeImportHelper.h
@@ -136,7 +136,7 @@ public:
 
 	static void ChooseBestComponentSizeForImport(int32 Width, int32 Height, int32& InOutQuadsPerSection, int32& InOutSectionsPerComponent, FIntPoint& OutComponentCount);
 
-	static bool ExtractCoordinates(FString BaseFilename, FIntPoint& OutCoord, FString& OutBaseFilePattern);
+	static bool ExtractCoordinates(const FString& BaseFilename, FIntPoint& OutCoord, FString& OutBaseFilePattern);
 	static void GetMatchingFiles(const FString& FilePathPattern, TArray<FString>& OutFilePaths);
 
 
