Landscape: Correctly implemented GetReferencerName for all landscape tools.

diff --git a/Engine/Source/Editor/LandscapeEditor/Private/LandscapeEdModeBPCustomTools.cpp b/Engine/Source/Editor/LandscapeEditor/Private/LandscapeEdModeBPCustomTools.cpp
--- a/Engine/Source/Editor/LandscapeEditor/Private/LandscapeEdModeBPCustomTools.cpp
+++ b/Engine/Source/Editor/LandscapeEditor/Private/LandscapeEdModeBPCustomTools.cpp
@@ -47,9 +47,9 @@ public:
 	{
 	}
 
-	virtual const TCHAR* GetToolName() override { return TEXT("BlueprintBrush"); }
-	virtual FText GetDisplayName() override { return FText(); };
-	virtual FText GetDisplayMessage() override { return FText(); };
+	virtual const TCHAR* GetToolName() const override { return TEXT("BlueprintBrush"); }
+	virtual FText GetDisplayName() const override { return FText(); };
+	virtual FText GetDisplayMessage() const override { return FText(); };
 
 	virtual void SetEditRenderType() override { GLandscapeEditRenderMode = ELandscapeEditRenderMode::None | (GLandscapeEditRenderMode & ELandscapeEditRenderMode::BitMaskForMask); }
 	virtual bool SupportsMask() override { return false; }
diff --git a/Engine/Source/Editor/LandscapeEditor/Private/LandscapeEdModeComponentTools.cpp b/Engine/Source/Editor/LandscapeEditor/Private/LandscapeEdModeComponentTools.cpp
--- a/Engine/Source/Editor/LandscapeEditor/Private/LandscapeEdModeComponentTools.cpp
+++ b/Engine/Source/Editor/LandscapeEditor/Private/LandscapeEdModeComponentTools.cpp
@@ -148,9 +148,9 @@ public:
 		return ELandscapeLayerUpdateMode::Update_None;
 	}
 
-	virtual const TCHAR* GetToolName() override { return TEXT("Select"); }
-	virtual FText GetDisplayName() override { return NSLOCTEXT("UnrealEd", "LandscapeMode_Selection", "Component Selection"); };
-	virtual FText GetDisplayMessage() override { return NSLOCTEXT("UnrealEd", "LandscapeMode_Selection_Message", "Paint a mask on the Landscape to protect areas from editing."); };
+	virtual const TCHAR* GetToolName() const override { return TEXT("Select"); }
+	virtual FText GetDisplayName() const override { return NSLOCTEXT("UnrealEd", "LandscapeMode_Selection", "Component Selection"); };
+	virtual FText GetDisplayMessage() const override { return NSLOCTEXT("UnrealEd", "LandscapeMode_Selection_Message", "Paint a mask on the Landscape to protect areas from editing."); };
 
 	virtual void SetEditRenderType() override { GLandscapeEditRenderMode = ELandscapeEditRenderMode::SelectComponent | (GLandscapeEditRenderMode & ELandscapeEditRenderMode::BitMaskForMask); }
 	virtual bool SupportsMask() override { return false; }
@@ -253,9 +253,9 @@ public:
 
 	virtual bool AffectsEditLayers() const { return false; }
 
-	virtual const TCHAR* GetToolName() override { return TEXT("Mask"); }
-	virtual FText GetDisplayName() override { return NSLOCTEXT("UnrealEd", "LandscapeMode_Mask", "Region Selection"); };
-	virtual FText GetDisplayMessage() override { return NSLOCTEXT("UnrealEd", "LandscapeMode_Mask_Message", "Region Selection"); };
+	virtual const TCHAR* GetToolName() const override { return TEXT("Mask"); }
+	virtual FText GetDisplayName() const override { return NSLOCTEXT("UnrealEd", "LandscapeMode_Mask", "Region Selection"); };
+	virtual FText GetDisplayMessage() const override { return NSLOCTEXT("UnrealEd", "LandscapeMode_Mask_Message", "Region Selection"); };
 
 	virtual void SetEditRenderType() override { GLandscapeEditRenderMode = ELandscapeEditRenderMode::SelectRegion | (GLandscapeEditRenderMode & ELandscapeEditRenderMode::BitMaskForMask); }
 	virtual bool SupportsMask() override { return true; }
@@ -339,9 +339,9 @@ public:
 		return FLandscapeToolBase<FLandscapeToolStrokeVisibility>::BeginTool(ViewportClient, InTarget, InHitLocation);
 	}
 
-	virtual const TCHAR* GetToolName() override { return TEXT("Visibility"); }
-	virtual FText GetDisplayName() override { return NSLOCTEXT("UnrealEd", "LandscapeMode_Visibility", "Visibility"); };
-	virtual FText GetDisplayMessage() override { return NSLOCTEXT("UnrealEd", "LandscapeMode_Visibility_Message", "This tool will allow you to mask out the visibility and collision of areas of your Landscape when used in conjunction with the Landscape Hole Material."); };
+	virtual const TCHAR* GetToolName() const override { return TEXT("Visibility"); }
+	virtual FText GetDisplayName() const override { return NSLOCTEXT("UnrealEd", "LandscapeMode_Visibility", "Visibility"); };
+	virtual FText GetDisplayMessage() const override { return NSLOCTEXT("UnrealEd", "LandscapeMode_Visibility_Message", "This tool will allow you to mask out the visibility and collision of areas of your Landscape when used in conjunction with the Landscape Hole Material."); };
 
 	virtual void SetEditRenderType() override { GLandscapeEditRenderMode = ELandscapeEditRenderMode::None | (GLandscapeEditRenderMode & ELandscapeEditRenderMode::BitMaskForMask); }
 	virtual bool SupportsMask() override { return false; }
@@ -501,9 +501,9 @@ public:
 	}
 	virtual bool AffectsEditLayers() const override { return false; }
 
-	virtual const TCHAR* GetToolName() override { return TEXT("MoveToLevel"); }
-	virtual FText GetDisplayName() override { return NSLOCTEXT("UnrealEd", "LandscapeMode_MoveToLevel", "Move to Streaming Level"); };
-	virtual FText GetDisplayMessage() override { return NSLOCTEXT("UnrealEd", "LandscapeMode_MoveToLevel_Message", "Move the selected components, via using the Selection tool, to the current streaming level.  This makes it possible to move sections of a Landscape into a streaming level so that they will be streamed in and out with that level, optimizing the performance of the Landscape."); };
+	virtual const TCHAR* GetToolName() const override { return TEXT("MoveToLevel"); }
+	virtual FText GetDisplayName() const override { return NSLOCTEXT("UnrealEd", "LandscapeMode_MoveToLevel", "Move to Streaming Level"); };
+	virtual FText GetDisplayMessage() const override { return NSLOCTEXT("UnrealEd", "LandscapeMode_MoveToLevel_Message", "Move the selected components, via using the Selection tool, to the current streaming level.  This makes it possible to move sections of a Landscape into a streaming level so that they will be streamed in and out with that level, optimizing the performance of the Landscape."); };
 
 	virtual void SetEditRenderType() override { GLandscapeEditRenderMode = ELandscapeEditRenderMode::SelectComponent | (GLandscapeEditRenderMode & ELandscapeEditRenderMode::BitMaskForMask); }
 	virtual bool SupportsMask() override { return false; }
@@ -792,9 +792,9 @@ public:
 	}
 	virtual bool AffectsEditLayers() const override { return false; }
 
-	virtual const TCHAR* GetToolName() override { return TEXT("AddComponent"); }
-	virtual FText GetDisplayName() override { return NSLOCTEXT("UnrealEd", "LandscapeMode_AddComponent", "Add New Landscape Component"); };
-	virtual FText GetDisplayMessage() override { return NSLOCTEXT("UnrealEd", "LandscapeMode_AddComponent_Message", "Create new components for the current Landscape, one at a time.  The cursor shows a green wireframe where new components can be added."); };
+	virtual const TCHAR* GetToolName() const override { return TEXT("AddComponent"); }
+	virtual FText GetDisplayName() const override { return NSLOCTEXT("UnrealEd", "LandscapeMode_AddComponent", "Add New Landscape Component"); };
+	virtual FText GetDisplayMessage() const override { return NSLOCTEXT("UnrealEd", "LandscapeMode_AddComponent_Message", "Create new components for the current Landscape, one at a time.  The cursor shows a green wireframe where new components can be added."); };
 
 	virtual void SetEditRenderType() override { GLandscapeEditRenderMode = ELandscapeEditRenderMode::None | (GLandscapeEditRenderMode & ELandscapeEditRenderMode::BitMaskForMask); }
 	virtual bool SupportsMask() override { return false; }
@@ -999,9 +999,9 @@ public:
 
 	virtual bool AffectsEditLayers() const override { return false; }
 
-	virtual const TCHAR* GetToolName() override { return TEXT("DeleteComponent"); }
-	virtual FText GetDisplayName() override { return NSLOCTEXT("UnrealEd", "LandscapeMode_DeleteComponent", "Delete Landscape Components"); };
-	virtual FText GetDisplayMessage() override { return NSLOCTEXT("UnrealEd", "LandscapeMode_DeleteComponent_Message", "Delete selected components . If no components are currently selected, deletes the component highlighted under the mouse cursor. "); };
+	virtual const TCHAR* GetToolName() const override { return TEXT("DeleteComponent"); }
+	virtual FText GetDisplayName() const override { return NSLOCTEXT("UnrealEd", "LandscapeMode_DeleteComponent", "Delete Landscape Components"); };
+	virtual FText GetDisplayMessage() const override { return NSLOCTEXT("UnrealEd", "LandscapeMode_DeleteComponent_Message", "Delete selected components . If no components are currently selected, deletes the component highlighted under the mouse cursor. "); };
 
 	virtual void SetEditRenderType() override { GLandscapeEditRenderMode = ELandscapeEditRenderMode::SelectComponent | (GLandscapeEditRenderMode & ELandscapeEditRenderMode::BitMaskForMask); }
 	virtual bool SupportsMask() override { return false; }
@@ -1325,9 +1325,9 @@ public:
 	{
 	}
 
-	virtual const TCHAR* GetToolName() override { return TEXT("Copy"); }
-	virtual FText GetDisplayName() override { return NSLOCTEXT("UnrealEd", "LandscapeMode_Copy", "Copy"); };
-	virtual FText GetDisplayMessage() override { return NSLOCTEXT("UnrealEd", "LandscapeMode_Copy_Message", "Copy and Paste allows you to copy terrain data from one area of your Landscape to another.  Use the select tool  in conjunction with the Copy gizmo to further refine your selection."); };
+	virtual const TCHAR* GetToolName() const override { return TEXT("Copy"); }
+	virtual FText GetDisplayName() const override { return NSLOCTEXT("UnrealEd", "LandscapeMode_Copy", "Copy"); };
+	virtual FText GetDisplayMessage() const override { return NSLOCTEXT("UnrealEd", "LandscapeMode_Copy_Message", "Copy and Paste allows you to copy terrain data from one area of your Landscape to another.  Use the select tool  in conjunction with the Copy gizmo to further refine your selection."); };
 
 
 	virtual void SetEditRenderType() override
@@ -1644,9 +1644,9 @@ public:
 	{
 	}
 
-	virtual const TCHAR* GetToolName() override { return TEXT("Paste"); }
-	virtual FText GetDisplayName() override { return NSLOCTEXT("UnrealEd", "LandscapeMode_Region", "Region Copy/Paste"); };
-	virtual FText GetDisplayMessage() override { return NSLOCTEXT("UnrealEd", "LandscapeMode_Region_Message", "Copy and Paste allows you to copy terrain data from one area of your Landscape to another.  Use the select tool  in conjunction with the Copy gizmo to further refine your selection."); };
+	virtual const TCHAR* GetToolName() const override { return TEXT("Paste"); }
+	virtual FText GetDisplayName() const override { return NSLOCTEXT("UnrealEd", "LandscapeMode_Region", "Region Copy/Paste"); };
+	virtual FText GetDisplayMessage() const override { return NSLOCTEXT("UnrealEd", "LandscapeMode_Region_Message", "Copy and Paste allows you to copy terrain data from one area of your Landscape to another.  Use the select tool  in conjunction with the Copy gizmo to further refine your selection."); };
 
 	virtual void SetEditRenderType() override
 	{
@@ -1725,9 +1725,9 @@ public:
 	}
 
 	// Just hybrid of Copy and Paste tool
-	virtual const TCHAR* GetToolName() override { return TEXT("CopyPaste"); }
-	virtual FText GetDisplayName() override { return NSLOCTEXT("UnrealEd", "LandscapeMode_Region", "Region Copy/Paste"); };
-	virtual FText GetDisplayMessage() override { return NSLOCTEXT("UnrealEd", "LandscapeMode_Region_Message", "Copy and Paste allows you to copy terrain data from one area of your Landscape to another.  Use the select tool  in conjunction with the Copy gizmo to further refine your selection."); };
+	virtual const TCHAR* GetToolName() const override { return TEXT("CopyPaste"); }
+	virtual FText GetDisplayName() const override { return NSLOCTEXT("UnrealEd", "LandscapeMode_Region", "Region Copy/Paste"); };
+	virtual FText GetDisplayMessage() const override { return NSLOCTEXT("UnrealEd", "LandscapeMode_Region_Message", "Copy and Paste allows you to copy terrain data from one area of your Landscape to another.  Use the select tool  in conjunction with the Copy gizmo to further refine your selection."); };
 
 	virtual void EnterTool() override
 	{
@@ -1863,9 +1863,9 @@ public:
 	{
 	}
 	virtual bool AffectsEditLayers() const { return false; }
-	virtual const TCHAR* GetToolName() override { return TEXT("NewLandscape"); }
-	virtual FText GetDisplayName() override { return NSLOCTEXT("UnrealEd", "LandscapeMode_NewLandscape", "New Landscape"); };
-	virtual FText GetDisplayMessage() override { return NSLOCTEXT("UnrealEd", "LandscapeMode_NewLandscape_Message", "Create or import a new heightmap.  Assign a material and configure the components.  When you are ready to create your new Landscape, press the Create button in the lower-right corner of this panel. "); };
+	virtual const TCHAR* GetToolName() const override { return TEXT("NewLandscape"); }
+	virtual FText GetDisplayName() const override { return NSLOCTEXT("UnrealEd", "LandscapeMode_NewLandscape", "New Landscape"); };
+	virtual FText GetDisplayMessage() const override { return NSLOCTEXT("UnrealEd", "LandscapeMode_NewLandscape_Message", "Create or import a new heightmap.  Assign a material and configure the components.  When you are ready to create your new Landscape, press the Create button in the lower-right corner of this panel. "); };
 
 	virtual void SetEditRenderType() override { GLandscapeEditRenderMode = ELandscapeEditRenderMode::None | (GLandscapeEditRenderMode & ELandscapeEditRenderMode::BitMaskForMask); }
 	virtual bool SupportsMask() override { return false; }
@@ -2321,9 +2321,9 @@ public:
 	}
 
 	virtual bool AffectsEditLayers() const { return false; }
-	virtual const TCHAR* GetToolName() override { return TEXT("ResizeLandscape"); }
-	virtual FText GetDisplayName() override { return LOCTEXT("LandscapeMode_ResizeLandscape", "Change Landscape Component Size"); };
-	virtual FText GetDisplayMessage() override { return LOCTEXT("LandscapeMode_ResizeLandscape_Message", "Change Landscape Component Size"); };
+	virtual const TCHAR* GetToolName() const override { return TEXT("ResizeLandscape"); }
+	virtual FText GetDisplayName() const override { return LOCTEXT("LandscapeMode_ResizeLandscape", "Change Landscape Component Size"); };
+	virtual FText GetDisplayMessage() const override { return LOCTEXT("LandscapeMode_ResizeLandscape_Message", "Change Landscape Component Size"); };
 
 	virtual void SetEditRenderType() override { GLandscapeEditRenderMode = ELandscapeEditRenderMode::None | (GLandscapeEditRenderMode & ELandscapeEditRenderMode::BitMaskForMask); }
 	virtual bool SupportsMask() override { return false; }
@@ -2387,9 +2387,9 @@ public:
 	{
 	}
 
-	virtual const TCHAR* GetToolName() override { return TEXT("ImportExport"); }
-	virtual FText GetDisplayName() override { return LOCTEXT("LandscapeMode_ImportExport", "Import Export"); };
-	virtual FText GetDisplayMessage() override { return LOCTEXT("LandscapeMode_ImportExport_Message", "Import/Export Landscape Data"); };
+	virtual const TCHAR* GetToolName() const override { return TEXT("ImportExport"); }
+	virtual FText GetDisplayName() const override { return LOCTEXT("LandscapeMode_ImportExport", "Import Export"); };
+	virtual FText GetDisplayMessage() const override { return LOCTEXT("LandscapeMode_ImportExport_Message", "Import/Export Landscape Data"); };
 
 	virtual bool UsesTransformWidget() const { return true; }
 	virtual void SetEditRenderType() override { GLandscapeEditRenderMode = ELandscapeEditRenderMode::SelectComponent | (GLandscapeEditRenderMode & ELandscapeEditRenderMode::BitMaskForMask); }
diff --git a/Engine/Source/Editor/LandscapeEditor/Private/LandscapeEdModeErosionTools.cpp b/Engine/Source/Editor/LandscapeEditor/Private/LandscapeEdModeErosionTools.cpp
--- a/Engine/Source/Editor/LandscapeEditor/Private/LandscapeEdModeErosionTools.cpp
+++ b/Engine/Source/Editor/LandscapeEditor/Private/LandscapeEdModeErosionTools.cpp
@@ -278,9 +278,9 @@ public:
 	{
 	}
 
-	virtual const TCHAR* GetToolName() override { return TEXT("Erosion"); }
-	virtual FText GetDisplayName() override { return NSLOCTEXT("UnrealEd", "LandscapeMode_Erosion", "Erosion"); };
-	virtual FText GetDisplayMessage() override { return NSLOCTEXT("UnrealEd", "LandscapeMode_Erosion_Message", "The Erosion tool uses a thermal erosion simulation to adjust the height of the landscape. This simulates the transfer of soil from higher elevations to lower elevations. The larger the difference in elevation, the more erosion will occur. This tool also applies a noise effect on top of the erosion, if desired, to provide a more natural random appearance. "); };
+	virtual const TCHAR* GetToolName() const override { return TEXT("Erosion"); }
+	virtual FText GetDisplayName() const override { return NSLOCTEXT("UnrealEd", "LandscapeMode_Erosion", "Erosion"); };
+	virtual FText GetDisplayMessage() const override { return NSLOCTEXT("UnrealEd", "LandscapeMode_Erosion_Message", "The Erosion tool uses a thermal erosion simulation to adjust the height of the landscape. This simulates the transfer of soil from higher elevations to lower elevations. The larger the difference in elevation, the more erosion will occur. This tool also applies a noise effect on top of the erosion, if desired, to provide a more natural random appearance. "); };
 
 };
 
@@ -499,9 +499,9 @@ public:
 	{
 	}
 
-	virtual const TCHAR* GetToolName() override { return TEXT("HydraErosion"); } // formerly HydraulicErosion
-	virtual FText GetDisplayName() override { return NSLOCTEXT("UnrealEd", "LandscapeMode_HydraErosion", "Hydraulic Erosion"); };
-	virtual FText GetDisplayMessage() override { return NSLOCTEXT("UnrealEd", "LandscapeMode_HydraErosion_Message", "The Hydro Erosion tool simulates erosion from water to adjust the height of the landscape. A noise filter is used to determine where the initial rain is distributed. Then the simulation is calculated to determine water flow from that initial rain as well as dissolving, water transfer, and evaporation. The result of that calculation provides the actual value used to lower the heightmap."); };
+	virtual const TCHAR* GetToolName() const override { return TEXT("HydraErosion"); } // formerly HydraulicErosion
+	virtual FText GetDisplayName() const override { return NSLOCTEXT("UnrealEd", "LandscapeMode_HydraErosion", "Hydraulic Erosion"); };
+	virtual FText GetDisplayMessage() const override { return NSLOCTEXT("UnrealEd", "LandscapeMode_HydraErosion_Message", "The Hydro Erosion tool simulates erosion from water to adjust the height of the landscape. A noise filter is used to determine where the initial rain is distributed. Then the simulation is calculated to determine water flow from that initial rain as well as dissolving, water transfer, and evaporation. The result of that calculation provides the actual value used to lower the heightmap."); };
 
 };
 
diff --git a/Engine/Source/Editor/LandscapeEditor/Private/LandscapeEdModeMirrorTool.cpp b/Engine/Source/Editor/LandscapeEditor/Private/LandscapeEdModeMirrorTool.cpp
--- a/Engine/Source/Editor/LandscapeEditor/Private/LandscapeEdModeMirrorTool.cpp
+++ b/Engine/Source/Editor/LandscapeEditor/Private/LandscapeEdModeMirrorTool.cpp
@@ -45,9 +45,9 @@ public:
 		Collector.AddReferencedObject(MirrorPlaneMaterial);
 	}
 
-	virtual const TCHAR* GetToolName() override { return TEXT("Mirror"); }
-	virtual FText GetDisplayName() override { return NSLOCTEXT("UnrealEd", "LandscapeTool_Mirror", "Mirror Landscape"); };
-	virtual FText GetDisplayMessage() override { return NSLOCTEXT("UnrealEd", "LandscapeTool_Mirror_Message", "Copy one side of a landscape to the other side so that you can easily mirror or rotate the landscape geometry along the X or Y axis."); };
+	virtual const TCHAR* GetToolName() const override { return TEXT("Mirror"); }
+	virtual FText GetDisplayName() const override { return NSLOCTEXT("UnrealEd", "LandscapeTool_Mirror", "Mirror Landscape"); };
+	virtual FText GetDisplayMessage() const override { return NSLOCTEXT("UnrealEd", "LandscapeTool_Mirror_Message", "Copy one side of a landscape to the other side so that you can easily mirror or rotate the landscape geometry along the X or Y axis."); };
 
 	virtual void SetEditRenderType() override { GLandscapeEditRenderMode = ELandscapeEditRenderMode::None | (GLandscapeEditRenderMode & ELandscapeEditRenderMode::BitMaskForMask); }
 	virtual bool SupportsMask() override { return false; }
diff --git a/Engine/Source/Editor/LandscapeEditor/Private/LandscapeEdModePaintTools.cpp b/Engine/Source/Editor/LandscapeEditor/Private/LandscapeEdModePaintTools.cpp
--- a/Engine/Source/Editor/LandscapeEditor/Private/LandscapeEdModePaintTools.cpp
+++ b/Engine/Source/Editor/LandscapeEditor/Private/LandscapeEdModePaintTools.cpp
@@ -317,9 +317,9 @@ public:
 	{
 	}
 
-	virtual const TCHAR* GetToolName() override { return TEXT("Paint"); }
-	virtual FText GetDisplayName() override { return NSLOCTEXT("UnrealEd", "LandscapeMode_Paint", "Paint"); };
-	virtual FText GetDisplayMessage() override { return NSLOCTEXT("UnrealEd", "LandscapeMode_Paint_Message", "The Paint tool increases or decreases the weight of the Material layer being applied to the Landscape."); };
+	virtual const TCHAR* GetToolName() const override { return TEXT("Paint"); }
+	virtual FText GetDisplayName() const override { return NSLOCTEXT("UnrealEd", "LandscapeMode_Paint", "Paint"); };
+	virtual FText GetDisplayMessage() const override { return NSLOCTEXT("UnrealEd", "LandscapeMode_Paint_Message", "The Paint tool increases or decreases the weight of the Material layer being applied to the Landscape."); };
 
 	virtual void EnterTool()
 	{
@@ -628,9 +628,9 @@ public:
 	{
 	}
 
-	virtual const TCHAR* GetToolName() override { return TEXT("Sculpt"); }
-	virtual FText GetDisplayName() override { return NSLOCTEXT("UnrealEd", "LandscapeMode_Sculpt", "Sculpt"); };
-	virtual FText GetDisplayMessage() override { return NSLOCTEXT("UnrealEd", "LandscapeMode_Sculpt_Message", "Raise or lower the Landscape using the selected brush shape and falloff."); };
+	virtual const TCHAR* GetToolName() const override { return TEXT("Sculpt"); }
+	virtual FText GetDisplayName() const override { return NSLOCTEXT("UnrealEd", "LandscapeMode_Sculpt", "Sculpt"); };
+	virtual FText GetDisplayMessage() const override { return NSLOCTEXT("UnrealEd", "LandscapeMode_Sculpt_Message", "Raise or lower the Landscape using the selected brush shape and falloff."); };
 };
 
 class FLandscapeToolErase : public FLandscapeToolPaintBase<FHeightmapToolTarget, FLandscapeToolStrokeErase>
@@ -641,9 +641,9 @@ public:
 	{
 	}
 
-	virtual const TCHAR* GetToolName() override { return TEXT("Erase"); }
-	virtual FText GetDisplayName() override { return NSLOCTEXT("UnrealEd", "LandscapeMode_Erase", "Erase"); };
-	virtual FText GetDisplayMessage() override { return NSLOCTEXT("UnrealEd", "LandscapeMode_Erase_Message", "Erase the Landscape using the selected brush shape and falloff."); };
+	virtual const TCHAR* GetToolName() const override { return TEXT("Erase"); }
+	virtual FText GetDisplayName() const override { return NSLOCTEXT("UnrealEd", "LandscapeMode_Erase", "Erase"); };
+	virtual FText GetDisplayMessage() const override { return NSLOCTEXT("UnrealEd", "LandscapeMode_Erase_Message", "Erase the Landscape using the selected brush shape and falloff."); };
 };
 
 // 
@@ -779,9 +779,9 @@ public:
 	{
 	}
 
-	virtual const TCHAR* GetToolName() override { return TEXT("Smooth"); }
-	virtual FText GetDisplayName() override { return NSLOCTEXT("UnrealEd", "LandscapeMode_Smooth", "Smooth"); };
-	virtual FText GetDisplayMessage() override { return NSLOCTEXT("UnrealEd", "LandscapeMode_Smooth_Message", "Smooth the Landscape within the brushes influence by averaging the Z position of the Landscape vertices."); };
+	virtual const TCHAR* GetToolName() const override { return TEXT("Smooth"); }
+	virtual FText GetDisplayName() const override { return NSLOCTEXT("UnrealEd", "LandscapeMode_Smooth", "Smooth"); };
+	virtual FText GetDisplayMessage() const override { return NSLOCTEXT("UnrealEd", "LandscapeMode_Smooth_Message", "Smooth the Landscape within the brushes influence by averaging the Z position of the Landscape vertices."); };
 
 };
 
@@ -1103,9 +1103,9 @@ public:
 		Collector.AddReferencedObject(HeightmapFlattenPreviewComponent);
 	}
 
-	virtual const TCHAR* GetToolName() override { return TEXT("Flatten"); }
-	virtual FText GetDisplayName() override { return NSLOCTEXT("UnrealEd", "LandscapeMode_Flatten", "Flatten"); };
-	virtual FText GetDisplayMessage() override { return NSLOCTEXT("UnrealEd", "LandscapeMode_Flatten_Message", "Raise and lower the Landscape to be the same Z height as the location from which you started using the tool."); };
+	virtual const TCHAR* GetToolName() const override { return TEXT("Flatten"); }
+	virtual FText GetDisplayName() const override { return NSLOCTEXT("UnrealEd", "LandscapeMode_Flatten", "Flatten"); };
+	virtual FText GetDisplayMessage() const override { return NSLOCTEXT("UnrealEd", "LandscapeMode_Flatten_Message", "Raise and lower the Landscape to be the same Z height as the location from which you started using the tool."); };
 
 	virtual void Tick(FEditorViewportClient* ViewportClient, float DeltaTime) override
 	{
@@ -1311,9 +1311,9 @@ public:
 	{
 	}
 
-	virtual const TCHAR* GetToolName() override { return TEXT("Noise"); }
-	virtual FText GetDisplayName() override { return NSLOCTEXT("UnrealEd", "LandscapeMode_Noise", "Noise"); };
-	virtual FText GetDisplayMessage() override { return NSLOCTEXT("UnrealEd", "LandscapeMode_Noise_Message", "The Noise tool applies a noise filter to the heightmap or layer weight. The strength determines the amount of noise."); };
+	virtual const TCHAR* GetToolName() const override { return TEXT("Noise"); }
+	virtual FText GetDisplayName() const override { return NSLOCTEXT("UnrealEd", "LandscapeMode_Noise", "Noise"); };
+	virtual FText GetDisplayMessage() const override { return NSLOCTEXT("UnrealEd", "LandscapeMode_Noise_Message", "The Noise tool applies a noise filter to the heightmap or layer weight. The strength determines the amount of noise."); };
 };
 
 
diff --git a/Engine/Source/Editor/LandscapeEditor/Private/LandscapeEdModeRampTool.cpp b/Engine/Source/Editor/LandscapeEditor/Private/LandscapeEdModeRampTool.cpp
--- a/Engine/Source/Editor/LandscapeEditor/Private/LandscapeEdModeRampTool.cpp
+++ b/Engine/Source/Editor/LandscapeEditor/Private/LandscapeEdModeRampTool.cpp
@@ -118,9 +118,9 @@ public:
 		Collector.AddReferencedObject(SpriteTexture);
 	}
 
-	virtual const TCHAR* GetToolName() override { return TEXT("Ramp"); }
-	virtual FText GetDisplayName() override { return NSLOCTEXT("UnrealEd", "LandscapeMode_Ramp", "Ramp"); };
-	virtual FText GetDisplayMessage() override { return NSLOCTEXT("UnrealEd", "LandscapeMode_Ramp_Message", "Create a ramp between two specified points, adding falloffs according to the settings. Note that this tool cannot use any brushes."); };
+	virtual const TCHAR* GetToolName() const override { return TEXT("Ramp"); }
+	virtual FText GetDisplayName() const override { return NSLOCTEXT("UnrealEd", "LandscapeMode_Ramp", "Ramp"); };
+	virtual FText GetDisplayMessage() const override { return NSLOCTEXT("UnrealEd", "LandscapeMode_Ramp_Message", "Create a ramp between two specified points, adding falloffs according to the settings. Note that this tool cannot use any brushes."); };
 
 	virtual void SetEditRenderType() override { GLandscapeEditRenderMode = ELandscapeEditRenderMode::None | (GLandscapeEditRenderMode & ELandscapeEditRenderMode::BitMaskForMask); }
 	virtual bool SupportsMask() override { return false; }
diff --git a/Engine/Source/Editor/LandscapeEditor/Private/LandscapeEdModeSplineTools.cpp b/Engine/Source/Editor/LandscapeEditor/Private/LandscapeEdModeSplineTools.cpp
--- a/Engine/Source/Editor/LandscapeEditor/Private/LandscapeEdModeSplineTools.cpp
+++ b/Engine/Source/Editor/LandscapeEditor/Private/LandscapeEdModeSplineTools.cpp
@@ -93,9 +93,9 @@ public:
 		}
 	}
 
-	virtual const TCHAR* GetToolName() override { return TEXT("Splines"); }
-	virtual FText GetDisplayName() override { return NSLOCTEXT("UnrealEd", "LandscapeMode_Splines", "Splines"); };
-	virtual FText GetDisplayMessage() override { return NSLOCTEXT("UnrealEd", "LandscapeMode_Splines_Message", "Create a Landscape Spline to carve your landscape, modify blendmasks and deform meshes into roads and other linear features.  Spline mesh settings can be found in the details panel when you have  segments selected."); };
+	virtual const TCHAR* GetToolName() const override { return TEXT("Splines"); }
+	virtual FText GetDisplayName() const override { return NSLOCTEXT("UnrealEd", "LandscapeMode_Splines", "Splines"); };
+	virtual FText GetDisplayMessage() const override { return NSLOCTEXT("UnrealEd", "LandscapeMode_Splines_Message", "Create a Landscape Spline to carve your landscape, modify blendmasks and deform meshes into roads and other linear features.  Spline mesh settings can be found in the details panel when you have  segments selected."); };
 
 	virtual void SetEditRenderType() override { GLandscapeEditRenderMode = ELandscapeEditRenderMode::None | (GLandscapeEditRenderMode & ELandscapeEditRenderMode::BitMaskForMask); }
 	virtual bool SupportsMask() override { return false; }
diff --git a/Engine/Source/Editor/LandscapeEditor/Private/LandscapeEdModeXYOffset.cpp b/Engine/Source/Editor/LandscapeEditor/Private/LandscapeEdModeXYOffset.cpp
--- a/Engine/Source/Editor/LandscapeEditor/Private/LandscapeEdModeXYOffset.cpp
+++ b/Engine/Source/Editor/LandscapeEditor/Private/LandscapeEdModeXYOffset.cpp
@@ -557,9 +557,9 @@ public:
 	{
 	}
 
-	virtual const TCHAR* GetToolName() override { return TEXT("Retopologize"); }
-	virtual FText GetDisplayName() override { return NSLOCTEXT("UnrealEd", "LandscapeMode_Retopologize", "Retopologize"); }
-	virtual FText GetDisplayMessage() override { return NSLOCTEXT("UnrealEd", "LandscapeMode_Retopologize_Message", "Similar to the smooth tool, Retopologize will push and pull triangles to smooth the transition of the terrain, however it will attempt to keep the basic form of the terrain, minimizing change in the Z direction. An X/Y offset map makes the Landscape slower to render and paint on with other tools, so only use the Retopologize tool if needed."); }
+	virtual const TCHAR* GetToolName() const override { return TEXT("Retopologize"); }
+	virtual FText GetDisplayName() const override { return NSLOCTEXT("UnrealEd", "LandscapeMode_Retopologize", "Retopologize"); }
+	virtual FText GetDisplayMessage() const override { return NSLOCTEXT("UnrealEd", "LandscapeMode_Retopologize_Message", "Similar to the smooth tool, Retopologize will push and pull triangles to smooth the transition of the terrain, however it will attempt to keep the basic form of the terrain, minimizing change in the Z direction. An X/Y offset map makes the Landscape slower to render and paint on with other tools, so only use the Retopologize tool if needed."); }
 
 	virtual ELandscapeToolTargetTypeMask::Type GetSupportedTargetTypes() override
 	{
diff --git a/Engine/Source/Editor/LandscapeEditor/Public/LandscapeToolInterface.h b/Engine/Source/Editor/LandscapeEditor/Public/LandscapeToolInterface.h
--- a/Engine/Source/Editor/LandscapeEditor/Public/LandscapeToolInterface.h
+++ b/Engine/Source/Editor/LandscapeEditor/Public/LandscapeToolInterface.h
@@ -206,9 +206,9 @@ public:
 
 	FLandscapeTool() {}
 	virtual ~FLandscapeTool() {}
-	virtual const TCHAR* GetToolName() = 0;
-	virtual FText GetDisplayName() = 0;
-	virtual FText GetDisplayMessage() = 0;
+	virtual const TCHAR* GetToolName() const = 0;
+	virtual FText GetDisplayName() const = 0;
+	virtual FText GetDisplayMessage() const = 0;
 	virtual void SetEditRenderType();
 	virtual void Render(const FSceneView* View, FViewport* Viewport, FPrimitiveDrawInterface* PDI) {}
 	virtual bool HitTrace(const FVector& TraceStart, const FVector& TraceEnd, FVector& OutHitLocation) { return false; }
@@ -261,7 +261,7 @@ public:
 	virtual void AddReferencedObjects(FReferenceCollector& Collector) override {}
 	virtual FString GetReferencerName() const override
 	{
-		return TEXT("FLandscapeTool");
+		return FString(TEXT("FLandscapeTool::")).Append(GetToolName());
 	}
 
 public:
