Landscape: fix for nullptr dereference of CurrentTool in static analysis

diff --git a/Engine/Source/Editor/LandscapeEditor/Private/LandscapeEdMode.cpp b/Engine/Source/Editor/LandscapeEditor/Private/LandscapeEdMode.cpp
--- a/Engine/Source/Editor/LandscapeEditor/Private/LandscapeEdMode.cpp
+++ b/Engine/Source/Editor/LandscapeEditor/Private/LandscapeEdMode.cpp
@@ -1873,9 +1873,16 @@ bool FEdModeLandscape::InputKey(FEditorViewportClient* ViewportClient, FViewport
 				ToolActiveViewport = nullptr;
 			}
 
-			int32 OldToolIndex = CurrentToolMode->ValidTools.Find(CurrentTool->GetToolName());
-			int32 NewToolIndex = FMath::Max(OldToolIndex - 1, 0);
-			SetCurrentTool(CurrentToolMode->ValidTools[NewToolIndex]);
+			if (CurrentTool)
+			{
+				int32 OldToolIndex = CurrentToolMode->ValidTools.Find(CurrentTool->GetToolName());
+				int32 NewToolIndex = FMath::Max(OldToolIndex - 1, 0);
+				SetCurrentTool(CurrentToolMode->ValidTools[NewToolIndex]);
+			}
+			else if (CurrentToolMode && CurrentToolMode->ValidTools.Num() > 0)
+			{
+				SetCurrentTool(CurrentToolMode->ValidTools[0]);
+			}
 
 			return true;
 		}
@@ -1890,9 +1897,16 @@ bool FEdModeLandscape::InputKey(FEditorViewportClient* ViewportClient, FViewport
 				ToolActiveViewport = nullptr;
 			}
 
-			int32 OldToolIndex = CurrentToolMode->ValidTools.Find(CurrentTool->GetToolName());
-			int32 NewToolIndex = FMath::Min(OldToolIndex + 1, CurrentToolMode->ValidTools.Num() - 1);
-			SetCurrentTool(CurrentToolMode->ValidTools[NewToolIndex]);
+			if (CurrentTool)
+			{
+				int32 OldToolIndex = CurrentToolMode->ValidTools.Find(CurrentTool->GetToolName());
+				int32 NewToolIndex = FMath::Min(OldToolIndex + 1, CurrentToolMode->ValidTools.Num() - 1);
+				SetCurrentTool(CurrentToolMode->ValidTools[NewToolIndex]);
+			}
+			else if (CurrentToolMode && CurrentToolMode->ValidTools.Num() > 0)
+			{
+				SetCurrentTool(CurrentToolMode->ValidTools[0]);
+			}
 
 			return true;
 		}
@@ -4139,7 +4153,7 @@ bool FEdModeLandscape::CanEditLayer(FText* Reason /*=nullptr*/, FLandscapeLayer*
 		}
 	}
 
-	if (CurrentToolTarget.TargetType == ELandscapeToolTargetType::Weightmap && CurrentToolTarget.LayerInfo == nullptr && CurrentTool->GetToolName() != FName("BlueprintBrush"))
+	if (CurrentToolTarget.TargetType == ELandscapeToolTargetType::Weightmap && CurrentToolTarget.LayerInfo == nullptr && (CurrentTool && CurrentTool->GetToolName() != FName("BlueprintBrush")))
 	{
 		if (Reason)
 		{
