Fixed viewport not updating when switching landscape visualization modes

diff --git a/Engine/Source/Editor/LandscapeEditor/Private/LandscapeEditorModule.cpp b/Engine/Source/Editor/LandscapeEditor/Private/LandscapeEditorModule.cpp
--- a/Engine/Source/Editor/LandscapeEditor/Private/LandscapeEditorModule.cpp
+++ b/Engine/Source/Editor/LandscapeEditor/Private/LandscapeEditorModule.cpp
@@ -264,7 +264,15 @@ public:
 
 	static void ChangeLandscapeViewMode(ELandscapeViewMode::Type ViewMode)
 	{
-		GLandscapeViewMode = ViewMode;
+		if (ViewMode != GLandscapeViewMode)
+		{
+			GLandscapeViewMode = ViewMode;
+
+			if (GEditor)
+			{
+				GEditor->RedrawAllViewports(/*bInvalidateHitProxies =*/false);
+			}
+		}
 	}
 
 	static bool IsLandscapeViewModeSelected(ELandscapeViewMode::Type ViewMode)
