Landscape: Order Landscape Details UI Categories

diff --git a/Engine/Source/Editor/LandscapeEditor/Private/LandscapeEditorDetails.cpp b/Engine/Source/Editor/LandscapeEditor/Private/LandscapeEditorDetails.cpp
--- a/Engine/Source/Editor/LandscapeEditor/Private/LandscapeEditorDetails.cpp
+++ b/Engine/Source/Editor/LandscapeEditor/Private/LandscapeEditorDetails.cpp
@@ -48,6 +48,36 @@ void FLandscapeEditorDetails::CustomizeDetails(IDetailLayoutBuilder& DetailBuild
 
 	IDetailCategoryBuilder& LandscapeEditorCategory = DetailBuilder.EditCategory("LandscapeEditor", NSLOCTEXT("Contexts", "LandscapeEditor", "Landscape Editor"), ECategoryPriority::TypeSpecific);
 
+	IDetailCategoryBuilder& BrushSettingsCategory = DetailBuilder.EditCategory("Brush Settings");
+
+	// Ensure the categories in the Landscape Editor Details panel is stable. Most importantly that the Brush
+	// and Tool Settings are adjacent to each other. 
+	auto CategorySorter = [](const TMap<FName, IDetailCategoryBuilder*>& Categories)
+	{
+		int32 Order = 0;
+		auto SafeSetOrder = [&Categories, &Order](const FName& CategoryName )
+		{
+			if (IDetailCategoryBuilder* const* Builder = Categories.Find(CategoryName))
+			{
+				(*Builder)->SetSortOrder(Order++);	
+			}
+		};
+		
+		SafeSetOrder(FName("LandscapeEditor"));
+		SafeSetOrder(FName("Import / Export"));
+		SafeSetOrder(FName("Change Component Size"));
+		SafeSetOrder(FName("New Landscape"));
+		
+		SafeSetOrder(FName("Tool Settings"));
+		SafeSetOrder(FName("Brush Settings"));
+		
+		SafeSetOrder(FName("Edit Layers"));
+		SafeSetOrder(FName("Edit Layer Blueprint Brushes"));
+		SafeSetOrder(FName("Target Layers"));
+	};
+	
+	DetailBuilder.SortCategories(CategorySorter);
+	
 	LandscapeEditorCategory.AddCustomRow(FText::GetEmpty())
 	.Visibility(TAttribute<EVisibility>::Create(TAttribute<EVisibility>::FGetter::CreateStatic(&FLandscapeEditorDetails::GetTargetLandscapeSelectorVisibility)))
 	[
