Landscape : separate tool falloff, size & strength values for paint & sculpt

* Implemented as new Paint  falloff, size & strength members on ULandscapeEditorObject. Doing it this way made it was easier to persist the values different values for paint & sculpt when exiting & re-entering landscape mode.
* Perhaps there are other parameters which could be separated out between paint & sculpt but I hope brush size, strength & falloff are sufficient to make the workflow smoother.

diff --git a/Engine/Source/Editor/LandscapeEditor/Private/LandscapeEdMode.cpp b/Engine/Source/Editor/LandscapeEditor/Private/LandscapeEdMode.cpp
--- a/Engine/Source/Editor/LandscapeEditor/Private/LandscapeEdMode.cpp
+++ b/Engine/Source/Editor/LandscapeEditor/Private/LandscapeEdMode.cpp
@@ -1168,13 +1168,13 @@ bool FEdModeLandscape::LandscapeTrace(const FVector& InRayOrigin, const FVector&
 	// If there is no landscape directly under the mouse search for a landscape collision
 	// under the shape of the brush.
 	FCollisionShape SphereShape;
-	SphereShape.SetSphere(UISettings->BrushRadius);
+	SphereShape.SetSphere(UISettings->GetCurrentToolBrushRadius());
 	if ( World->SweepMultiByObjectType(Results, Start, End, FQuat::Identity, FCollisionObjectQueryParams(ECollisionChannel::ECC_Visibility), SphereShape, FCollisionQueryParams(SCENE_QUERY_STAT(LandscapeTrace), true)))
 	{
 		if (FProcessLandscapeTraceHitsResult SweepProcessResult; ProcessLandscapeTraceHits(Results, SweepProcessResult))
 		{
-			UE_VLOG_LOCATION(World, LogLandscapeEdMode, VeryVerbose,  SweepProcessResult.HitLocation, UISettings->BrushRadius,  FColor(255,100,100), TEXT("landscape:sweep-hit-location"));
-			FSphere HitSphere(SweepProcessResult.HitLocation, UISettings->BrushRadius);
+			UE_VLOG_LOCATION(World, LogLandscapeEdMode, VeryVerbose,  SweepProcessResult.HitLocation, UISettings->GetCurrentToolBrushRadius(),  FColor(255,100,100), TEXT("landscape:sweep-hit-location"));
+			FSphere HitSphere(SweepProcessResult.HitLocation, UISettings->GetCurrentToolBrushRadius());
 
 			float MeanHeight = 0.0f;
 			int32 Count = 0;
@@ -1633,7 +1633,7 @@ void FEdModeLandscape::ChangeBrushSize(bool bIncrease)
 	}
 	else
 	{
-		float Radius = UISettings->BrushRadius;
+		float Radius = UISettings->GetCurrentToolBrushRadius();
 		const float SliderMin = 10.0f;
 		const float SliderMax = 8192.0f;
 		float Diff = 0.05f; //6.0f / SliderMax;
@@ -1662,7 +1662,7 @@ void FEdModeLandscape::ChangeBrushSize(bool bIncrease)
 void FEdModeLandscape::ChangeBrushFalloff(bool bIncrease)
 {
 	UISettings->Modify();
-	float Falloff = UISettings->BrushFalloff;
+	float Falloff = UISettings->GetCurrentToolBrushFalloff();
 	const float SliderMin = 0.0f;
 	const float SliderMax = 1.0f;
 	float Diff = 0.05f; 
@@ -1690,7 +1690,7 @@ void FEdModeLandscape::ChangeBrushFalloff(bool bIncrease)
 void FEdModeLandscape::ChangeBrushStrength(bool bIncrease)
 {
 	UISettings->Modify();
-	float Strength = UISettings->ToolStrength;
+	float Strength = UISettings->GetCurrentToolStrength();
 	const float SliderMin = 0.01f;
 	const float SliderMax = 10.0f;
 	float Diff = 0.05f; //6.0f / SliderMax;
@@ -1711,7 +1711,7 @@ void FEdModeLandscape::ChangeBrushStrength(bool bIncrease)
 	}
 
 	NewValue = FMath::Clamp(NewValue, SliderMin, SliderMax);
-	UISettings->ToolStrength = NewValue;
+	UISettings->SetCurrentToolStrength(NewValue);
 }
 
 void FEdModeLandscape::ChangeAlphaBrushRotation(bool bIncrease)
diff --git a/Engine/Source/Editor/LandscapeEditor/Private/LandscapeEdModeBrushes.cpp b/Engine/Source/Editor/LandscapeEditor/Private/LandscapeEdModeBrushes.cpp
--- a/Engine/Source/Editor/LandscapeEditor/Private/LandscapeEdModeBrushes.cpp
+++ b/Engine/Source/Editor/LandscapeEditor/Private/LandscapeEdModeBrushes.cpp
@@ -129,9 +129,9 @@ public:
 		ALandscapeProxy* Proxy = LandscapeInfo->GetLandscapeProxy();
 
 		const float ScaleXY = FMath::Abs(LandscapeInfo->DrawScale.X);
-		const float TotalRadius = EdMode->UISettings->BrushRadius / ScaleXY;
-		const float Radius = (1.0f - EdMode->UISettings->BrushFalloff) * TotalRadius;
-		const float Falloff = EdMode->UISettings->BrushFalloff * TotalRadius;
+		const float TotalRadius = EdMode->UISettings->GetCurrentToolBrushRadius() / ScaleXY;
+		const float Radius = (1.0f - EdMode->UISettings->GetCurrentToolBrushFalloff()) * TotalRadius;
+		const float Falloff = EdMode->UISettings->GetCurrentToolBrushFalloff() * TotalRadius;
 
 		FIntRect Bounds;
 		Bounds.Min.X = FMath::FloorToInt(LastMousePosition.X - TotalRadius);
@@ -260,9 +260,9 @@ public:
 
 		ULandscapeInfo* LandscapeInfo = EdMode->CurrentToolTarget.LandscapeInfo.Get();
 		const float ScaleXY = FMath::Abs(LandscapeInfo->DrawScale.X);
-		const float TotalRadius = EdMode->UISettings->BrushRadius / ScaleXY;
-		const float Radius = (1.0f - EdMode->UISettings->BrushFalloff) * TotalRadius;
-		const float Falloff = EdMode->UISettings->BrushFalloff * TotalRadius;
+		const float TotalRadius = EdMode->UISettings->GetCurrentToolBrushRadius() / ScaleXY;
+		const float Radius = (1.0f - EdMode->UISettings->GetCurrentToolBrushFalloff()) * TotalRadius;
+		const float Falloff = EdMode->UISettings->GetCurrentToolBrushFalloff() * TotalRadius;
 
 		// Cap number of mouse positions to a sensible number
 		TArray<FLandscapeToolInteractorPosition> InteractorPositions;
@@ -778,7 +778,7 @@ public:
 					if (EdMode->UISettings->bSmoothGizmoBrush)
 					{
 						FVector TransformedLocal(FMath::Abs(GizmoLocal.X - LW), FMath::Abs(GizmoLocal.Y - LH) * (W / H), 0);
-						float FalloffRadius = LW * EdMode->UISettings->BrushFalloff;
+						float FalloffRadius = LW * EdMode->UISettings->GetCurrentToolBrushFalloff();
 						float SquareRadius = LW - FalloffRadius;
 						float Cos = FMath::Abs(TransformedLocal.X) / TransformedLocal.Size2D();
 						float Sin = FMath::Abs(TransformedLocal.Y) / TransformedLocal.Size2D();
@@ -1063,9 +1063,9 @@ public:
 
 		ULandscapeInfo* LandscapeInfo = EdMode->CurrentToolTarget.LandscapeInfo.Get();
 		const float ScaleXY = FMath::Abs(LandscapeInfo->DrawScale.X);
-		const float TotalRadius = EdMode->UISettings->BrushRadius / ScaleXY;
-		const float Radius = (1.0f - EdMode->UISettings->BrushFalloff) * TotalRadius;
-		const float Falloff = EdMode->UISettings->BrushFalloff * TotalRadius;
+		const float TotalRadius = EdMode->UISettings->GetCurrentToolBrushRadius() / ScaleXY;
+		const float Radius = (1.0f - EdMode->UISettings->GetCurrentToolBrushFalloff()) * TotalRadius;
+		const float Falloff = EdMode->UISettings->GetCurrentToolBrushFalloff() * TotalRadius;
 
 		int32 SizeX = EdMode->UISettings->AlphaTextureSizeX;
 		int32 SizeY = EdMode->UISettings->AlphaTextureSizeY;
@@ -1272,7 +1272,7 @@ public:
 		else
 		{
 			float ScaleXY = FMath::Abs(LandscapeInfo->DrawScale.X);
-			float Radius = EdMode->UISettings->BrushRadius / ScaleXY;
+			float Radius = EdMode->UISettings->GetCurrentToolBrushRadius() / ScaleXY;
 			int32 SizeX = EdMode->UISettings->AlphaTextureSizeX;
 			int32 SizeY = EdMode->UISettings->AlphaTextureSizeY;
 			float MaxSize = 2.0f * FMath::Sqrt(FMath::Square(Radius) / 2.0f);
@@ -1378,7 +1378,7 @@ public:
 			float ScaleXY = FMath::Abs(LandscapeInfo->DrawScale.X);
 			int32 SizeX = EdMode->UISettings->AlphaTextureSizeX;
 			int32 SizeY = EdMode->UISettings->AlphaTextureSizeY;
-			float Radius = EdMode->UISettings->BrushRadius / ScaleXY;
+			float Radius = EdMode->UISettings->GetCurrentToolBrushRadius() / ScaleXY;
 			float MaxSize = 2.0f * FMath::Sqrt(FMath::Square(Radius) / 2.0f);
 			float AlphaBrushScale = MaxSize / (float)FMath::Max<int32>(SizeX, SizeY);
 
diff --git a/Engine/Source/Editor/LandscapeEditor/Private/LandscapeEdModeComponentTools.cpp b/Engine/Source/Editor/LandscapeEditor/Private/LandscapeEdModeComponentTools.cpp
--- a/Engine/Source/Editor/LandscapeEditor/Private/LandscapeEdModeComponentTools.cpp
+++ b/Engine/Source/Editor/LandscapeEditor/Private/LandscapeEdModeComponentTools.cpp
@@ -209,7 +209,7 @@ public:
 
 					if (BrushValue > 0.0f && LandscapeInfo->IsValidPosition(X, Y))
 					{
-						float PaintValue = BrushValue * UISettings->ToolStrength * Pressure;
+						float PaintValue = BrushValue * UISettings->GetCurrentToolStrength() * Pressure;
 						float Value = DataScanline[X] / 255.0f;
 						checkSlow(FMath::IsNearlyEqual(Value, LandscapeInfo->SelectedRegion.FindRef(Key), 1 / 255.0f));
 						if (bInvert)
@@ -1495,7 +1495,7 @@ public:
 
 						// Value before we apply our painting
 						int32 index = (X - X1) + (Y - Y1)*(1 + X2 - X1);
-						float PaintAmount = (Brush->GetBrushType() == ELandscapeBrushType::Gizmo) ? BrushValue : BrushValue * EdMode->UISettings->ToolStrength * Pressure;
+						float PaintAmount = (Brush->GetBrushType() == ELandscapeBrushType::Gizmo) ? BrushValue : BrushValue * EdMode->UISettings->GetCurrentToolStrength() * Pressure;
 
 						FVector GizmoLocal = LandscapeToGizmoLocal.TransformPosition(FVector(X, Y, 0));
 						GizmoLocal.X *= ScaleX * SignX;
diff --git a/Engine/Source/Editor/LandscapeEditor/Private/LandscapeEdModeErosionTools.cpp b/Engine/Source/Editor/LandscapeEditor/Private/LandscapeEdModeErosionTools.cpp
--- a/Engine/Source/Editor/LandscapeEditor/Private/LandscapeEdModeErosionTools.cpp
+++ b/Engine/Source/Editor/LandscapeEditor/Private/LandscapeEdModeErosionTools.cpp
@@ -180,7 +180,7 @@ public:
 										uint16 Slope = HeightData[Center] - HeightData[Neighbor[Idx]];
 										if (Slope > Thresh)
 										{
-											float WeightDiff = Softness * UISettings->ToolStrength * Pressure * ((float)Slope / SlopeTotal) * BrushValue;
+											float WeightDiff = Softness * UISettings->GetCurrentToolStrength() * Pressure * ((float)Slope / SlopeTotal) * BrushValue;
 											//uint16 HeightDiff = (uint16)((SlopeMax - Thresh) * WeightDiff);
 											float HeightDiff = (SlopeMax - Thresh) * WeightDiff;
 											HeightData[Neighbor[Idx]] += HeightDiff;
@@ -210,7 +210,7 @@ public:
 								if (bWeightApplied)
 								{
 									float TotalWeight = 0.0f;
-									float WeightDiff = Softness * UISettings->ToolStrength * Pressure * BrushValue;
+									float WeightDiff = Softness * UISettings->GetCurrentToolStrength() * Pressure * BrushValue;
 
 									for (int32 LayerIdx = 0; LayerIdx < LayerNum; LayerIdx++)
 									{
@@ -238,9 +238,9 @@ public:
 		}
 
 		float BrushSizeAdjust = 1.0f;
-		if (UISettings->BrushRadius < UISettings->MaximumValueRadius)
+		if (UISettings->GetCurrentToolBrushRadius() < UISettings->MaximumValueRadius)
 		{
-			BrushSizeAdjust = UISettings->BrushRadius / UISettings->MaximumValueRadius;
+			BrushSizeAdjust = UISettings->GetCurrentToolBrushRadius() / UISettings->MaximumValueRadius;
 		}
 
 		// Make some noise...
@@ -254,7 +254,7 @@ public:
 
 				if (BrushValue > 0.0f)
 				{
-					FNoiseParameter NoiseParam(0, UISettings->ErosionNoiseScale, BrushValue * Thresh * UISettings->ToolStrength * BrushSizeAdjust);
+					FNoiseParameter NoiseParam(0, UISettings->ErosionNoiseScale, BrushValue * Thresh * UISettings->GetCurrentToolStrength() * BrushSizeAdjust);
 					float PaintAmount = NoiseModeConversion((ELandscapeToolNoiseMode)UISettings->ErosionNoiseMode, NoiseParam.NoiseAmount, NoiseParam.Sample(X, Y));
 					HeightData[(X - X1) + (Y - Y1)*(1 + X2 - X1)] = FLandscapeHeightCache::ClampValue(HeightData[(X - X1) + (Y - Y1)*(1 + X2 - X1)] + PaintAmount);
 				}
@@ -329,7 +329,7 @@ public:
 
 		const int32 Iteration = UISettings->HErodeIterationNum;
 		const uint16 RainAmount = UISettings->RainAmount;
-		const float DissolvingRatio = 0.07 * UISettings->ToolStrength * Pressure;  //0.01;
+		const float DissolvingRatio = 0.07 * UISettings->GetCurrentToolStrength() * Pressure;  //0.01;
 		const float EvaporateRatio = 0.5;
 		const float SedimentCapacity = 0.10 * UISettings->SedimentCapacity; //DissolvingRatio; //0.01;
 
@@ -435,7 +435,7 @@ public:
 							// This is not mathematically correct, but makes good result
 							if (TotalHeightDiff)
 							{
-								AverageAltitude *= (1.0f - 0.1 * UISettings->ToolStrength * Pressure);
+								AverageAltitude *= (1.0f - 0.1 * UISettings->GetCurrentToolStrength() * Pressure);
 								//AverageAltitude -= 4000.0f * UISettings->ToolStrength;
 							}
 
diff --git a/Engine/Source/Editor/LandscapeEditor/Private/LandscapeEdModePaintTools.cpp b/Engine/Source/Editor/LandscapeEditor/Private/LandscapeEdModePaintTools.cpp
--- a/Engine/Source/Editor/LandscapeEditor/Private/LandscapeEdModePaintTools.cpp
+++ b/Engine/Source/Editor/LandscapeEditor/Private/LandscapeEdModePaintTools.cpp
@@ -244,10 +244,10 @@ public:
 		}
 
 		// Adjust strength based on brush size and drawscale, so strength 1 = one hemisphere
-		const float AdjustedStrength = ToolTarget::StrengthMultiplier(this->LandscapeInfo, UISettings->BrushRadius);
+		const float AdjustedStrength = ToolTarget::StrengthMultiplier(this->LandscapeInfo, UISettings->GetCurrentToolBrushRadius());
 		FWeightmapToolTarget::CacheClass::DataType DestValue = FWeightmapToolTarget::CacheClass::ClampValue(255.0f * UISettings->WeightTargetValue);
 
-		float PaintStrength = UISettings->ToolStrength * Pressure * AdjustedStrength;
+		float PaintStrength = UISettings->GetCurrentToolStrength() * Pressure * AdjustedStrength;
 
 		// TODO: make paint tool framerate independent like the sculpt tool
 		// const float DeltaTime = FMath::Min<float>(FApp::GetDeltaTime(), 0.1f); // Under 10 fps slow down paint speed
@@ -387,7 +387,7 @@ public:
 
 				if (BrushValue > 0.0f)
 				{
-					float Strength = FMath::Clamp<float>(BrushValue * UISettings->ToolStrength * Pressure, 0.0f, 1.0f);
+					float Strength = FMath::Clamp<float>(BrushValue * UISettings->GetCurrentToolStrength() * Pressure, 0.0f, 1.0f);
 
 					int32 Delta = DataScanline[X] - FlattenHeight;
 					if (Delta > 0)
@@ -461,9 +461,9 @@ public:
 		FMatrix FromWorld = ToolTarget::FromWorldMatrix(this->LandscapeInfo);
 
 		// Adjust strength based on brush size and drawscale, so strength 1 = one hemisphere
-		const float AdjustedStrength = ToolTarget::StrengthMultiplier(this->LandscapeInfo, UISettings->BrushRadius);
+		const float AdjustedStrength = ToolTarget::StrengthMultiplier(this->LandscapeInfo, UISettings->GetCurrentToolBrushRadius());
 
-		float SculptStrength = UISettings->ToolStrength * Pressure * AdjustedStrength;
+		float SculptStrength = UISettings->GetCurrentToolStrength() * Pressure * AdjustedStrength;
 		const float DeltaTime = FMath::Min<float>(FApp::GetDeltaTime(), 0.1f); // Under 10 fps slow down paint speed
 		SculptStrength *= DeltaTime * 3.0f; // * 3.0f to partially compensate for impact of DeltaTime on slowing the tools down compared to the old framerate-dependent version
 
@@ -696,7 +696,7 @@ public:
 		TArray<typename ToolTarget::CacheClass::DataType> Data;
 		LayerDataCache.Initialize(this->LandscapeInfo, bCombinedLayerOperation);
 		LayerDataCache.Read(X1, Y1, X2, Y2, Data);
-		const float ToolStrength = FMath::Clamp<float>(UISettings->ToolStrength * Pressure, 0.0f, 1.0f);
+		const float ToolStrength = FMath::Clamp<float>(UISettings->GetCurrentToolStrength() * Pressure, 0.0f, 1.0f);
 
 		// Apply the brush
 		if (UISettings->bDetailSmooth)
@@ -925,7 +925,7 @@ public:
 
 				if (BrushValue > 0.0f)
 				{
-					float Strength = FMath::Clamp<float>(BrushValue * UISettings->ToolStrength * Pressure, 0.0f, 1.0f);
+					float Strength = FMath::Clamp<float>(BrushValue * UISettings->GetCurrentToolStrength() * Pressure, 0.0f, 1.0f);
 
 					if (!(bUseSlopeFlatten && bTargetIsHeightmap))
 					{
@@ -1237,9 +1237,9 @@ public:
 
 		float BrushSizeAdjust = 1.0f;
 		CA_SUPPRESS(6326);
-		if (ToolTarget::TargetType != ELandscapeToolTargetType::Weightmap && UISettings->BrushRadius < UISettings->MaximumValueRadius)
+		if (ToolTarget::TargetType != ELandscapeToolTargetType::Weightmap && UISettings->GetCurrentToolBrushRadius() < UISettings->MaximumValueRadius)
 		{
-			BrushSizeAdjust = UISettings->BrushRadius / UISettings->MaximumValueRadius;
+			BrushSizeAdjust = UISettings->GetCurrentToolBrushRadius() / UISettings->MaximumValueRadius;
 		}
 
 		CA_SUPPRESS(6326);
@@ -1278,11 +1278,11 @@ public:
 							}
 							break;
 						}
-						DataScanline[X] = ToolTarget::CacheClass::ClampValue(FMath::RoundToInt(FMath::Lerp(OriginalValue, DestValue, BrushValue * UISettings->ToolStrength * Pressure)));
+						DataScanline[X] = ToolTarget::CacheClass::ClampValue(FMath::RoundToInt(FMath::Lerp(OriginalValue, DestValue, BrushValue * UISettings->GetCurrentToolStrength() * Pressure)));
 					}
 					else
 					{
-						float TotalStrength = BrushValue * UISettings->ToolStrength * Pressure * ToolTarget::StrengthMultiplier(this->LandscapeInfo, UISettings->BrushRadius);
+						float TotalStrength = BrushValue * UISettings->GetCurrentToolStrength() * Pressure * ToolTarget::StrengthMultiplier(this->LandscapeInfo, UISettings->GetCurrentToolBrushRadius());
 						FNoiseParameter NoiseParam(0, UISettings->NoiseScale, TotalStrength * BrushSizeAdjust);
 						float PaintAmount = NoiseModeConversion(UISettings->NoiseMode, NoiseParam.NoiseAmount, NoiseParam.Sample(X, Y));
 						DataScanline[X] = ToolTarget::CacheClass::ClampValue(OriginalValue + PaintAmount);
diff --git a/Engine/Source/Editor/LandscapeEditor/Private/LandscapeEdModeXYOffset.cpp b/Engine/Source/Editor/LandscapeEditor/Private/LandscapeEdModeXYOffset.cpp
--- a/Engine/Source/Editor/LandscapeEditor/Private/LandscapeEdModeXYOffset.cpp
+++ b/Engine/Source/Editor/LandscapeEditor/Private/LandscapeEdModeXYOffset.cpp
@@ -520,7 +520,7 @@ public:
 		// Same as Gizmo fall off...
 		float W = X2 - X1 + 1;
 		float H = Y2 - Y1 + 1;
-		float FalloffRadius = W * 0.5f * UISettings->BrushFalloff;
+		float FalloffRadius = W * 0.5f * UISettings->GetCurrentToolBrushFalloff();
 		float SquareRadius = W * 0.5f - FalloffRadius;
 		for (int32 Y = 0; Y <= Y2 - Y1; ++Y)
 		{
diff --git a/Engine/Source/Editor/LandscapeEditor/Private/LandscapeEditorObject.cpp b/Engine/Source/Editor/LandscapeEditor/Private/LandscapeEditorObject.cpp
--- a/Engine/Source/Editor/LandscapeEditor/Private/LandscapeEditorObject.cpp
+++ b/Engine/Source/Editor/LandscapeEditor/Private/LandscapeEditorObject.cpp
@@ -22,6 +22,7 @@ ULandscapeEditorObject::ULandscapeEditorObject(const FObjectInitializer& ObjectI
 
 	// Tool Settings:
 	, ToolStrength(0.3f)
+    , PaintToolStrength(0.3f)
 	, bUseWeightTargetValue(false)
 	, WeightTargetValue(1.0f)
 	, MaximumValueRadius(10000.0f)
@@ -90,7 +91,9 @@ ULandscapeEditorObject::ULandscapeEditorObject(const FObjectInitializer& ObjectI
 
 	// Brush Settings:
 	, BrushRadius(2048.0f)
+    , PaintBrushRadius(2048.0f) 
 	, BrushFalloff(0.5f)
+	, PaintBrushFalloff(0.5f)
 	, bUseClayBrush(false)
 
 	, AlphaBrushScale(0.5f)
@@ -191,14 +194,17 @@ void ULandscapeEditorObject::PostEditChangeProperty(FPropertyChangedEvent& Prope
 void ULandscapeEditorObject::Load()
 {
 	GConfig->GetFloat(TEXT("LandscapeEdit"), TEXT("ToolStrength"), ToolStrength, GEditorPerProjectIni);
+	GConfig->GetFloat(TEXT("LandscapeEdit"), TEXT("PaintToolStrength"), PaintToolStrength, GEditorPerProjectIni);
 	GConfig->GetFloat(TEXT("LandscapeEdit"), TEXT("WeightTargetValue"), WeightTargetValue, GEditorPerProjectIni);
 	bool InbUseWeightTargetValue = bUseWeightTargetValue;
 	GConfig->GetBool(TEXT("LandscapeEdit"), TEXT("bUseWeightTargetValue"), InbUseWeightTargetValue, GEditorPerProjectIni);
 	bUseWeightTargetValue = InbUseWeightTargetValue;
 
 	GConfig->GetFloat(TEXT("LandscapeEdit"), TEXT("BrushRadius"), BrushRadius, GEditorPerProjectIni);
+	GConfig->GetFloat(TEXT("LandscapeEdit"), TEXT("PaintBrushRadius"), PaintBrushRadius, GEditorPerProjectIni);
 	GConfig->GetInt(TEXT("LandscapeEdit"), TEXT("BrushComponentSize"), BrushComponentSize, GEditorPerProjectIni);
 	GConfig->GetFloat(TEXT("LandscapeEdit"), TEXT("BrushFalloff"), BrushFalloff, GEditorPerProjectIni);
+	GConfig->GetFloat(TEXT("LandscapeEdit"), TEXT("PaintBrushFalloff"), PaintBrushFalloff, GEditorPerProjectIni);
 	bool InbUseClayBrush = bUseClayBrush;
 	GConfig->GetBool(TEXT("LandscapeEdit"), TEXT("bUseClayBrush"), InbUseClayBrush, GEditorPerProjectIni);
 	bUseClayBrush = InbUseClayBrush;
@@ -359,12 +365,15 @@ void ULandscapeEditorObject::Load()
 void ULandscapeEditorObject::Save()
 {
 	GConfig->SetFloat(TEXT("LandscapeEdit"), TEXT("ToolStrength"), ToolStrength, GEditorPerProjectIni);
+	GConfig->SetFloat(TEXT("LandscapeEdit"), TEXT("PaintToolStrength"), PaintToolStrength, GEditorPerProjectIni);
 	GConfig->SetFloat(TEXT("LandscapeEdit"), TEXT("WeightTargetValue"), WeightTargetValue, GEditorPerProjectIni);
 	GConfig->SetBool(TEXT("LandscapeEdit"), TEXT("bUseWeightTargetValue"), bUseWeightTargetValue, GEditorPerProjectIni);
 
 	GConfig->SetFloat(TEXT("LandscapeEdit"), TEXT("BrushRadius"), BrushRadius, GEditorPerProjectIni);
+	GConfig->SetFloat(TEXT("LandscapeEdit"), TEXT("PaintBrushRadius"), PaintBrushRadius, GEditorPerProjectIni);
 	GConfig->SetInt(TEXT("LandscapeEdit"), TEXT("BrushComponentSize"), BrushComponentSize, GEditorPerProjectIni);
 	GConfig->SetFloat(TEXT("LandscapeEdit"), TEXT("BrushFalloff"), BrushFalloff, GEditorPerProjectIni);
+	GConfig->SetFloat(TEXT("LandscapeEdit"), TEXT("PaintBrushFalloff"), PaintBrushFalloff, GEditorPerProjectIni);
 	GConfig->SetBool(TEXT("LandscapeEdit"), TEXT("bUseClayBrush"), bUseClayBrush, GEditorPerProjectIni);
 	GConfig->SetFloat(TEXT("LandscapeEdit"), TEXT("AlphaBrushScale"), AlphaBrushScale, GEditorPerProjectIni);
 	GConfig->SetBool(TEXT("LandscapeEdit"), TEXT("AlphaBrushAutoRotate"), bAlphaBrushAutoRotate, GEditorPerProjectIni);
@@ -939,3 +948,69 @@ void ULandscapeEditorObject::UpdateShowUnusedLayers()
 	}
 }
 
+float ULandscapeEditorObject::GetCurrentToolStrength() const
+{
+	if (IsWeightmapTarget())
+	{
+		return PaintToolStrength;
+	}
+	return ToolStrength;
+}
+
+void ULandscapeEditorObject::SetCurrentToolStrength(float NewToolStrength)
+{
+	if (IsWeightmapTarget())
+	{
+		PaintToolStrength = NewToolStrength;		
+	}
+	else
+	{
+		ToolStrength = NewToolStrength;
+	}
+}
+
+float ULandscapeEditorObject::GetCurrentToolBrushRadius() const
+{
+	if (IsWeightmapTarget())
+	{
+		return PaintBrushRadius;
+	}
+	return BrushRadius;
+	
+	
+}
+
+void ULandscapeEditorObject::SetCurrentToolBrushRadius(float NewBrushStrength)
+{
+	if (IsWeightmapTarget())
+	{
+		PaintBrushRadius = NewBrushStrength;
+	}
+	else
+	{
+		BrushRadius = NewBrushStrength;
+	}
+}
+
+float ULandscapeEditorObject::GetCurrentToolBrushFalloff() const
+{
+	if (IsWeightmapTarget())
+	{
+		return PaintBrushFalloff;
+		
+	}
+	return BrushFalloff;
+	
+}
+
+void ULandscapeEditorObject::SetCurrentToolBrushFalloff(float NewBrushFalloff)
+{
+	if (IsWeightmapTarget())
+	{
+		PaintBrushFalloff = NewBrushFalloff;
+	}
+	else
+	{
+		BrushFalloff = NewBrushFalloff;
+	}
+}
diff --git a/Engine/Source/Editor/LandscapeEditor/Public/LandscapeEditorObject.h b/Engine/Source/Editor/LandscapeEditor/Public/LandscapeEditorObject.h
--- a/Engine/Source/Editor/LandscapeEditor/Public/LandscapeEditorObject.h
+++ b/Engine/Source/Editor/LandscapeEditor/Public/LandscapeEditorObject.h
@@ -272,10 +272,14 @@ class ULandscapeEditorObject : public UObject
 
 	// Common Tool Settings:
 
-	// Strength of the tool. If you're using a pen/tablet with pressure-sensing, the pressure used affects the strength of the tool.
-	UPROPERTY(Category="Tool Settings", EditAnywhere, NonTransactional, meta=(ShowForTools="Paint,Sculpt,Erase,Smooth,Flatten,Erosion,HydraErosion,Noise,Mask,CopyPaste", ClampMin="0", ClampMax="10", UIMin="0", UIMax="1"))
+	// Strength of the Sculpt tool. If you're using a pen/tablet with pressure-sensing, the pressure used affects the strength of the tool.
+	UPROPERTY(Category="Tool Settings", EditAnywhere, NonTransactional, meta=(DisplayName="Tool Strength", ShowForTools="Paint,Sculpt,Erase,Smooth,Flatten,Erosion,HydraErosion,Noise,Mask,CopyPaste", ShowForTargetTypes = "Heightmap,Visibility",  ClampMin="0", ClampMax="10", UIMin="0", UIMax="1"))
 	float ToolStrength;
 
+	// Strength of the Paint tool. If you're using a pen/tablet with pressure-sensing, the pressure used affects the strength of the tool.
+	UPROPERTY(Category="Tool Settings", EditAnywhere, NonTransactional, meta=(DisplayName="Tool Strength", ShowForTools="Paint,Sculpt,Erase,Smooth,Flatten,Erosion,HydraErosion,Noise,Mask,CopyPaste", ShowForTargetTypes = "Weightmap", ClampMin="0", ClampMax="10", UIMin="0", UIMax="1"))
+	float PaintToolStrength;
+
 	// Enable to make tools blend towards a target value
 	UPROPERTY(Category = "Tool Settings", NonTransactional, EditAnywhere, meta = (InlineEditConditionToggle))
 	bool bUseWeightTargetValue;
@@ -590,13 +594,21 @@ public:
 
 	// Common Brush Settings:
 
-	// The radius of the brush, in unreal units
-	UPROPERTY(Category="Brush Settings", EditAnywhere, NonTransactional, meta=(DisplayName="Brush Size", ShowForBrushes="BrushSet_Circle,BrushSet_Alpha,BrushSet_Pattern", ClampMin="1", ClampMax="65536", UIMin="1", UIMax="8192", SliderExponent="3"))
+	// The radius of the sculpt brush, in unreal units
+	UPROPERTY(Category="Brush Settings", EditAnywhere, NonTransactional, meta=(DisplayName="Brush Size", ShowForBrushes="BrushSet_Circle,BrushSet_Alpha,BrushSet_Pattern", ShowForTargetTypes = "Heightmap,Visibility", ClampMin="1", ClampMax="65536", UIMin="1", UIMax="8192", SliderExponent="3"))
 	float BrushRadius;
 
-	// The falloff at the edge of the brush, as a fraction of the brush's size. 0 = no falloff, 1 = all falloff
-	UPROPERTY(Category="Brush Settings", EditAnywhere, NonTransactional, meta=(ShowForBrushes="BrushSet_Circle,BrushSet_Gizmo,BrushSet_Pattern", ClampMin="0", ClampMax="1", UIMin = "0", UIMax = "1"))
+	// The radius of the paint brush, in unreal units
+	UPROPERTY(Category="Brush Settings", EditAnywhere, NonTransactional, meta=(DisplayName="Brush Size", ShowForBrushes="BrushSet_Circle,BrushSet_Alpha,BrushSet_Pattern", ShowForTargetTypes = "Weightmap", ClampMin="1", ClampMax="65536", UIMin="1", UIMax="8192", SliderExponent="3"))
+	float PaintBrushRadius;
+
+	// The falloff at the edge of the sculpt brush, as a fraction of the brush's size. 0 = no falloff, 1 = all falloff
+	UPROPERTY(Category="Brush Settings", EditAnywhere, NonTransactional, meta=(DisplayName="Brush Falloff", ShowForBrushes="BrushSet_Circle,BrushSet_Gizmo,BrushSet_Pattern", ShowForTargetTypes = "Heightmap,Visibility", ClampMin="0", ClampMax="1", UIMin = "0", UIMax = "1"))
 	float BrushFalloff;
+	
+	// The falloff at the edge of the point brush, as a fraction of the brush's size. 0 = no falloff, 1 = all falloff
+	UPROPERTY(Category="Brush Settings", EditAnywhere, NonTransactional, meta=(DisplayName="Brush Falloff", ShowForBrushes="BrushSet_Circle,BrushSet_Gizmo,BrushSet_Pattern", ShowForTargetTypes = "Weightmap", ClampMin="0", ClampMax="1", UIMin = "0", UIMax = "1"))
+	float PaintBrushFalloff;
 
 	// Selects the Clay Brush painting mode
 	UPROPERTY(Category="Brush Settings", EditAnywhere, NonTransactional, meta=(ShowForTools="Sculpt", ShowForBrushes="BrushSet_Circle,BrushSet_Alpha,BrushSet_Pattern"))
@@ -762,4 +774,21 @@ public:
 
 	void UpdateTargetLayerDisplayOrder();
 	void UpdateShowUnusedLayers();
+
+	float GetCurrentToolStrength() const;
+	void SetCurrentToolStrength(float NewToolStrength);
+
+	float GetCurrentToolBrushRadius() const;
+	void SetCurrentToolBrushRadius(float NewBrushStrength);
+
+	float GetCurrentToolBrushFalloff() const;
+	void SetCurrentToolBrushFalloff(float NewBrushFalloff);
+
+private:
+
+	bool IsWeightmapTarget() const
+	{
+		check(ParentMode);
+		return !(ParentMode->CurrentToolTarget.TargetType == ELandscapeToolTargetType::Heightmap || ParentMode->CurrentToolTarget.TargetType == ELandscapeToolTargetType::Visibility);
+	}
 };
