Fixed keyboard shortcuts for landscape paint brush

#rnx

diff --git a/Engine/Source/Editor/LandscapeEditor/Private/LandscapeEdMode.cpp b/Engine/Source/Editor/LandscapeEditor/Private/LandscapeEdMode.cpp
--- a/Engine/Source/Editor/LandscapeEditor/Private/LandscapeEdMode.cpp
+++ b/Engine/Source/Editor/LandscapeEditor/Private/LandscapeEdMode.cpp
@@ -1654,7 +1654,7 @@ void FEdModeLandscape::ChangeBrushSize(bool bIncrease)
 		}
 
 		NewValue = (int32)FMath::Clamp(NewValue, SliderMin, SliderMax);
-		UISettings->BrushRadius = NewValue;
+		UISettings->SetCurrentToolBrushRadius(NewValue);
 	}
 }
 
@@ -1683,7 +1683,7 @@ void FEdModeLandscape::ChangeBrushFalloff(bool bIncrease)
 	}
 
 	NewValue = FMath::Clamp(NewValue, SliderMin, SliderMax);
-	UISettings->BrushFalloff = NewValue;
+	UISettings->SetCurrentToolBrushFalloff(NewValue);
 }
 
 
