Landscape: Limit "Fill World" in restrictive mode to 2k x 2k

diff --git a/Engine/Source/Editor/LandscapeEditor/Private/LandscapeEditorDetailCustomization_NewLandscape.cpp b/Engine/Source/Editor/LandscapeEditor/Private/LandscapeEditorDetailCustomization_NewLandscape.cpp
--- a/Engine/Source/Editor/LandscapeEditor/Private/LandscapeEditorDetailCustomization_NewLandscape.cpp
+++ b/Engine/Source/Editor/LandscapeEditor/Private/LandscapeEditorDetailCustomization_NewLandscape.cpp
@@ -1014,7 +1014,23 @@ FReply FLandscapeEditorDetailCustomization_NewLandscape::OnFillWorldButtonClicke
 		const int32 QuadsPerComponent = LandscapeEdMode->UISettings->NewLandscape_SectionsPerComponent * LandscapeEdMode->UISettings->NewLandscape_QuadsPerSection;
 		LandscapeEdMode->UISettings->NewLandscape_ComponentCount.X = FMath::CeilToInt(WORLD_MAX / QuadsPerComponent / LandscapeEdMode->UISettings->NewLandscape_Scale.X);
 		LandscapeEdMode->UISettings->NewLandscape_ComponentCount.Y = FMath::CeilToInt(WORLD_MAX / QuadsPerComponent / LandscapeEdMode->UISettings->NewLandscape_Scale.Y);
-		LandscapeEdMode->UISettings->NewLandscape_ClampSize();
+
+		const ULandscapeSettings* Settings = GetDefault<ULandscapeSettings>();
+		if (Settings->IsLandscapeResolutionRestricted())
+		{
+			auto ClampComponentCount = [Settings, &QuadsPerComponent](int32& ComponentCount)
+			{
+				const float MaxResolution = Settings->GetSideResolutionLimit();
+				ComponentCount = FMath::Clamp(ComponentCount, 1, FMath::Min(32, FMath::FloorToInt((MaxResolution - 1) / QuadsPerComponent)));
+			};
+
+			ClampComponentCount(LandscapeEdMode->UISettings->NewLandscape_ComponentCount.X);
+			ClampComponentCount(LandscapeEdMode->UISettings->NewLandscape_ComponentCount.Y);
+		}
+		else
+		{
+			LandscapeEdMode->UISettings->NewLandscape_ClampSize();
+		}
 	}
 
 	return FReply::Handled();
