Fixed ensure in Chaos when LandscapeFlatten tool is used with an invalid mouse position (e.g. while working in RDP)

#rnx

diff --git a/Engine/Source/Editor/LandscapeEditor/Private/LandscapeEdModePaintTools.cpp b/Engine/Source/Editor/LandscapeEditor/Private/LandscapeEdModePaintTools.cpp
--- a/Engine/Source/Editor/LandscapeEditor/Private/LandscapeEdModePaintTools.cpp
+++ b/Engine/Source/Editor/LandscapeEditor/Private/LandscapeEdModePaintTools.cpp
@@ -1131,21 +1131,22 @@ public:
 		if (ViewportClient->IsLevelEditorClient() && HeightmapFlattenPreviewComponent != nullptr)
 		{
 			FVector MousePosition;
-			this->EdMode->LandscapeMouseTrace((FEditorViewportClient*)ViewportClient, x, y, MousePosition);
-
-			const FTransform LocalToWorld = this->EdMode->CurrentToolTarget.LandscapeInfo->GetLandscapeProxy()->ActorToWorld();
-			FVector Origin;
-			Origin.X = FMath::RoundToFloat(MousePosition.X);
-			Origin.Y = FMath::RoundToFloat(MousePosition.Y);
-			Origin.Z = (FMath::RoundToFloat((this->EdMode->UISettings->FlattenTarget - LocalToWorld.GetTranslation().Z) / LocalToWorld.GetScale3D().Z * LANDSCAPE_INV_ZSCALE) - 0.1f) * LANDSCAPE_ZSCALE;
-			HeightmapFlattenPreviewComponent->SetRelativeLocation(Origin, false);
-
-			// Clamp the value to the height map
-			uint16 TexHeight = LandscapeDataAccess::GetTexHeight(MousePosition.Z);
-			float Height = LandscapeDataAccess::GetLocalHeight(TexHeight);
-
-			// Convert the height back to world space
-			this->EdMode->UISettings->FlattenEyeDropperModeDesiredTarget = (Height * LocalToWorld.GetScale3D().Z) + LocalToWorld.GetTranslation().Z;
+			if (this->EdMode->LandscapeMouseTrace((FEditorViewportClient*)ViewportClient, x, y, MousePosition))
+			{
+				const FTransform LocalToWorld = this->EdMode->CurrentToolTarget.LandscapeInfo->GetLandscapeProxy()->ActorToWorld();
+				FVector Origin;
+				Origin.X = FMath::RoundToFloat(MousePosition.X);
+				Origin.Y = FMath::RoundToFloat(MousePosition.Y);
+				Origin.Z = (FMath::RoundToFloat((this->EdMode->UISettings->FlattenTarget - LocalToWorld.GetTranslation().Z) / LocalToWorld.GetScale3D().Z * LANDSCAPE_INV_ZSCALE) - 0.1f) * LANDSCAPE_ZSCALE;
+				HeightmapFlattenPreviewComponent->SetRelativeLocation(Origin, false);
+
+				// Clamp the value to the height map
+				uint16 TexHeight = LandscapeDataAccess::GetTexHeight(MousePosition.Z);
+				float Height = LandscapeDataAccess::GetLocalHeight(TexHeight);
+
+				// Convert the height back to world space
+				this->EdMode->UISettings->FlattenEyeDropperModeDesiredTarget = (Height * LocalToWorld.GetScale3D().Z) + LocalToWorld.GetTranslation().Z;
+			}
 		}
 
 		return bResult;
