Fixed the 2 places in landscape where we're reading texture data on the CPU and assume that it's valid :
* SplineFalloffModulationTexture for Layer Info Object : we now log an error when applying the change and don't apply modulation in case the texture data is invalid
* AlphaTexture for the alpha brush : we now log an error when setting the texture and revert to the default texture in case the texture data is invalid

Misc:
* Fixed undo/redo for Layer Info Object not triggering an update of the landscape splines

diff --git a/Engine/Source/Editor/LandscapeEditor/Private/LandscapeEdModeTools.h b/Engine/Source/Editor/LandscapeEditor/Private/LandscapeEdModeTools.h
--- a/Engine/Source/Editor/LandscapeEditor/Private/LandscapeEdModeTools.h
+++ b/Engine/Source/Editor/LandscapeEditor/Private/LandscapeEdModeTools.h
@@ -15,13 +15,11 @@
 #include "InstancedFoliageActor.h"
 #include "AI/NavigationSystemBase.h"
 #include "Landscape.h"
+#include "LandscapeEditorPrivate.h"
 #include "Logging/LogMacros.h"
 #include "VisualLogger/VisualLogger.h"
 
 
-DECLARE_LOG_CATEGORY_EXTERN(LogLandscapeTools, Log, All);
-
-
 //
 //	FNoiseParameter - Perlin noise
 //
diff --git a/Engine/Source/Editor/LandscapeEditor/Private/LandscapeEditorObject.cpp b/Engine/Source/Editor/LandscapeEditor/Private/LandscapeEditorObject.cpp
--- a/Engine/Source/Editor/LandscapeEditor/Private/LandscapeEditorObject.cpp
+++ b/Engine/Source/Editor/LandscapeEditor/Private/LandscapeEditorObject.cpp
@@ -7,6 +7,7 @@
 #include "Modules/ModuleManager.h"
 #include "UObject/ConstructorHelpers.h"
 #include "LandscapeEditorModule.h"
+#include "LandscapeEditorPrivate.h"
 #include "LandscapeRender.h"
 #include "LandscapeSettings.h"
 #include "LandscapeImportHelper.h"
@@ -78,7 +79,7 @@ ULandscapeEditorObject::ULandscapeEditorObject(const FObjectInitializer& ObjectI
 	, ResizeLandscape_ComponentCount(0, 0)
 	, ResizeLandscape_ConvertMode(ELandscapeConvertMode::Expand)
 
-	, NewLandscape_Material(NULL)
+	, NewLandscape_Material(nullptr)
 	, NewLandscape_QuadsPerSection(63)
 	, NewLandscape_SectionsPerComponent(1)
 	, NewLandscape_ComponentCount(8, 8)
@@ -103,7 +104,7 @@ ULandscapeEditorObject::ULandscapeEditorObject(const FObjectInitializer& ObjectI
 	, AlphaBrushPanV(0.5f)
 	, bUseWorldSpacePatternBrush(false)
 	, WorldSpacePatternBrushSettings(FVector2D::ZeroVector, 0.0f, false, 3200)
-	, AlphaTexture(NULL)
+	, AlphaTexture(nullptr)
 	, AlphaTextureChannel(EColorChannel::Red)
 	, AlphaTextureSizeX(1)
 	, AlphaTextureSizeY(1)
@@ -217,12 +218,12 @@ void ULandscapeEditorObject::Load()
 	GConfig->GetVector2D(TEXT("LandscapeEdit"), TEXT("WorldSpacePatternBrushSettings.Origin"), WorldSpacePatternBrushSettings.Origin, GEditorPerProjectIni);
 	GConfig->GetBool(TEXT("LandscapeEdit"), TEXT("WorldSpacePatternBrushSettings.bCenterTextureOnOrigin"), WorldSpacePatternBrushSettings.bCenterTextureOnOrigin, GEditorPerProjectIni);
 	GConfig->GetFloat(TEXT("LandscapeEdit"), TEXT("WorldSpacePatternBrushSettings.RepeatSize"), WorldSpacePatternBrushSettings.RepeatSize, GEditorPerProjectIni);
-	FString AlphaTextureName = (AlphaTexture != NULL) ? AlphaTexture->GetPathName() : FString();
+	FString AlphaTextureName = (AlphaTexture != nullptr) ? AlphaTexture->GetPathName() : FString();
 	int32 InAlphaTextureChannel = AlphaTextureChannel;
 	GConfig->GetString(TEXT("LandscapeEdit"), TEXT("AlphaTextureName"), AlphaTextureName, GEditorPerProjectIni);
 	GConfig->GetInt(TEXT("LandscapeEdit"), TEXT("AlphaTextureChannel"), InAlphaTextureChannel, GEditorPerProjectIni);
 	AlphaTextureChannel = (EColorChannel::Type)InAlphaTextureChannel;
-	SetAlphaTexture(LoadObject<UTexture2D>(NULL, *AlphaTextureName, NULL, LOAD_NoWarn), AlphaTextureChannel);
+	SetAlphaTexture(LoadObject<UTexture2D>(nullptr, *AlphaTextureName, nullptr, LOAD_NoWarn), AlphaTextureChannel);
 
 	int32 InFlattenMode = (int32)ELandscapeToolFlattenMode::Both;
 	GConfig->GetInt(TEXT("LandscapeEdit"), TEXT("FlattenMode"), InFlattenMode, GEditorPerProjectIni);
@@ -351,7 +352,7 @@ void ULandscapeEditorObject::Load()
 	
 	if (!NewLandscapeMaterialName.IsEmpty())
 	{
-		NewLandscape_Material = LoadObject<UMaterialInterface>(NULL, *NewLandscapeMaterialName, NULL, LOAD_NoWarn);
+		NewLandscape_Material = LoadObject<UMaterialInterface>(nullptr, *NewLandscapeMaterialName, nullptr, LOAD_NoWarn);
 	}
 	
 	int32 AlphamapType = (uint8)ImportLandscape_AlphamapType;
@@ -383,7 +384,7 @@ void ULandscapeEditorObject::Save()
 	GConfig->SetVector2D(TEXT("LandscapeEdit"), TEXT("WorldSpacePatternBrushSettings.Origin"), WorldSpacePatternBrushSettings.Origin, GEditorPerProjectIni);
 	GConfig->SetBool(TEXT("LandscapeEdit"), TEXT("WorldSpacePatternBrushSettings.bCenterTextureOnOrigin"), WorldSpacePatternBrushSettings.bCenterTextureOnOrigin, GEditorPerProjectIni);
 	GConfig->SetFloat(TEXT("LandscapeEdit"), TEXT("WorldSpacePatternBrushSettings.RepeatSize"), WorldSpacePatternBrushSettings.RepeatSize, GEditorPerProjectIni);
-	const FString AlphaTextureName = (AlphaTexture != NULL) ? AlphaTexture->GetPathName() : FString();
+	const FString AlphaTextureName = (AlphaTexture != nullptr) ? AlphaTexture->GetPathName() : FString();
 	GConfig->SetString(TEXT("LandscapeEdit"), TEXT("AlphaTextureName"), *AlphaTextureName, GEditorPerProjectIni);
 	GConfig->SetInt(TEXT("LandscapeEdit"), TEXT("AlphaTextureChannel"), (int32)AlphaTextureChannel, GEditorPerProjectIni);
 
@@ -433,7 +434,7 @@ void ULandscapeEditorObject::Save()
 	//GConfig->SetBool(TEXT("LandscapeEdit"), TEXT("bUseNegativeMask"), bUseNegativeMask, GEditorPerProjectIni);
 	GConfig->SetBool(TEXT("LandscapeEdit"), TEXT("bApplyToAllTargets"), bApplyToAllTargets, GEditorPerProjectIni);
 
-	const FString NewLandscapeMaterialName = (NewLandscape_Material != NULL) ? NewLandscape_Material->GetPathName() : FString();
+	const FString NewLandscapeMaterialName = (NewLandscape_Material != nullptr) ? NewLandscape_Material->GetPathName() : FString();
 	GConfig->SetString(TEXT("LandscapeEdit"), TEXT("NewLandscapeMaterialName"), *NewLandscapeMaterialName, GEditorPerProjectIni);
 
 	GConfig->SetInt(TEXT("LandscapeEdit"), TEXT("ImportLandscape_AlphamapType"), (uint8)ImportLandscape_AlphamapType, GEditorPerProjectIni);
@@ -499,80 +500,78 @@ void ULandscapeEditorObject::SetbSnapGizmo(bool InbSnapGizmo)
 	}
 }
 
-bool ULandscapeEditorObject::SetAlphaTexture(UTexture2D* InTexture, EColorChannel::Type InTextureChannel)
+void ULandscapeEditorObject::SetAlphaTexture(UTexture2D* InTexture, EColorChannel::Type InTextureChannel)
 {
-	bool Result = true;
-
 	TArray64<uint8> NewTextureData;
-	UTexture2D* NewAlphaTexture = InTexture;
-
-	// No texture or no source art, try to use the previous texture.
-	if (NewAlphaTexture == NULL || !NewAlphaTexture->Source.IsValid())
-	{
-		NewAlphaTexture = AlphaTexture;
-		Result = false;
-	}
-
-	if (NewAlphaTexture != NULL && NewAlphaTexture->Source.IsValid())
-	{
-		NewAlphaTexture->Source.GetMipData(NewTextureData, 0);
-	}
-
-	const bool bSourceDataIsG8 = NewAlphaTexture != NULL && NewAlphaTexture->Source.IsValid() && NewAlphaTexture->Source.GetFormat() == TSF_G8;
-	const int32 NumChannels = bSourceDataIsG8 ? 1 : 4;
+	UTexture2D* NewAlphaTexture = nullptr;
+	int32 NumChannels = 0;
 
-	// Load fallback if there's no texture or data
-	if (NewAlphaTexture == NULL || (NewTextureData.Num() != NumChannels * NewAlphaTexture->Source.GetSizeX() * NewAlphaTexture->Source.GetSizeY()))
+	// Validate that the input texture is valid, if not, we'll display an error message and use the default brush alpha texture : 
+	if (InTexture != nullptr)
 	{
-		NewAlphaTexture = GetClass()->GetDefaultObject<ULandscapeEditorObject>()->AlphaTexture;
-		if (NewAlphaTexture)
+		if (!InTexture->Source.IsValid())
 		{
-			NewAlphaTexture->Source.GetMipData(NewTextureData, 0);
-		}
-		Result = false;
-	}
-
-	if (NewAlphaTexture)
-	{
-		AlphaTexture = NewAlphaTexture;
-		AlphaTextureSizeX = NewAlphaTexture->Source.GetSizeX();
-		AlphaTextureSizeY = NewAlphaTexture->Source.GetSizeY();
-		AlphaTextureChannel = NumChannels == 1 ? EColorChannel::Red : InTextureChannel;
-		AlphaTextureData.Empty(AlphaTextureSizeX * AlphaTextureSizeY);
-
-		if (NewTextureData.Num() != NumChannels * AlphaTextureSizeX * AlphaTextureSizeY)
-		{
-			// Don't crash if for some reason we couldn't load any source art
-			AlphaTextureData.AddZeroed(AlphaTextureSizeX * AlphaTextureSizeY);
+			UE_LOG(LogLandscapeTools, Error, TEXT("Invalid source data detected for texture (%s), the default AlphaTexture (%s) will be used."), *InTexture->GetPathName(), *GetClass()->GetDefaultObject<ULandscapeEditorObject>()->AlphaTexture->GetPathName());
 		}
 		else
 		{
-			uint8* SrcPtr;
-			switch (AlphaTextureChannel)
+			// Try to read the new texture data now : 
+			const bool bSourceDataIsG8 = (InTexture->Source.GetFormat() == TSF_G8);
+			NumChannels = bSourceDataIsG8 ? 1 : 4;
+			InTexture->Source.GetMipData(NewTextureData, 0);
+			if (NewTextureData.Num() == (NumChannels * InTexture->Source.GetSizeX() * InTexture->Source.GetSizeY()))
 			{
-			case 1:
-				SrcPtr = &((FColor*)NewTextureData.GetData())->G;
-				break;
-			case 2:
-				SrcPtr = &((FColor*)NewTextureData.GetData())->B;
-				break;
-			case 3:
-				SrcPtr = &((FColor*)NewTextureData.GetData())->A;
-				break;
-			default:
-				SrcPtr = &((FColor*)NewTextureData.GetData())->R;
-				break;
+				// Valid new texture
+				NewAlphaTexture = InTexture;
 			}
-
-			for (int32 i = 0; i < AlphaTextureSizeX * AlphaTextureSizeY; i++)
+			else
 			{
-				AlphaTextureData.Add(*SrcPtr);
-				SrcPtr += NumChannels;
+				UE_LOG(LogLandscapeTools, Error, TEXT("Invalid data size detected for texture (%s), the default AlphaTexture (%s) will be used."), *InTexture->GetPathName(), *GetClass()->GetDefaultObject<ULandscapeEditorObject>()->AlphaTexture->GetPathName());
 			}
 		}
 	}
 
-	return Result;
+	// Load fallback if there's no texture or valid data
+	if (NewAlphaTexture == nullptr)
+	{
+		UTexture2D* DefaultAlphaTexture = GetClass()->GetDefaultObject<ULandscapeEditorObject>()->AlphaTexture;
+		check((DefaultAlphaTexture != nullptr) && DefaultAlphaTexture->Source.IsValid()); // The default texture should always be valid
+		DefaultAlphaTexture->Source.GetMipData(NewTextureData, 0);
+		NewAlphaTexture = DefaultAlphaTexture;
+		const bool bSourceDataIsG8 = (DefaultAlphaTexture->Source.GetFormat() == TSF_G8);
+		NumChannels = bSourceDataIsG8 ? 1 : 4;
+	}
+
+	check((NewAlphaTexture != nullptr) && !NewTextureData.IsEmpty() && (NumChannels > 0));
+
+	AlphaTexture = NewAlphaTexture;
+	AlphaTextureSizeX = NewAlphaTexture->Source.GetSizeX();
+	AlphaTextureSizeY = NewAlphaTexture->Source.GetSizeY();
+	AlphaTextureChannel = (NumChannels == 1) ? EColorChannel::Red : InTextureChannel;
+	AlphaTextureData.Empty(AlphaTextureSizeX * AlphaTextureSizeY);
+
+	uint8* SrcPtr;
+	switch (AlphaTextureChannel)
+	{
+	case 1:
+		SrcPtr = &((FColor*)NewTextureData.GetData())->G;
+		break;
+	case 2:
+		SrcPtr = &((FColor*)NewTextureData.GetData())->B;
+		break;
+	case 3:
+		SrcPtr = &((FColor*)NewTextureData.GetData())->A;
+		break;
+	default:
+		SrcPtr = &((FColor*)NewTextureData.GetData())->R;
+		break;
+	}
+
+	for (int32 i = 0; i < AlphaTextureSizeX * AlphaTextureSizeY; i++)
+	{
+		AlphaTextureData.Add(*SrcPtr);
+		SrcPtr += NumChannels;
+	}
 }
 
 void ULandscapeEditorObject::ChooseBestComponentSizeForImport()
@@ -901,12 +900,12 @@ void ULandscapeEditorObject::RefreshImportLayersList(bool bRefreshFromTarget)
 		{
 			if (!ThumbnailWeightmap)
 			{
-				ThumbnailWeightmap = LoadObject<UTexture2D>(NULL, TEXT("/Engine/EditorLandscapeResources/LandscapeThumbnailWeightmap.LandscapeThumbnailWeightmap"), NULL, LOAD_None, NULL);
+				ThumbnailWeightmap = LoadObject<UTexture2D>(nullptr, TEXT("/Engine/EditorLandscapeResources/LandscapeThumbnailWeightmap.LandscapeThumbnailWeightmap"), nullptr, LOAD_None, nullptr);
 			}
 
 			if (!ThumbnailHeightmap)
 			{
-				ThumbnailHeightmap = LoadObject<UTexture2D>(NULL, TEXT("/Engine/EditorLandscapeResources/LandscapeThumbnailHeightmap.LandscapeThumbnailHeightmap"), NULL, LOAD_None, NULL);
+				ThumbnailHeightmap = LoadObject<UTexture2D>(nullptr, TEXT("/Engine/EditorLandscapeResources/LandscapeThumbnailHeightmap.LandscapeThumbnailHeightmap"), nullptr, LOAD_None, nullptr);
 			}
 
 			NewImportLayer.LayerName = LayerName;
diff --git a/Engine/Source/Editor/LandscapeEditor/Private/LandscapeEditorPrivate.h b/Engine/Source/Editor/LandscapeEditor/Private/LandscapeEditorPrivate.h
new file mode 100644
--- /dev/null
+++ b/Engine/Source/Editor/LandscapeEditor/Private/LandscapeEditorPrivate.h
@@ -0,0 +1,9 @@
+// Copyright Epic Games, Inc. All Rights Reserved.
+
+#pragma once
+
+#include "CoreMinimal.h"
+#include "Logging/LogMacros.h"
+
+DECLARE_LOG_CATEGORY_EXTERN(LogLandscapeTools, Log, All);
+
diff --git a/Engine/Source/Editor/LandscapeEditor/Public/LandscapeEditorObject.h b/Engine/Source/Editor/LandscapeEditor/Public/LandscapeEditorObject.h
--- a/Engine/Source/Editor/LandscapeEditor/Public/LandscapeEditorObject.h
+++ b/Engine/Source/Editor/LandscapeEditor/Public/LandscapeEditorObject.h
@@ -697,7 +697,7 @@ public:
 	void SetPasteMode(ELandscapeToolPasteMode InPasteMode);
 
 	// Alpha/Pattern Brush
-	bool SetAlphaTexture(UTexture2D* InTexture, EColorChannel::Type InTextureChannel);
+	void SetAlphaTexture(UTexture2D* InTexture, EColorChannel::Type InTextureChannel);
 
 	// New Landscape
 	FString LastImportPath;
diff --git a/Engine/Source/Runtime/Landscape/Classes/LandscapeLayerInfoObject.h b/Engine/Source/Runtime/Landscape/Classes/LandscapeLayerInfoObject.h
--- a/Engine/Source/Runtime/Landscape/Classes/LandscapeLayerInfoObject.h
+++ b/Engine/Source/Runtime/Landscape/Classes/LandscapeLayerInfoObject.h
@@ -69,6 +69,7 @@ class ULandscapeLayerInfoObject : public UObject
 #if WITH_EDITOR
 	//~ Begin UObject Interface
 	virtual void PostEditChangeProperty(FPropertyChangedEvent& PropertyChangedEvent) override;
+	virtual void PostEditUndo() override;
 	virtual void PostLoad() override;
 	//~ End UObject Interface
 #endif
diff --git a/Engine/Source/Runtime/Landscape/Private/LandscapeEdit.cpp b/Engine/Source/Runtime/Landscape/Private/LandscapeEdit.cpp
--- a/Engine/Source/Runtime/Landscape/Private/LandscapeEdit.cpp
+++ b/Engine/Source/Runtime/Landscape/Private/LandscapeEdit.cpp
@@ -4409,6 +4409,20 @@ void ULandscapeLayerInfoObject::PostEditChangeProperty(FPropertyChangedEvent& Pr
 	}
 }
 
+void ULandscapeLayerInfoObject::PostEditUndo()
+{
+	Super::PostEditUndo();
+
+	// Force the update of spline data for the (potentially) affected landscapes in case of an undo
+	for (TObjectIterator<ULandscapeInfo> It; It; ++It)
+	{
+		if (ALandscape* Landscape = It->LandscapeActor.Get())
+		{
+			Landscape->OnLayerInfoSplineFalloffModulationChanged(this);
+		}
+	}
+}
+
 void ULandscapeLayerInfoObject::PostLoad()
 {
 	Super::PostLoad();
diff --git a/Engine/Source/Runtime/Landscape/Private/LandscapeSplineRaster.cpp b/Engine/Source/Runtime/Landscape/Private/LandscapeSplineRaster.cpp
--- a/Engine/Source/Runtime/Landscape/Private/LandscapeSplineRaster.cpp
+++ b/Engine/Source/Runtime/Landscape/Private/LandscapeSplineRaster.cpp
@@ -16,6 +16,7 @@
 #include "LandscapeEdit.h"
 #include "LandscapeSplinesComponent.h"
 #include "LandscapeSplineControlPoint.h"
+#include "LandscapePrivate.h"
 #if WITH_EDITOR
 #include "ScopedTransaction.h"
 #include "Raster.h"
@@ -126,11 +127,19 @@ class FModulateAlpha
 public:
 	static TSharedPtr<FModulateAlpha> CreateFromLayerInfo(ULandscapeLayerInfoObject* InLayerInfo, int32 InLandscapeMinX, int32 InLandscapeMinY)
 	{
-		if (CVarLandscapeSplineFalloffModulation.GetValueOnAnyThread() == 0 || InLayerInfo == nullptr || InLayerInfo->SplineFalloffModulationTexture == nullptr)
+		if (CVarLandscapeSplineFalloffModulation.GetValueOnAnyThread() == 0 
+			|| (InLayerInfo == nullptr) 
+			|| (InLayerInfo->SplineFalloffModulationTexture == nullptr))
 		{
 			return nullptr;
 		}
 
+		if (!InLayerInfo->SplineFalloffModulationTexture->Source.IsValid())
+		{
+			UE_LOG(LogLandscape, Error, TEXT("Invalid source data for spline falloff modulation texture (%s). Alpha modulation will be disabled."), *InLayerInfo->SplineFalloffModulationTexture->GetPathName());
+			return nullptr;
+		}
+
 		return MakeShareable(new FModulateAlpha(InLayerInfo, InLandscapeMinX, InLandscapeMinY));
 	}
 
@@ -164,7 +173,7 @@ private:
 	{
 		check(InLayerInfo && InLayerInfo->SplineFalloffModulationTexture);
 
-		InLayerInfo->SplineFalloffModulationTexture->Source.GetMipData(MipData, 0);
+		verify(InLayerInfo->SplineFalloffModulationTexture->Source.GetMipData(MipData, 0));
 		TextureWidth = InLayerInfo->SplineFalloffModulationTexture->Source.GetSizeX();
 		TextureHeight = InLayerInfo->SplineFalloffModulationTexture->Source.GetSizeY();
 
