Fix all known cases of GetTypedOuter<UPackage> since it's not valid anymore with override packages.

#rnx

diff --git a/Engine/Plugins/Experimental/PythonScriptPlugin/Source/PythonScriptPlugin/Private/PyWrapperTypeRegistry.cpp b/Engine/Plugins/Experimental/PythonScriptPlugin/Source/PythonScriptPlugin/Private/PyWrapperTypeRegistry.cpp
--- a/Engine/Plugins/Experimental/PythonScriptPlugin/Source/PythonScriptPlugin/Private/PyWrapperTypeRegistry.cpp
+++ b/Engine/Plugins/Experimental/PythonScriptPlugin/Source/PythonScriptPlugin/Private/PyWrapperTypeRegistry.cpp
@@ -2775,7 +2775,7 @@ void FPyWrapperTypeRegistry::GenerateStubCodeForWrappedTypes(const EPyOnlineDocs
 
 				if (MetaType)
 				{
-					FSourceCodeNavigation::FindModulePath(MetaType->GetTypedOuter<UPackage>(), ModulePath);
+					FSourceCodeNavigation::FindModulePath(MetaType->GetPackage(), ModulePath);
 				}
 
 				if (!ModulePath.IsEmpty())
diff --git a/Engine/Plugins/MovieScene/MovieRenderPipeline/Source/MovieRenderPipelineCore/Private/MoviePipelineUtils.cpp b/Engine/Plugins/MovieScene/MovieRenderPipeline/Source/MovieRenderPipelineCore/Private/MoviePipelineUtils.cpp
--- a/Engine/Plugins/MovieScene/MovieRenderPipeline/Source/MovieRenderPipelineCore/Private/MoviePipelineUtils.cpp
+++ b/Engine/Plugins/MovieScene/MovieRenderPipeline/Source/MovieRenderPipelineCore/Private/MoviePipelineUtils.cpp
@@ -256,7 +256,7 @@ namespace MoviePipeline
 					Node->bOriginalMovieSceneReadOnly = Node->MovieScene->IsReadOnly();
 					Node->bOriginalMovieScenePlaybackRangeLocked = Node->MovieScene->IsPlaybackRangeLocked();
 #endif
-					if (UPackage* OwningPackage = Node->MovieScene->GetTypedOuter<UPackage>())
+					if (UPackage* OwningPackage = Node->MovieScene->GetPackage())
 					{
 						Node->bOriginalMovieScenePackageDirty = OwningPackage->IsDirty();
 					}
@@ -332,7 +332,7 @@ namespace MoviePipeline
 #endif
 					Node->MovieScene->MarkAsChanged();
 
-					if (UPackage* OwningPackage = Node->MovieScene->GetTypedOuter<UPackage>())
+					if (UPackage* OwningPackage = Node->MovieScene->GetPackage())
 					{
 						OwningPackage->SetDirtyFlag(Node->bOriginalMovieScenePackageDirty);
 					}
diff --git a/Engine/Plugins/VirtualProduction/RemoteControl/Source/RemoteControl/Private/RemoteControlBinding.cpp b/Engine/Plugins/VirtualProduction/RemoteControl/Source/RemoteControl/Private/RemoteControlBinding.cpp
--- a/Engine/Plugins/VirtualProduction/RemoteControl/Source/RemoteControl/Private/RemoteControlBinding.cpp
+++ b/Engine/Plugins/VirtualProduction/RemoteControl/Source/RemoteControl/Private/RemoteControlBinding.cpp
@@ -31,7 +31,7 @@ namespace UE::RemoteControlBinding
 		return InObject
 			&& !InObject->IsA<UBillboardComponent>()
 			&& InObject->GetName().Find(TEXT("TRASH_")) == INDEX_NONE
-			&& InObject->GetTypedOuter<UPackage>() != GetTransientPackage()
+			&& InObject->GetPackage() != GetTransientPackage()
 			&& InObject->GetWorld() == PresetWorld;
 	}
 
diff --git a/Engine/Source/Developer/AssetTools/Private/AssetTypeActions/AssetTypeActions_World.cpp b/Engine/Source/Developer/AssetTools/Private/AssetTypeActions/AssetTypeActions_World.cpp
--- a/Engine/Source/Developer/AssetTools/Private/AssetTypeActions/AssetTypeActions_World.cpp
+++ b/Engine/Source/Developer/AssetTools/Private/AssetTypeActions/AssetTypeActions_World.cpp
@@ -37,7 +37,7 @@ void FAssetTypeActions_World::OpenAssetEditor( const TArray<UObject*>& InObjects
 	{
 		UWorld* World = Cast<UWorld>(*ObjIt);
 		if (World != nullptr && 
-			ensureMsgf(World->GetTypedOuter<UPackage>(), TEXT("World(%s) is not in a package and cannot be opened"), *World->GetFullName()) && 
+			ensureMsgf(World->GetPackage(), TEXT("World(%s) is not in a package and cannot be opened"), *World->GetFullName()) && 
 			ensureMsgf(!World->GetPackage()->HasAnyPackageFlags(PKG_NewlyCreated), TEXT("World(%s) is unsaved and cannot be opened")))
 		{
 			const FString FileToOpen = FPackageName::LongPackageNameToFilename(World->GetOutermost()->GetName(), FPackageName::GetMapPackageExtension());
diff --git a/Engine/Source/Editor/LandscapeEditor/Private/LandscapeEdMode.cpp b/Engine/Source/Editor/LandscapeEditor/Private/LandscapeEdMode.cpp
--- a/Engine/Source/Editor/LandscapeEditor/Private/LandscapeEdMode.cpp
+++ b/Engine/Source/Editor/LandscapeEditor/Private/LandscapeEdMode.cpp
@@ -4245,7 +4245,7 @@ void FEdModeLandscape::UpdateBrushList()
 	for (TObjectIterator<ALandscapeBlueprintBrushBase> BrushIt(RF_Transient|RF_ClassDefaultObject|RF_ArchetypeObject, true, EInternalObjectFlags::Garbage); BrushIt; ++BrushIt)
 	{
 		ALandscapeBlueprintBrushBase* Brush = *BrushIt;
-		if (Brush->GetTypedOuter<UPackage>() != GetTransientPackage())
+		if (Brush->GetPackage() != GetTransientPackage())
 		{
 			BrushList.Add(Brush);
 		}
@@ -4261,7 +4261,7 @@ void FEdModeLandscape::OnLevelActorAdded(AActor* InActor)
 	}
 
 	ALandscapeBlueprintBrushBase* Brush = Cast<ALandscapeBlueprintBrushBase>(InActor);
-	if (Brush && Brush->GetTypedOuter<UPackage>() != GetTransientPackage())
+	if (Brush && Brush->GetPackage() != GetTransientPackage())
 	{
 		if (!GIsReinstancing)
 		{
@@ -4280,7 +4280,7 @@ void FEdModeLandscape::OnLevelActorRemoved(AActor* InActor)
 	}
 
 	ALandscapeBlueprintBrushBase* Brush = Cast<ALandscapeBlueprintBrushBase>(InActor);
-	if (Brush && Brush->GetTypedOuter<UPackage>() != GetTransientPackage())
+	if (Brush && Brush->GetPackage() != GetTransientPackage())
 	{
 		UpdateBrushList();
 		RefreshDetailPanel();
diff --git a/Engine/Source/Editor/UnrealEd/Private/Kismet2/KismetDebugUtilities.cpp b/Engine/Source/Editor/UnrealEd/Private/Kismet2/KismetDebugUtilities.cpp
--- a/Engine/Source/Editor/UnrealEd/Private/Kismet2/KismetDebugUtilities.cpp
+++ b/Engine/Source/Editor/UnrealEd/Private/Kismet2/KismetDebugUtilities.cpp
@@ -1217,7 +1217,7 @@ void FKismetDebugUtilities::RestoreBreakpointsOnLoad(const UBlueprint* Blueprint
 		[Blueprint](const FBlueprintBreakpoint& Breakpoint)
 		{
 			const UEdGraphNode* const Location = Breakpoint.GetLocation();
-			return !Location || Location->GetTypedOuter<UPackage>()->GetPersistentGuid() != Blueprint->GetTypedOuter<UPackage>()->GetPersistentGuid();
+			return !Location || Location->GetPackage()->GetPersistentGuid() != Blueprint->GetPackage()->GetPersistentGuid();
 		}
 	);
 #endif
diff --git a/Engine/Source/Editor/UnrealEd/Private/SourceCodeNavigation.cpp b/Engine/Source/Editor/UnrealEd/Private/SourceCodeNavigation.cpp
--- a/Engine/Source/Editor/UnrealEd/Private/SourceCodeNavigation.cpp
+++ b/Engine/Source/Editor/UnrealEd/Private/SourceCodeNavigation.cpp
@@ -2007,7 +2007,7 @@ bool FSourceCodeNavigation::FindClassHeaderPath( const UField* InField, FString
 	}
 
 	// Get the class package, and skip past the "/Script/" portion to get the module name
-	UPackage* ModulePackage = InField->GetTypedOuter<UPackage>();
+	UPackage* ModulePackage = InField->GetPackage();
 
 	// Find the base path for the module
 	FString ModuleBasePath;
@@ -2032,7 +2032,7 @@ bool FSourceCodeNavigation::FindClassSourcePath( const UField* InField, FString
 	}
 
 	// Get the class package, and skip past the "/Script/" portion to get the module name
-	UPackage *ModulePackage = InField->GetTypedOuter<UPackage>();
+	UPackage *ModulePackage = InField->GetPackage();
 
 	// Find the base path for the module
 	FString ModuleBasePath;
diff --git a/Engine/Source/Runtime/CoreUObject/Private/UObject/UObjectBaseUtility.cpp b/Engine/Source/Runtime/CoreUObject/Private/UObject/UObjectBaseUtility.cpp
--- a/Engine/Source/Runtime/CoreUObject/Private/UObject/UObjectBaseUtility.cpp
+++ b/Engine/Source/Runtime/CoreUObject/Private/UObject/UObjectBaseUtility.cpp
@@ -371,6 +371,8 @@ bool UObjectBaseUtility::IsTemplate( EObjectFlags TemplateTypes ) const
  */
 UObject* UObjectBaseUtility::GetTypedOuter(UClass* Target) const
 {
+	ensureMsgf(Target != UPackage::StaticClass(), TEXT("Calling GetTypedOuter to retrieve a package is now invalid, you should use GetPackage() instead."));
+
 	UObject* Result = NULL;
 	for ( UObject* NextOuter = GetOuter(); Result == NULL && NextOuter != NULL; NextOuter = NextOuter->GetOuter() )
 	{
