LandscapeSpline: While one Control Point is selected on a LandscapeSplineActor, "Ctrl+Click" a Control Point from a different LandscapeSplineActor will merge both actors together.
- Moves the second spline into the first actor and adds a segment between points
- Deletes the second actor

#rnx

diff --git a/Engine/Source/Editor/LandscapeEditor/Private/LandscapeEdModeSplineTools.cpp b/Engine/Source/Editor/LandscapeEditor/Private/LandscapeEdModeSplineTools.cpp
--- a/Engine/Source/Editor/LandscapeEditor/Private/LandscapeEdModeSplineTools.cpp
+++ b/Engine/Source/Editor/LandscapeEditor/Private/LandscapeEdModeSplineTools.cpp
@@ -1339,8 +1339,42 @@ public:
 
 					if (ClickedControlPoint != nullptr)
 					{
-						FScopedTransaction Transaction(LOCTEXT("LandscapeSpline_AddSegment", "Add Landscape Spline Segment"));
+						// Merge Spline into the same actor if a single Control Point is currently selected and the ClickedControlPoint is from a different owner
+						if (SelectedSplineControlPoints.Num() == 1)
+						{
+							ULandscapeSplineControlPoint* SourceControlPoint = *SelectedSplineControlPoints.CreateIterator();
+
+							ULandscapeSplinesComponent* SourceComponent = SourceControlPoint->GetOuterULandscapeSplinesComponent();
+							ALandscapeSplineActor* SourceSplineActor = SourceComponent ? Cast<ALandscapeSplineActor>(SourceComponent->GetOuter()) : nullptr;
+
+							ULandscapeSplinesComponent* ClickedComponent = ClickedControlPoint->GetOuterULandscapeSplinesComponent();
+							ALandscapeSplineActor* ClickedSplineActor = ClickedComponent ? Cast<ALandscapeSplineActor>(ClickedComponent->GetOuter()) : nullptr;
+
+							if (SourceSplineActor && ClickedSplineActor && SourceSplineActor != ClickedSplineActor)
+							{
+								if (SourceSplineActor->GetLandscapeGuid() != ClickedSplineActor->GetLandscapeGuid())
+								{
+									UE_LOG(LogLandscapeEdMode, Warning, TEXT("Can't merge LandscapeSplineActors belonging to different Landscapes"));
+									return true;
+								}
+
+								FScopedTransaction Transaction(LOCTEXT("LandscapeSpline_MergeSpline", "Merge Spline"));
+								
+								ULandscapeInfo* SourceLandscapeInfo = SourceSplineActor->GetLandscapeInfo();
+								check(SourceLandscapeInfo);
 
+								SourceLandscapeInfo->MoveSpline(ClickedControlPoint, SourceSplineActor);
+								AddSegment(SourceControlPoint, ClickedControlPoint, true, true);
+
+								// Moving the spline should leave us with an empty actor that we can delete
+								check(ClickedComponent->GetControlPoints().Num() == 0);
+								ClickedSplineActor->GetWorld()->EditorDestroyActor(ClickedSplineActor, true);
+								
+								return true;
+							}
+						}
+
+						FScopedTransaction Transaction(LOCTEXT("LandscapeSpline_AddSegment", "Add Landscape Spline Segment"));
 						for (ULandscapeSplineControlPoint* ControlPoint : SelectedSplineControlPoints)
 						{
 							AddSegment(ControlPoint, ClickedControlPoint, bAutoRotateOnJoin, bAutoRotateOnJoin);
