* Disabled Layer Info object creation from the asset picker context menu in the Create and the Paint landscape panels.
* Prevented landscape layer info objects to be created when the default landscape layer info object in the project settings is invalid (i.e. it has a path but the asset is missing)

diff --git a/Engine/Source/Editor/LandscapeEditor/Private/LandscapeEditorDetailCustomization_NewLandscape.cpp b/Engine/Source/Editor/LandscapeEditor/Private/LandscapeEditorDetailCustomization_NewLandscape.cpp
--- a/Engine/Source/Editor/LandscapeEditor/Private/LandscapeEditorDetailCustomization_NewLandscape.cpp
+++ b/Engine/Source/Editor/LandscapeEditor/Private/LandscapeEditorDetailCustomization_NewLandscape.cpp
@@ -941,8 +941,7 @@ FReply FLandscapeEditorDetailCustomization_NewLandscape::OnCreateButtonClicked()
 		// Import doesn't fill in the LayerInfo for layers with no data, do that now
 		const TArray<FLandscapeImportLayer>& ImportLandscapeLayersList = UISettings->ImportLandscape_Layers;
 		const ULandscapeSettings* Settings = GetDefault<ULandscapeSettings>();
-		TSoftObjectPtr<ULandscapeLayerInfoObject> DefaultLayerInfoObject = Settings->GetDefaultLayerInfoObject();
-		const bool bIsDefaultLayerInfoObjectSet = !DefaultLayerInfoObject.IsNull();
+		TSoftObjectPtr<ULandscapeLayerInfoObject> DefaultLayerInfoObject = Settings->GetDefaultLayerInfoObject().LoadSynchronous();
 
 		for (int32 i = 0; i < ImportLandscapeLayersList.Num(); i++)
 		{
@@ -950,9 +949,9 @@ FReply FLandscapeEditorDetailCustomization_NewLandscape::OnCreateButtonClicked()
 			FName LayerName = ImportLandscapeLayersList[i].LayerName;
 
 			// If DefaultLayerInfoObject is set and LayerInfo does not exist, we will try to create the new LayerInfo by cloning DefaultLayerInfoObject. Except for VisibilityLayer which doesn't require an asset.
-			if (bIsDefaultLayerInfoObjectSet && (LayerInfo == nullptr) && (LayerName != ALandscapeProxy::VisibilityLayer->LayerName))
+			if (DefaultLayerInfoObject.IsValid() && (LayerInfo == nullptr) && (LayerName != ALandscapeProxy::VisibilityLayer->LayerName))
 			{
-				LayerInfo = Landscape->CreateLayerInfo(*LayerName.ToString(), DefaultLayerInfoObject.LoadSynchronous());
+				LayerInfo = Landscape->CreateLayerInfo(*LayerName.ToString(), DefaultLayerInfoObject.Get());
 
 				if (LayerInfo != nullptr)
 				{
diff --git a/Engine/Source/Editor/LandscapeEditor/Private/LandscapeEditorDetailCustomization_TargetLayers.cpp b/Engine/Source/Editor/LandscapeEditor/Private/LandscapeEditorDetailCustomization_TargetLayers.cpp
--- a/Engine/Source/Editor/LandscapeEditor/Private/LandscapeEditorDetailCustomization_TargetLayers.cpp
+++ b/Engine/Source/Editor/LandscapeEditor/Private/LandscapeEditorDetailCustomization_TargetLayers.cpp
@@ -645,6 +645,7 @@ TSharedPtr<SWidget> FLandscapeEditorCustomNodeBuilder_TargetLayers::GenerateRow(
 							.AllowedClass(ULandscapeLayerInfoObject::StaticClass())
 							.OnObjectChanged_Static(&FLandscapeEditorCustomNodeBuilder_TargetLayers::OnTargetLayerSetObject, Target)
 							.OnShouldFilterAsset_Static(&FLandscapeEditorCustomNodeBuilder_TargetLayers::ShouldFilterLayerInfo, Target->LayerName)
+							.AllowCreate(false)
 							.AllowClear(false)
 						]
 						+ SHorizontalBox::Slot()
diff --git a/Engine/Source/Runtime/Landscape/Classes/LandscapeProxy.h b/Engine/Source/Runtime/Landscape/Classes/LandscapeProxy.h
--- a/Engine/Source/Runtime/Landscape/Classes/LandscapeProxy.h
+++ b/Engine/Source/Runtime/Landscape/Classes/LandscapeProxy.h
@@ -151,7 +151,7 @@ struct FLandscapeImportLayerInfo
 	UPROPERTY(Category="Import", VisibleAnywhere)
 	FName LayerName;
 
-	UPROPERTY(Category="Import", EditAnywhere)
+	UPROPERTY(Category="Import", EditAnywhere, meta=(NoCreate))
 	TObjectPtr<ULandscapeLayerInfoObject> LayerInfo;
 
 	UPROPERTY(Category="Import", EditAnywhere, meta=(DisplayName="Layer File"))
