Fix failed check (wrong assumption) when merging LandscapeSplineActors

#rnx

diff --git a/Engine/Source/Editor/LandscapeEditor/Private/LandscapeEdModeSplineTools.cpp b/Engine/Source/Editor/LandscapeEditor/Private/LandscapeEdModeSplineTools.cpp
--- a/Engine/Source/Editor/LandscapeEditor/Private/LandscapeEdModeSplineTools.cpp
+++ b/Engine/Source/Editor/LandscapeEditor/Private/LandscapeEdModeSplineTools.cpp
@@ -1367,8 +1367,10 @@ public:
 								AddSegment(SourceControlPoint, ClickedControlPoint, true, true);
 
 								// Moving the spline should leave us with an empty actor that we can delete
-								check(ClickedComponent->GetControlPoints().Num() == 0);
-								ClickedSplineActor->GetWorld()->EditorDestroyActor(ClickedSplineActor, true);
+								if(ClickedComponent->GetControlPoints().Num() == 0)
+								{
+									ClickedSplineActor->GetWorld()->EditorDestroyActor(ClickedSplineActor, true);
+								}
 								
 								return true;
 							}
