Fixed crash when creating ULandscapeEditorObject's CDO. We don't really care whether the default alpha texture's source data is readable or not at this point.

#rnx

diff --git a/Engine/Source/Editor/LandscapeEditor/Private/LandscapeEditorObject.cpp b/Engine/Source/Editor/LandscapeEditor/Private/LandscapeEditorObject.cpp
--- a/Engine/Source/Editor/LandscapeEditor/Private/LandscapeEditorObject.cpp
+++ b/Engine/Source/Editor/LandscapeEditor/Private/LandscapeEditorObject.cpp
@@ -126,7 +126,10 @@ ULandscapeEditorObject::ULandscapeEditorObject(const FObjectInitializer& ObjectI
 	};
 	static FConstructorStatics ConstructorStatics;
 
-	SetAlphaTexture(ConstructorStatics.AlphaTexture.Object, AlphaTextureChannel);
+	if (!IsTemplate())
+	{
+		SetAlphaTexture(ConstructorStatics.AlphaTexture.Object, AlphaTextureChannel);
+	}
 }
 
 void ULandscapeEditorObject::PostEditChangeProperty(FPropertyChangedEvent& PropertyChangedEvent)
