Fix for crash on Undo after selecting Landscape

diff --git a/Engine/Source/Editor/LandscapeEditor/Private/LandscapeEdMode.cpp b/Engine/Source/Editor/LandscapeEditor/Private/LandscapeEdMode.cpp
--- a/Engine/Source/Editor/LandscapeEditor/Private/LandscapeEdMode.cpp
+++ b/Engine/Source/Editor/LandscapeEditor/Private/LandscapeEdMode.cpp
@@ -2939,11 +2939,12 @@ bool FEdModeLandscape::Select(AActor* InActor, bool bInSelected)
 	if (InActor->IsA<ALandscapeProxy>() && bInSelected)
 	{
 		ALandscapeProxy* Landscape = CastChecked<ALandscapeProxy>(InActor);
-		if (Landscape->GetLandscapeInfo()->SupportsLandscapeEditing())
+		ULandscapeInfo* LandscapeInfo = Landscape->GetLandscapeInfo();
+		if (LandscapeInfo && LandscapeInfo->SupportsLandscapeEditing())
 		{
-			if (CurrentToolTarget.LandscapeInfo != Landscape->GetLandscapeInfo())
+			if (CurrentToolTarget.LandscapeInfo != LandscapeInfo)
 			{
-				SetLandscapeInfo(Landscape->GetLandscapeInfo());
+				SetLandscapeInfo(LandscapeInfo);
 
 				// If we were in "New Landscape" mode and we select a landscape then switch to editing mode
 				if (NewLandscapePreviewMode != ENewLandscapePreviewMode::None)
