Fix for editor freezing (going into near-infinite loop) when editing Landscape in Orthographic view mode

diff --git a/Engine/Source/Editor/LandscapeEditor/Private/LandscapeEdMode.cpp b/Engine/Source/Editor/LandscapeEditor/Private/LandscapeEdMode.cpp
--- a/Engine/Source/Editor/LandscapeEditor/Private/LandscapeEdMode.cpp
+++ b/Engine/Source/Editor/LandscapeEditor/Private/LandscapeEdMode.cpp
@@ -1214,8 +1214,15 @@ bool FEdModeLandscape::LandscapeTrace(const FVector& InRayOrigin, const FVector&
 
 			UE_VLOG_LOCATION(World, LogLandscapeEdMode, VeryVerbose,  PointOnPlane, 10.0,  FColor(100,100,255), TEXT("landscape:point-on-plane"));
 			
+			if (FMath::Abs(FVector::DotProduct(InDirection, FVector::ZAxisVector)) < SMALL_NUMBER)
+			{
+				// the ray and plane are nearly parallel, and won't intersect (or if they do it will be wildly far away)
+				return false;
+			}
+
 			const FPlane Plane( PointOnPlane, FVector::ZAxisVector);
 			FVector EstimatedHitLocation = FMath::RayPlaneIntersection(Start, InDirection, Plane);
+			check(!EstimatedHitLocation.ContainsNaN());
 
 			UE_VLOG_LOCATION(World, LogLandscapeEdMode, VeryVerbose,  EstimatedHitLocation, 10.0,  FColor(100,100,255), TEXT("landscape:estimated-hit-location"));
 			
