Fix Landscape Add components tool so it does not add tiles disconnected from the existing tiles

diff --git a/Engine/Source/Editor/LandscapeEditor/Private/LandscapeEdMode.cpp b/Engine/Source/Editor/LandscapeEditor/Private/LandscapeEdMode.cpp
--- a/Engine/Source/Editor/LandscapeEditor/Private/LandscapeEdMode.cpp
+++ b/Engine/Source/Editor/LandscapeEditor/Private/LandscapeEdMode.cpp
@@ -1166,71 +1166,74 @@ bool FEdModeLandscape::LandscapeTrace(const FVector& InRayOrigin, const FVector&
 
 	UE_VLOG_SEGMENT_THICK(World, LogLandscapeEdMode, VeryVerbose, InRayOrigin, InRayOrigin + 20000.0f * InDirection, FColor(255,100,100), 4, TEXT("landscape:ray-miss"));
 		
-	// If there is no landscape directly under the mouse search for a landscape collision
-	// under the shape of the brush.
-	FCollisionShape SphereShape;
-	SphereShape.SetSphere(UISettings->GetCurrentToolBrushRadius());
-	if ( World->SweepMultiByObjectType(Results, Start, End, FQuat::Identity, FCollisionObjectQueryParams(ECollisionChannel::ECC_Visibility), SphereShape, FCollisionQueryParams(SCENE_QUERY_STAT(LandscapeTrace), true)))
+	if (CurrentTool->UseSphereTrace())
 	{
-		if (FProcessLandscapeTraceHitsResult SweepProcessResult; ProcessLandscapeTraceHits(Results, SweepProcessResult))
+		// If there is no landscape directly under the mouse search for a landscape collision
+		// under the shape of the brush.
+		FCollisionShape SphereShape;
+		SphereShape.SetSphere(UISettings->GetCurrentToolBrushRadius());
+		if (World->SweepMultiByObjectType(Results, Start, End, FQuat::Identity, FCollisionObjectQueryParams(ECollisionChannel::ECC_Visibility), SphereShape, FCollisionQueryParams(SCENE_QUERY_STAT(LandscapeTrace), true)))
 		{
-			UE_VLOG_LOCATION(World, LogLandscapeEdMode, VeryVerbose,  SweepProcessResult.HitLocation, UISettings->GetCurrentToolBrushRadius(),  FColor(255,100,100), TEXT("landscape:sweep-hit-location"));
-			FSphere HitSphere(SweepProcessResult.HitLocation, UISettings->GetCurrentToolBrushRadius());
-
-			float MeanHeight = 0.0f;
-			int32 Count = 0;
-			const int32 NumHeightSamples = 16;
-			
-			for (int32 Y = 0; Y < NumHeightSamples; ++Y)
+			if (FProcessLandscapeTraceHitsResult SweepProcessResult; ProcessLandscapeTraceHits(Results, SweepProcessResult))
 			{
-				float HY = (Y / (float) NumHeightSamples) * HitSphere.W * 2.0f +  HitSphere.Center.Y - HitSphere.W;
-				for (int32 X = 0; X < NumHeightSamples; ++X)
+				UE_VLOG_LOCATION(World, LogLandscapeEdMode, VeryVerbose, SweepProcessResult.HitLocation, UISettings->GetCurrentToolBrushRadius(), FColor(255, 100, 100), TEXT("landscape:sweep-hit-location"));
+				FSphere HitSphere(SweepProcessResult.HitLocation, UISettings->GetCurrentToolBrushRadius());
+
+				float MeanHeight = 0.0f;
+				int32 Count = 0;
+				const int32 NumHeightSamples = 16;
+
+				for (int32 Y = 0; Y < NumHeightSamples; ++Y)
 				{
-					float HX = (X / (float) NumHeightSamples) * HitSphere.W * 2.0f +  HitSphere.Center.X - HitSphere.W;
+					float HY = (Y / (float)NumHeightSamples) * HitSphere.W * 2.0f + HitSphere.Center.Y - HitSphere.W;
+					for (int32 X = 0; X < NumHeightSamples; ++X)
+					{
+						float HX = (X / (float)NumHeightSamples) * HitSphere.W * 2.0f + HitSphere.Center.X - HitSphere.W;
 
-					FVector HeightSampleLocation(HX, HY, HitSphere.Center.Z);
+						FVector HeightSampleLocation(HX, HY, HitSphere.Center.Z);
 
-					TOptional<float> Height = SweepProcessResult.LandscapeProxy->GetHeightAtLocation(HeightSampleLocation, EHeightfieldSource::Editor);
+						TOptional<float> Height = SweepProcessResult.LandscapeProxy->GetHeightAtLocation(HeightSampleLocation, EHeightfieldSource::Editor);
 
-					if (!Height.IsSet())
-					{
-						continue;
+						if (!Height.IsSet())
+						{
+							continue;
+						}
+
+						UE_VLOG_LOCATION(World, LogLandscapeEdMode, VeryVerbose, FVector(HX, HY, Height.GetValue()), 2.0f, FColor(100, 100, 255), TEXT(""));
+
+						MeanHeight += Height.GetValue();
+						Count++;
 					}
+				}
+
+				FVector PointOnPlane(HitSphere.Center.X, HitSphere.Center.Y, HitSphere.Center.Z);
 
-					UE_VLOG_LOCATION(World, LogLandscapeEdMode, VeryVerbose,  FVector(HX, HY, Height.GetValue()), 2.0f,  FColor(100,100,255), TEXT(""));
-					
-					MeanHeight += Height.GetValue();
-					Count++;
+				if (Count > 0)
+				{
+					MeanHeight /= (float)Count;
+					PointOnPlane.Z = MeanHeight;
 				}
-			}
 
-			FVector PointOnPlane ( HitSphere.Center.X, HitSphere.Center.Y, HitSphere.Center.Z );
-			
-			if  ( Count > 0)
-			{
-				MeanHeight /= (float) Count;
-				PointOnPlane.Z = MeanHeight;
-			}
+				UE_VLOG_LOCATION(World, LogLandscapeEdMode, VeryVerbose, PointOnPlane, 10.0, FColor(100, 100, 255), TEXT("landscape:point-on-plane"));
 
-			UE_VLOG_LOCATION(World, LogLandscapeEdMode, VeryVerbose,  PointOnPlane, 10.0,  FColor(100,100,255), TEXT("landscape:point-on-plane"));
-			
-			if (FMath::Abs(FVector::DotProduct(InDirection, FVector::ZAxisVector)) < SMALL_NUMBER)
-			{
-				// the ray and plane are nearly parallel, and won't intersect (or if they do it will be wildly far away)
-				return false;
-			}
+				if (FMath::Abs(FVector::DotProduct(InDirection, FVector::ZAxisVector)) < SMALL_NUMBER)
+				{
+					// the ray and plane are nearly parallel, and won't intersect (or if they do it will be wildly far away)
+					return false;
+				}
 
-			const FPlane Plane( PointOnPlane, FVector::ZAxisVector);
-			FVector EstimatedHitLocation = FMath::RayPlaneIntersection(Start, InDirection, Plane);
-			check(!EstimatedHitLocation.ContainsNaN());
+				const FPlane Plane(PointOnPlane, FVector::ZAxisVector);
+				FVector EstimatedHitLocation = FMath::RayPlaneIntersection(Start, InDirection, Plane);
+				check(!EstimatedHitLocation.ContainsNaN());
 
-			UE_VLOG_LOCATION(World, LogLandscapeEdMode, VeryVerbose,  EstimatedHitLocation, 10.0,  FColor(100,100,255), TEXT("landscape:estimated-hit-location"));
-			
-			OutHitLocation = SweepProcessResult.LandscapeProxy->LandscapeActorToWorld().InverseTransformPosition(EstimatedHitLocation);
-			return true;
+				UE_VLOG_LOCATION(World, LogLandscapeEdMode, VeryVerbose, EstimatedHitLocation, 10.0, FColor(100, 100, 255), TEXT("landscape:estimated-hit-location"));
+
+				OutHitLocation = SweepProcessResult.LandscapeProxy->LandscapeActorToWorld().InverseTransformPosition(EstimatedHitLocation);
+				return true;
+			}
 		}
 	}
-	
+
 	return false;
 }
 
diff --git a/Engine/Source/Editor/LandscapeEditor/Private/LandscapeEdModeComponentTools.cpp b/Engine/Source/Editor/LandscapeEditor/Private/LandscapeEdModeComponentTools.cpp
--- a/Engine/Source/Editor/LandscapeEditor/Private/LandscapeEdModeComponentTools.cpp
+++ b/Engine/Source/Editor/LandscapeEditor/Private/LandscapeEdModeComponentTools.cpp
@@ -799,6 +799,9 @@ public:
 	virtual void SetEditRenderType() override { GLandscapeEditRenderMode = ELandscapeEditRenderMode::None | (GLandscapeEditRenderMode & ELandscapeEditRenderMode::BitMaskForMask); }
 	virtual bool SupportsMask() override { return false; }
 
+	// Sphere traces can result in components being added at a good distance from any neighboring components, because it intersects against a virtual plane
+	virtual bool UseSphereTrace() override { return false; }
+
 	virtual bool CanToolBeActivated() const override
 	{ 
 		return FLandscapeToolBase<FLandscapeToolStrokeAddComponent>::CanToolBeActivated() && bIsToolActionResolutionCompliant;
diff --git a/Engine/Source/Editor/LandscapeEditor/Public/LandscapeToolInterface.h b/Engine/Source/Editor/LandscapeEditor/Public/LandscapeToolInterface.h
--- a/Engine/Source/Editor/LandscapeEditor/Public/LandscapeToolInterface.h
+++ b/Engine/Source/Editor/LandscapeEditor/Public/LandscapeToolInterface.h
@@ -212,6 +212,7 @@ public:
 	virtual void SetEditRenderType();
 	virtual void Render(const FSceneView* View, FViewport* Viewport, FPrimitiveDrawInterface* PDI) {}
 	virtual bool HitTrace(const FVector& TraceStart, const FVector& TraceEnd, FVector& OutHitLocation) { return false; }
+	virtual bool UseSphereTrace() { return true; }
 	virtual bool SupportsMask() { return true; }
 	virtual bool SupportsComponentSelection() { return false; }
 	virtual bool OverrideSelection() const { return false; }
