Overrides reset to default button for LandscapeMaterial when creating a new landscape. Now resets the value to the set default landscape material.

diff --git a/Engine/Source/Editor/LandscapeEditor/Private/LandscapeEditorDetailCustomization_NewLandscape.cpp b/Engine/Source/Editor/LandscapeEditor/Private/LandscapeEditorDetailCustomization_NewLandscape.cpp
--- a/Engine/Source/Editor/LandscapeEditor/Private/LandscapeEditorDetailCustomization_NewLandscape.cpp
+++ b/Engine/Source/Editor/LandscapeEditor/Private/LandscapeEditorDetailCustomization_NewLandscape.cpp
@@ -219,7 +219,14 @@ void FLandscapeEditorDetailCustomization_NewLandscape::CustomizeDetails(IDetailL
 	];
 
 	TSharedRef<IPropertyHandle> PropertyHandle_Material = DetailBuilder.GetProperty(GET_MEMBER_NAME_CHECKED(ULandscapeEditorObject, NewLandscape_Material));
-	NewLandscapeCategory.AddProperty(PropertyHandle_Material);
+	IDetailPropertyRow& MaterialPropertyRow = NewLandscapeCategory.AddProperty(PropertyHandle_Material);
+
+
+	FIsResetToDefaultVisible IsResetVisible = FIsResetToDefaultVisible::CreateSP(this, &FLandscapeEditorDetailCustomization_NewLandscape::ShouldShowResetMaterialToDefault);
+	FResetToDefaultHandler ResetHandler = FResetToDefaultHandler::CreateSP(this, &FLandscapeEditorDetailCustomization_NewLandscape::ResetMaterialToDefault);
+	FResetToDefaultOverride ResetOverride = FResetToDefaultOverride::Create(IsResetVisible, ResetHandler);
+
+	MaterialPropertyRow.OverrideResetToDefault(ResetOverride);
 
 	NewLandscapeCategory.AddCustomRow(LOCTEXT("LayersLabel", "Layers"))
 	.Visibility(TAttribute<EVisibility>(this, &FLandscapeEditorDetailCustomization_NewLandscape::GetMaterialTipVisibility))
@@ -1232,4 +1239,29 @@ EVisibility FLandscapeEditorDetailCustomization_NewLandscape::GetMaterialTipVisi
 	return EVisibility::Collapsed;
 }
 
+void FLandscapeEditorDetailCustomization_NewLandscape::ResetMaterialToDefault(TSharedPtr<IPropertyHandle> InPropertyHandle)
+{
+	const ULandscapeSettings* Settings = GetDefault<ULandscapeSettings>();
+	TSoftObjectPtr<UMaterialInterface> DefaultMaterial = Settings->GetDefaultLandscapeMaterial();
+	UObject* ValueToSet = nullptr;
+
+	if (!DefaultMaterial.IsNull())
+	{
+		ValueToSet = DefaultMaterial.LoadSynchronous();
+	}
+
+	InPropertyHandle->SetValue(ValueToSet);
+}
+
+bool FLandscapeEditorDetailCustomization_NewLandscape::ShouldShowResetMaterialToDefault(TSharedPtr<IPropertyHandle> InPropertyHandle)
+{
+	const ULandscapeSettings* Settings = GetDefault<ULandscapeSettings>();
+	TSoftObjectPtr<UMaterialInterface> DefaultMaterial = Settings->GetDefaultLandscapeMaterial();
+
+	UObject* Object;
+	InPropertyHandle->GetValue(Object);
+
+	return DefaultMaterial.Get() != Object;
+}
+
 #undef LOCTEXT_NAMESPACE
diff --git a/Engine/Source/Editor/LandscapeEditor/Private/LandscapeEditorDetailCustomization_NewLandscape.h b/Engine/Source/Editor/LandscapeEditor/Private/LandscapeEditorDetailCustomization_NewLandscape.h
--- a/Engine/Source/Editor/LandscapeEditor/Private/LandscapeEditorDetailCustomization_NewLandscape.h
+++ b/Engine/Source/Editor/LandscapeEditor/Private/LandscapeEditorDetailCustomization_NewLandscape.h
@@ -89,6 +89,9 @@ public:
 	// Import layers
 	EVisibility GetMaterialTipVisibility() const;
 
+	void ResetMaterialToDefault(TSharedPtr<IPropertyHandle> InPropertyHandle);
+	bool ShouldShowResetMaterialToDefault(TSharedPtr<IPropertyHandle> InPropertyHandle);
+
 protected:
 	bool bUsingSlider;
 };
