Landscape :  Don't change Tool selection on LandscapeEd undo

Prevent Select & Sculpt modes being selected on LandscapeEdMode Undo hander.

diff --git a/Engine/Source/Editor/LandscapeEditor/Private/LandscapeEdMode.cpp b/Engine/Source/Editor/LandscapeEditor/Private/LandscapeEdMode.cpp
--- a/Engine/Source/Editor/LandscapeEditor/Private/LandscapeEdMode.cpp
+++ b/Engine/Source/Editor/LandscapeEditor/Private/LandscapeEdMode.cpp
@@ -479,8 +479,8 @@ void FEdModeLandscape::Enter()
 	UpdateLandscapeList();
 	UpdateBrushList();
 
-	OnWorldChangeDelegateHandle                 = FEditorSupportDelegates::WorldChange.AddRaw(this, &FEdModeLandscape::HandleLevelsChanged, true);
-	OnLevelsChangedDelegateHandle				= GetWorld()->OnLevelsChanged().AddRaw(this, &FEdModeLandscape::HandleLevelsChanged, true);
+	OnWorldChangeDelegateHandle                 = FEditorSupportDelegates::WorldChange.AddRaw(this, &FEdModeLandscape::HandleLevelsChanged);
+	OnLevelsChangedDelegateHandle				= GetWorld()->OnLevelsChanged().AddRaw(this, &FEdModeLandscape::HandleLevelsChanged);
 	OnMaterialCompilationFinishedDelegateHandle = UMaterial::OnMaterialCompilationFinished().AddRaw(this, &FEdModeLandscape::OnMaterialCompilationFinished);
 
 	if (CurrentToolTarget.LandscapeInfo.IsValid())
@@ -2786,16 +2786,15 @@ void FEdModeLandscape::MoveTargetLayerDisplayOrder(int32 IndexToMove, int32 Inde
 
 FEdModeLandscape::FTargetsListUpdated FEdModeLandscape::TargetsListUpdated;
 
-void FEdModeLandscape::HandleLevelsChanged(bool ShouldExitMode)
+void FEdModeLandscape::HandleLevelsChanged()
 {
-	bool bHadLandscape = (NewLandscapePreviewMode == ENewLandscapePreviewMode::None);
-
 	UpdateLandscapeList();
 	UpdateTargetList();
 	UpdateBrushList();
 
 	// if the Landscape is deleted then close the landscape editor
-	if (ShouldExitMode && bHadLandscape && CurrentToolTarget.LandscapeInfo == nullptr)
+	const bool bHadLandscape = (NewLandscapePreviewMode == ENewLandscapePreviewMode::None);
+	if (bHadLandscape && CurrentToolTarget.LandscapeInfo == nullptr)
 	{
 		RequestDeletion();
 	}
@@ -3031,7 +3030,9 @@ void FEdModeLandscape::ActorMoveNotify()
 
 void FEdModeLandscape::PostUndo()
 {
-	HandleLevelsChanged(false);
+	UpdateLandscapeList();
+	UpdateTargetList();
+	UpdateBrushList();
 }
 
 /** Forces all level editor viewports to realtime mode */
diff --git a/Engine/Source/Editor/LandscapeEditor/Private/LandscapeEdMode.h b/Engine/Source/Editor/LandscapeEditor/Private/LandscapeEdMode.h
--- a/Engine/Source/Editor/LandscapeEditor/Private/LandscapeEdMode.h
+++ b/Engine/Source/Editor/LandscapeEditor/Private/LandscapeEdMode.h
@@ -589,7 +589,7 @@ public:
 	void OnPreSaveWorld(class UWorld* InWorld, FObjectPreSaveContext ObjectSaveContext);
 
 	/** Handle notification that visible levels may have changed and we should update the editable landscapes list */
-	void HandleLevelsChanged(bool ShouldExitMode);
+	void HandleLevelsChanged();
 
 	void OnMaterialCompilationFinished(UMaterialInterface* MaterialInterface);
 
