Fixes Landscape ImportTool resolution computation. Should be 0 and not depending on the imported landscape resolution since the tool will import the data in the existing texture and not create a new one.

diff --git a/Engine/Source/Editor/LandscapeEditor/Private/LandscapeEdModeComponentTools.cpp b/Engine/Source/Editor/LandscapeEditor/Private/LandscapeEdModeComponentTools.cpp
--- a/Engine/Source/Editor/LandscapeEditor/Private/LandscapeEdModeComponentTools.cpp
+++ b/Engine/Source/Editor/LandscapeEditor/Private/LandscapeEdModeComponentTools.cpp
@@ -2613,22 +2613,6 @@ public:
 			}
 		}
 	}
-
-	virtual int32 GetToolActionResolutionDelta() const override
-	{
-		if (EdMode != nullptr)
-		{
-			int32 ImportWidth = EdMode->UISettings->ImportLandscape_Width;
-			int32 ImportHeight = EdMode->UISettings->ImportLandscape_Height;
-
-			ImportWidth = (ImportWidth > 0) ? ImportWidth : 1;
-			ImportHeight = (ImportHeight > 0) ? ImportHeight : 1;
-
-			return ImportWidth * ImportHeight;
-		}
-
-		return 0;
-	}
 };
 
 //////////////////////////////////////////////////////////////////////////
