Landscape : Fix height pack function.

diff --git a/Engine/Shaders/Private/Landscape/LandscapeCommon.ush b/Engine/Shaders/Private/Landscape/LandscapeCommon.ush
--- a/Engine/Shaders/Private/Landscape/LandscapeCommon.ush
+++ b/Engine/Shaders/Private/Landscape/LandscapeCommon.ush
@@ -5,7 +5,7 @@
 // Height is packed from a [0, 2^16-1] value to two bytes, each of which we need to provide as a 0-1 float for writing 
 float2 PackHeight(float InHeight)
 {
-	uint Height = (uint) clamp(InHeight, 0.0f, 65535.0f);
+	uint Height = (uint) clamp(round(InHeight), 0.0f, 65535.0f);
 	return float2((float) (Height >> 8) / 255.0, (float) (Height & 0xFF) / 255.0);
 }
 
diff --git a/Engine/Source/Editor/LandscapeEditor/Private/LandscapeEdModePaintTools.cpp b/Engine/Source/Editor/LandscapeEditor/Private/LandscapeEdModePaintTools.cpp
--- a/Engine/Source/Editor/LandscapeEditor/Private/LandscapeEdModePaintTools.cpp
+++ b/Engine/Source/Editor/LandscapeEditor/Private/LandscapeEdModePaintTools.cpp
@@ -696,6 +696,8 @@ public:
 		TArray<typename ToolTarget::CacheClass::DataType> Data;
 		LayerDataCache.Initialize(this->LandscapeInfo, bCombinedLayerOperation);
 		LayerDataCache.Read(X1, Y1, X2, Y2, Data);
+		const TArray<typename ToolTarget::CacheClass::DataType> ReadData { Data };
+		
 		const float ToolStrength = FMath::Clamp<float>(UISettings->GetCurrentToolStrength() * Pressure, 0.0f, 1.0f);
 
 		// Apply the brush
@@ -714,7 +716,7 @@ public:
 
 				for (int32 X = BrushInfo.GetBounds().Min.X; X < BrushInfo.GetBounds().Max.X; X++)
 				{
-					const float BrushValue = BrushScanline[X];
+					const float BrushValue =  BrushScanline[X] * ToolStrength;
 
 					if (BrushValue > 0.0f)
 					{
@@ -734,7 +736,7 @@ public:
 						{
 							const float* SampleBrushScanline = BrushInfo.GetDataPtr(FIntPoint(0, SampleY));
 							const float* SampleBrushScanline2 = BrushInfo.GetDataPtr(FIntPoint(0, Y + (Y - SampleY)));
-							auto* SampleDataScanline = Data.GetData() + (SampleY - Y1) * (X2 - X1 + 1) + (0 - X1);
+							const auto* SampleDataScanline = ReadData.GetData() + (SampleY - Y1) * (X2 - X1 + 1) + (0 - X1);
 
 							for (int32 SampleX = SampleX1; SampleX <= SampleX2; SampleX++)
 							{
@@ -754,7 +756,7 @@ public:
 
 						FilterValue /= FilterSamplingNumber;
 
-						DataScanline[X] = FMath::Lerp(DataScanline[X], (typename ToolTarget::CacheClass::DataType)FilterValue, BrushValue * ToolStrength);
+						DataScanline[X] = FMath::Lerp(DataScanline[X], (typename ToolTarget::CacheClass::DataType)FilterValue, BrushValue);
 					}
 				}
 			}
