Fixed editor transaction being incorrectly started when using the mouse shortcut (ctrl+shift+right click-drag) to adjust the brush radius/falloff in landscape mode

diff --git a/Engine/Source/Editor/LandscapeEditor/Private/LandscapeEdMode.cpp b/Engine/Source/Editor/LandscapeEditor/Private/LandscapeEdMode.cpp
--- a/Engine/Source/Editor/LandscapeEditor/Private/LandscapeEdMode.cpp
+++ b/Engine/Source/Editor/LandscapeEditor/Private/LandscapeEdMode.cpp
@@ -80,6 +80,11 @@ namespace LandscapeTool
 		MaterialInstance->PostEditChange();
 		return MaterialInstance;
 	}
+
+	/** Indicates the user is currently moving the landscape gizmo object by dragging the mouse. */
+	bool GIsGizmoDragging = false;
+	/** Indicates the user is currently changing the landscape brush radius/falloff by dragging the mouse. */
+	bool GIsAdjustingBrush = false;
 }
 
 //
@@ -632,11 +637,19 @@ void FEdModeLandscape::Enter()
 	{
 		GizmoBrush->EnterBrush();
 	}
+
+	// Reset mouse tracking info : 
+	LandscapeTool::GIsGizmoDragging = false;
+	LandscapeTool::GIsAdjustingBrush  = false;
 }
 
 /** FEdMode: Called when the mode is exited */
 void FEdModeLandscape::Exit()
 {
+	// Reset mouse tracking info : 
+	LandscapeTool::GIsGizmoDragging = false;
+	LandscapeTool::GIsAdjustingBrush  = false;
+
 	if (UWorld* World = GetWorld())
 	{
 		for (auto It = ULandscapeInfoMap::GetLandscapeInfoMap(World).Map.CreateIterator(); It; ++It)
@@ -1021,17 +1034,18 @@ bool FEdModeLandscape::CapturedMouseMove(FEditorViewportClient* ViewportClient,
 	return MouseMove(ViewportClient, Viewport, MouseX, MouseY);
 }
 
-namespace
-{
-	bool GIsGizmoDragging = false;
-}
-
 /** FEdMode: Called when a mouse button is pressed */
 bool FEdModeLandscape::StartTracking(FEditorViewportClient* InViewportClient, FViewport* InViewport)
 {
 	if (CurrentGizmoActor.IsValid() && CurrentGizmoActor->IsSelected() && GLandscapeEditRenderMode & ELandscapeEditRenderMode::Gizmo)
 	{
-		GIsGizmoDragging = true;
+		LandscapeTool::GIsGizmoDragging = true;
+		return true;
+	}
+	else if (IsAdjustingBrush(InViewportClient))
+	{ 
+		LandscapeTool::GIsAdjustingBrush = true;
+		// We're adjusting the brush via mouse tracking, return true in order to prevent the viewport client from doing any mouse dragging operation while we're doing it
 		return true;
 	}
 	return false;
@@ -1042,9 +1056,14 @@ bool FEdModeLandscape::StartTracking(FEditorViewportClient* InViewportClient, FV
 /** FEdMode: Called when the a mouse button is released */
 bool FEdModeLandscape::EndTracking(FEditorViewportClient* InViewportClient, FViewport* InViewport)
 {
-	if (GIsGizmoDragging)
+	if (LandscapeTool::GIsGizmoDragging)
+	{
+		LandscapeTool::GIsGizmoDragging = false;
+		return true;
+	}
+	if (LandscapeTool::GIsAdjustingBrush)
 	{
-		GIsGizmoDragging = false;
+		LandscapeTool::GIsAdjustingBrush = false;
 		return true;
 	}
 	return false;
