Use more explicit names for landscape tool materials

diff --git a/Engine/Source/Editor/LandscapeEditor/Private/LandscapeEdMode.cpp b/Engine/Source/Editor/LandscapeEditor/Private/LandscapeEdMode.cpp
--- a/Engine/Source/Editor/LandscapeEditor/Private/LandscapeEdMode.cpp
+++ b/Engine/Source/Editor/LandscapeEditor/Private/LandscapeEdMode.cpp
@@ -74,7 +74,10 @@ namespace LandscapeTool
 {
 	UMaterialInstance* CreateMaterialInstance(UMaterialInterface* BaseMaterial)
 	{
-		ULandscapeMaterialInstanceConstant* MaterialInstance = NewObject<ULandscapeMaterialInstanceConstant>(GetTransientPackage());
+		UObject* Outer = GetTransientPackage();
+		// Use the base material's name as the base of our MIC to help debug: 
+		FString MICName(FString::Format(TEXT("LandscapeMaterialInstanceConstant_{0}"), { *BaseMaterial->GetName() }));
+		ULandscapeMaterialInstanceConstant* MaterialInstance = NewObject<ULandscapeMaterialInstanceConstant>(Outer, MakeUniqueObjectName(Outer, ULandscapeMaterialInstanceConstant::StaticClass(), FName(MICName)));
 		MaterialInstance->bEditorToolUsage = true;
 		MaterialInstance->SetParentEditorOnly(BaseMaterial);
 		MaterialInstance->PostEditChange();
