Landscape fix : prevent invalid paint layer names from being displayed in the New/Import landscape tool, like they are in the Paint panel

diff --git a/Engine/Source/Editor/LandscapeEditor/Private/LandscapeEditorObject.cpp b/Engine/Source/Editor/LandscapeEditor/Private/LandscapeEditorObject.cpp
--- a/Engine/Source/Editor/LandscapeEditor/Private/LandscapeEditorObject.cpp
+++ b/Engine/Source/Editor/LandscapeEditor/Private/LandscapeEditorObject.cpp
@@ -875,54 +875,57 @@ void ULandscapeEditorObject::RefreshImportLayersList(bool bRefreshFromTarget)
 	{
 		const FName& LayerName = LayerNames[i];
 
-		bool bFound = false;
-		FLandscapeImportLayer NewImportLayer;
-		NewImportLayer.ImportResult = ELandscapeImportResult::Success;
-		NewImportLayer.ErrorMessage = FText();
-
-		for (int32 j = 0; j < OldLayersList.Num(); j++)
+		if (!LayerName.IsNone())
 		{
-			if (OldLayersList[j].LayerName == LayerName)
+			bool bFound = false;
+			FLandscapeImportLayer NewImportLayer;
+			NewImportLayer.ImportResult = ELandscapeImportResult::Success;
+			NewImportLayer.ErrorMessage = FText();
+
+			for (int32 j = 0; j < OldLayersList.Num(); j++)
 			{
-				NewImportLayer = OldLayersList[j];
-				bFound = true;
-				break;
+				if (OldLayersList[j].LayerName == LayerName)
+				{
+					NewImportLayer = OldLayersList[j];
+					bFound = true;
+					break;
+				}
 			}
-		}
 
-		if (bFound)
-		{
-			if (NewImportLayer.ThumbnailMIC->Parent != Material)
+			if (bFound)
 			{
-				FMaterialUpdateContext Context;
-				NewImportLayer.ThumbnailMIC->SetParentEditorOnly(Material);
-				Context.AddMaterialInterface(NewImportLayer.ThumbnailMIC);
+				if (NewImportLayer.ThumbnailMIC->Parent != Material)
+				{
+					FMaterialUpdateContext Context;
+					NewImportLayer.ThumbnailMIC->SetParentEditorOnly(Material);
+					Context.AddMaterialInterface(NewImportLayer.ThumbnailMIC);
+				}
 			}
-		}
-		else
-		{
-			if (!ThumbnailWeightmap)
+			else
 			{
-				ThumbnailWeightmap = LoadObject<UTexture2D>(nullptr, TEXT("/Engine/EditorLandscapeResources/LandscapeThumbnailWeightmap.LandscapeThumbnailWeightmap"), nullptr, LOAD_None, nullptr);
+				if (!ThumbnailWeightmap)
+				{
+					ThumbnailWeightmap = LoadObject<UTexture2D>(nullptr, TEXT("/Engine/EditorLandscapeResources/LandscapeThumbnailWeightmap.LandscapeThumbnailWeightmap"), nullptr, LOAD_None, nullptr);
+				}
+
+				if (!ThumbnailHeightmap)
+				{
+					ThumbnailHeightmap = LoadObject<UTexture2D>(nullptr, TEXT("/Engine/EditorLandscapeResources/LandscapeThumbnailHeightmap.LandscapeThumbnailHeightmap"), nullptr, LOAD_None, nullptr);
+				}
+
+				NewImportLayer.LayerName = LayerName;
+				NewImportLayer.ThumbnailMIC = ALandscapeProxy::GetLayerThumbnailMIC(Material, LayerName, ThumbnailWeightmap, ThumbnailHeightmap, nullptr);
 			}
 
-			if (!ThumbnailHeightmap)
+			if (bRefreshFromTarget)
 			{
-				ThumbnailHeightmap = LoadObject<UTexture2D>(nullptr, TEXT("/Engine/EditorLandscapeResources/LandscapeThumbnailHeightmap.LandscapeThumbnailHeightmap"), nullptr, LOAD_None, nullptr);
+				NewImportLayer.LayerInfo = LayerInfoObjs[i];
 			}
 
-			NewImportLayer.LayerName = LayerName;
-			NewImportLayer.ThumbnailMIC = ALandscapeProxy::GetLayerThumbnailMIC(Material, LayerName, ThumbnailWeightmap, ThumbnailHeightmap, nullptr);
-		}
+			RefreshLayerImport(NewImportLayer);
 
-		if (bRefreshFromTarget)
-		{
-			NewImportLayer.LayerInfo = LayerInfoObjs[i];
+			ImportLandscape_Layers.Add(MoveTemp(NewImportLayer));
 		}
-
-		RefreshLayerImport(NewImportLayer);
-
-		ImportLandscape_Layers.Add(MoveTemp(NewImportLayer));
 	}
 }
 
