Fixed assert when resizing a landscape with a landscape splines component

diff --git a/Engine/Source/Editor/LandscapeEditor/Private/LandscapeEdMode.cpp b/Engine/Source/Editor/LandscapeEditor/Private/LandscapeEdMode.cpp
--- a/Engine/Source/Editor/LandscapeEditor/Private/LandscapeEdMode.cpp
+++ b/Engine/Source/Editor/LandscapeEditor/Private/LandscapeEdMode.cpp
@@ -3698,6 +3698,12 @@ ALandscape* FEdModeLandscape::ChangeComponentSetting(int32 NumComponentsX, int32
 			{
 				ULandscapeSplinesComponent* OldSplines = OldLandscapeActor->GetSplinesComponent();
 				ULandscapeSplinesComponent* NewSplines = DuplicateObject<ULandscapeSplinesComponent>(OldSplines, NewLandscape, OldSplines->GetFName());
+				// It's possible that the duplication of the ULandscapeSplinesComponent object actually triggered the creation of a new ULandscapeSplinesComponent already, in which case it's just simpler to use it: 
+				if (ULandscapeSplinesComponent* AlreadySetSplines = NewLandscape->GetSplinesComponent())
+				{
+					// This allows SetSplinesComponent() later on not to complain about setting a different ULandscapeSplinesComponent in the landscape :
+					NewSplines = AlreadySetSplines;
+				}
 				NewSplines->AttachToComponent(NewLandscape->GetRootComponent(), FAttachmentTransformRules::KeepWorldTransform);
 
 				const FVector OldSplineScale = OldSplines->GetRelativeTransform().GetScale3D();
