Disabled Layer Info object creation from the asset picker context menu in the Create landscape panels

diff --git a/Engine/Source/Editor/LandscapeEditor/Private/LandscapeEditorDetailCustomization_ImportLayers.cpp b/Engine/Source/Editor/LandscapeEditor/Private/LandscapeEditorDetailCustomization_ImportLayers.cpp
--- a/Engine/Source/Editor/LandscapeEditor/Private/LandscapeEditorDetailCustomization_ImportLayers.cpp
+++ b/Engine/Source/Editor/LandscapeEditor/Private/LandscapeEditorDetailCustomization_ImportLayers.cpp
@@ -172,6 +172,7 @@ void FLandscapeEditorStructCustomization_FLandscapeImportLayer::CustomizeChildre
 					.AllowedClass(ULandscapeLayerInfoObject::StaticClass())
 					.PropertyHandle(PropertyHandle_LayerInfo)
 					.OnShouldFilterAsset_Static(&ShouldFilterLayerInfo, LayerName)
+					.AllowCreate(false)
 				]
 				+ SHorizontalBox::Slot()
 				.AutoWidth()
