Landscape : removed CollisionThickness (unused since even before Chaos was introduced)

diff --git a/Engine/Source/Editor/LandscapeEditor/Private/LandscapeEdMode.cpp b/Engine/Source/Editor/LandscapeEditor/Private/LandscapeEdMode.cpp
--- a/Engine/Source/Editor/LandscapeEditor/Private/LandscapeEdMode.cpp
+++ b/Engine/Source/Editor/LandscapeEditor/Private/LandscapeEdMode.cpp
@@ -3666,7 +3666,6 @@ ALandscape* FEdModeLandscape::ChangeComponentSetting(int32 NumComponentsX, int32
 			NewLandscape->CustomDepthStencilWriteMask = OldLandscape->CustomDepthStencilWriteMask;
 			NewLandscape->CustomDepthStencilValue = OldLandscape->CustomDepthStencilValue;
 			NewLandscape->LightmassSettings = OldLandscape->LightmassSettings;
-			NewLandscape->CollisionThickness = OldLandscape->CollisionThickness;
 			NewLandscape->BodyInstance.SetCollisionProfileName(OldLandscape->BodyInstance.GetCollisionProfileName());
 			if (NewLandscape->BodyInstance.DoesUseCollisionProfile() == false)
 			{
diff --git a/Engine/Source/Runtime/Landscape/Classes/LandscapeProxy.h b/Engine/Source/Runtime/Landscape/Classes/LandscapeProxy.h
--- a/Engine/Source/Runtime/Landscape/Classes/LandscapeProxy.h
+++ b/Engine/Source/Runtime/Landscape/Classes/LandscapeProxy.h
@@ -664,8 +664,7 @@ public:
 	UPROPERTY(EditAnywhere, Category=Collision)
 	int32 SimpleCollisionMipLevel;
 
-	/** Thickness of the collision surface, in unreal units */
-	UPROPERTY(EditAnywhere, Category=Collision)
+	UE_DEPRECATED(5.2, "CollisionThickness is not supported by Chaos Physics and therefore deprecated. Please remove any usage of this property")
 	float CollisionThickness;
 
 	/** Collision profile settings for this landscape */
diff --git a/Engine/Source/Runtime/Landscape/Private/Landscape.cpp b/Engine/Source/Runtime/Landscape/Private/Landscape.cpp
--- a/Engine/Source/Runtime/Landscape/Private/Landscape.cpp
+++ b/Engine/Source/Runtime/Landscape/Private/Landscape.cpp
@@ -1305,7 +1305,6 @@ ALandscapeProxy::ALandscapeProxy(const FObjectInitializer& ObjectInitializer)
 	bCastStaticShadow = true;
 	bUsedForNavigation = true;
 	bFillCollisionUnderLandscapeForNavmesh = false;
-	CollisionThickness = 16;
 	BodyInstance.SetCollisionProfileName(UCollisionProfile::BlockAll_ProfileName);
 	bGenerateOverlapEvents = false;
 #if WITH_EDITORONLY_DATA
diff --git a/Engine/Source/Runtime/Landscape/Private/LandscapeEdit.cpp b/Engine/Source/Runtime/Landscape/Private/LandscapeEdit.cpp
--- a/Engine/Source/Runtime/Landscape/Private/LandscapeEdit.cpp
+++ b/Engine/Source/Runtime/Landscape/Private/LandscapeEdit.cpp
@@ -5400,7 +5400,6 @@ void ALandscapeProxy::PostEditChangeProperty(FPropertyChangedEvent& PropertyChan
 	else if (GIsEditor &&
 		(PropertyName == GET_MEMBER_NAME_CHECKED(ALandscapeProxy, CollisionMipLevel) ||
 		 PropertyName == GET_MEMBER_NAME_CHECKED(ALandscapeProxy, SimpleCollisionMipLevel) ||
-		 PropertyName == GET_MEMBER_NAME_CHECKED(ALandscapeProxy, CollisionThickness) ||
 		 PropertyName == GET_MEMBER_NAME_CHECKED(ALandscapeProxy, bBakeMaterialPositionOffsetIntoCollision) ||
 		 PropertyName == GET_MEMBER_NAME_CHECKED(ALandscapeProxy, bGenerateOverlapEvents)))
 	{
