Landscape editor Brush Size now has a maximum fractional digits of 2.

diff --git a/Engine/Source/Editor/LandscapeEditor/Public/LandscapeEditorObject.h b/Engine/Source/Editor/LandscapeEditor/Public/LandscapeEditorObject.h
--- a/Engine/Source/Editor/LandscapeEditor/Public/LandscapeEditorObject.h
+++ b/Engine/Source/Editor/LandscapeEditor/Public/LandscapeEditorObject.h
@@ -594,7 +594,7 @@ public:
 	// Common Brush Settings:
 
 	// The radius of the sculpt brush, in unreal units
-	UPROPERTY(Category="Brush Settings", EditAnywhere, NonTransactional, meta=(DisplayName="Brush Size", ShowForBrushes="BrushSet_Circle,BrushSet_Alpha,BrushSet_Pattern", ShowForTargetTypes = "Heightmap,Visibility", ClampMin="1", ClampMax="65536", UIMin="1", UIMax="8192", SliderExponent="3"))
+	UPROPERTY(Category="Brush Settings", EditAnywhere, NonTransactional, meta=(DisplayName="Brush Size", ShowForBrushes="BrushSet_Circle,BrushSet_Alpha,BrushSet_Pattern", ShowForTargetTypes = "Heightmap,Visibility", ClampMin="1", ClampMax="65536", UIMin="1", UIMax="8192", SliderExponent="3", MaxFractionalDigits="2"))
 	float BrushRadius;
 
 	// The radius of the paint brush, in unreal units
