Removed some miscellaneous TIsSame usage (part of a larger change to deprecate TIsSame).

diff --git a/Engine/Source/Editor/LandscapeEditor/Public/LandscapeEditorModule.h b/Engine/Source/Editor/LandscapeEditor/Public/LandscapeEditorModule.h
--- a/Engine/Source/Editor/LandscapeEditor/Public/LandscapeEditorModule.h
+++ b/Engine/Source/Editor/LandscapeEditor/Public/LandscapeEditorModule.h
@@ -39,13 +39,13 @@ public:
 	virtual const ILandscapeWeightmapFileFormat* GetWeightmapFormatByExtension(const TCHAR* Extension) const = 0;
 
 	template<typename T>
-	typename TEnableIf<TIsSame<T, uint16>::Value, const ILandscapeHeightmapFileFormat*>::Type GetFormatByExtension(const TCHAR* Extension)
+	typename TEnableIf<std::is_same_v<T, uint16>, const ILandscapeHeightmapFileFormat*>::Type GetFormatByExtension(const TCHAR* Extension)
 	{
 		return GetHeightmapFormatByExtension(Extension);
 	}
 
 	template<typename T>
-	typename TEnableIf<TIsSame<T, uint8>::Value, const ILandscapeWeightmapFileFormat*>::Type GetFormatByExtension(const TCHAR* Extension)
+	typename TEnableIf<std::is_same_v<T, uint8>, const ILandscapeWeightmapFileFormat*>::Type GetFormatByExtension(const TCHAR* Extension)
 	{
 		return GetWeightmapFormatByExtension(Extension);
 	}
diff --git a/Engine/Source/Editor/LandscapeEditor/Public/LandscapeImportHelper.h b/Engine/Source/Editor/LandscapeEditor/Public/LandscapeImportHelper.h
--- a/Engine/Source/Editor/LandscapeEditor/Public/LandscapeImportHelper.h
+++ b/Engine/Source/Editor/LandscapeEditor/Public/LandscapeImportHelper.h
@@ -141,34 +141,34 @@ public:
 
 
 	template<typename T>
-	static typename TEnableIf<TIsSame<T, uint16>::Value, ELandscapeImportResult>::Type GetImportDescriptor(const FString& FilePath, bool bSingleFile, bool bFlipYAxis, FName LayerName, FLandscapeImportDescriptor& OutImportDescriptor, FText& OutMessage)
+	static typename TEnableIf<std::is_same_v<T, uint16>, ELandscapeImportResult>::Type GetImportDescriptor(const FString& FilePath, bool bSingleFile, bool bFlipYAxis, FName LayerName, FLandscapeImportDescriptor& OutImportDescriptor, FText& OutMessage)
 	{
 		return GetHeightmapImportDescriptor(FilePath, bSingleFile, bFlipYAxis, OutImportDescriptor, OutMessage);
 	}
 
 	template<typename T>
-	static typename TEnableIf<TIsSame<T, uint16>::Value, ELandscapeImportResult>::Type GetImportData(const FLandscapeImportDescriptor& ImportDescriptor, int32 DescriptorIndex, FName LayerName, TArray<T>& OutData, FText& OutMessage)
+	static typename TEnableIf<std::is_same_v<T, uint16>, ELandscapeImportResult>::Type GetImportData(const FLandscapeImportDescriptor& ImportDescriptor, int32 DescriptorIndex, FName LayerName, TArray<T>& OutData, FText& OutMessage)
 	{
 		return GetHeightmapImportData(ImportDescriptor, DescriptorIndex, OutData, OutMessage);
 	}
 	template<typename T>
-	static typename TEnableIf<TIsSame<T, uint16>::Value, void>::Type TransformImportData(const TArray<T>& InData, TArray<T>& OutData, const FLandscapeImportResolution& CurrentResolution, const FLandscapeImportResolution& RequiredResolution, ELandscapeImportTransformType TransformType = ELandscapeImportTransformType::ExpandCentered, FIntPoint Offset = FIntPoint(0, 0))
+	static typename TEnableIf<std::is_same_v<T, uint16>, void>::Type TransformImportData(const TArray<T>& InData, TArray<T>& OutData, const FLandscapeImportResolution& CurrentResolution, const FLandscapeImportResolution& RequiredResolution, ELandscapeImportTransformType TransformType = ELandscapeImportTransformType::ExpandCentered, FIntPoint Offset = FIntPoint(0, 0))
 	{
 		return TransformHeightmapImportData(InData, OutData, CurrentResolution, RequiredResolution, TransformType, Offset);
 	}
 
 	template<typename T>
-	static typename TEnableIf<TIsSame<T, uint8>::Value, ELandscapeImportResult>::Type GetImportDescriptor(const FString& FilePath, bool bSingleFile, bool bFlipYAxis, FName LayerName, FLandscapeImportDescriptor& OutImportDescriptor, FText& OutMessage)
+	static typename TEnableIf<std::is_same_v<T, uint8>, ELandscapeImportResult>::Type GetImportDescriptor(const FString& FilePath, bool bSingleFile, bool bFlipYAxis, FName LayerName, FLandscapeImportDescriptor& OutImportDescriptor, FText& OutMessage)
 	{
 		return GetWeightmapImportDescriptor(FilePath, bSingleFile, bFlipYAxis, LayerName, OutImportDescriptor, OutMessage);
 	}
 	template<typename T>
-	static typename TEnableIf<TIsSame<T, uint8>::Value, ELandscapeImportResult>::Type GetImportData(const FLandscapeImportDescriptor& ImportDescriptor, int32 DescriptorIndex, FName LayerName, TArray<T>& OutData, FText& OutMessage)
+	static typename TEnableIf<std::is_same_v<T, uint8>, ELandscapeImportResult>::Type GetImportData(const FLandscapeImportDescriptor& ImportDescriptor, int32 DescriptorIndex, FName LayerName, TArray<T>& OutData, FText& OutMessage)
 	{
 		return GetWeightmapImportData(ImportDescriptor, DescriptorIndex, LayerName, OutData, OutMessage);
 	}
 	template<typename T>
-	static typename TEnableIf<TIsSame<T, uint8>::Value, void>::Type TransformImportData(const TArray<T>& InData, TArray<T>& OutData, const FLandscapeImportResolution& CurrentResolution, const FLandscapeImportResolution& RequiredResolution, ELandscapeImportTransformType TransformType = ELandscapeImportTransformType::ExpandCentered, FIntPoint Offset = FIntPoint(0, 0))
+	static typename TEnableIf<std::is_same_v<T, uint8>, void>::Type TransformImportData(const TArray<T>& InData, TArray<T>& OutData, const FLandscapeImportResolution& CurrentResolution, const FLandscapeImportResolution& RequiredResolution, ELandscapeImportTransformType TransformType = ELandscapeImportTransformType::ExpandCentered, FIntPoint Offset = FIntPoint(0, 0))
 	{
 		return TransformWeightmapImportData(InData, OutData, CurrentResolution, RequiredResolution, TransformType, Offset);
 	}
diff --git a/Engine/Source/Editor/Sequencer/Public/SequencerChannelInterface.h b/Engine/Source/Editor/Sequencer/Public/SequencerChannelInterface.h
--- a/Engine/Source/Editor/Sequencer/Public/SequencerChannelInterface.h
+++ b/Engine/Source/Editor/Sequencer/Public/SequencerChannelInterface.h
@@ -295,6 +295,6 @@ struct TSequencerChannelInterfaceBase<ChannelType, true> : TSequencerChannelInte
 
 /** Generic sequencer channel interface to any channel type */
 template<typename ChannelType>
-struct TSequencerChannelInterface : TSequencerChannelInterfaceBase<ChannelType, !TIsSame<typename TMovieSceneChannelTraits<ChannelType>::ExtendedEditorDataType, void>::Value>
+struct TSequencerChannelInterface : TSequencerChannelInterfaceBase<ChannelType, !std::is_same_v<typename TMovieSceneChannelTraits<ChannelType>::ExtendedEditorDataType, void>>
 {
 };
