Fixed crash when using alpha brush/pattern landscape tool and generally better support various cases with the alpha brush texture :
* The brushes becomes red and non-interactive when the alpha brush is invalid
* Setting an invalid alpha brush (e.g. a cooked texture in a cooked editor or a virtualized texture with invalid virtualization settings such that the texture source cannot be retrieved) is not crashing the UI anymore and will show as empty in the UI (+ a log error will be done)
* The brushes now support the same failure conditions as the normal brush (unloaded components, etc.)
* The alpha brush textures now supports all texture formats and will be linearized if necessary when reading the texture data
* The alpha brush texture channel is auto-corrected if one uses an invalid channel (e.g. Green on a G8 texture)

diff --git a/Engine/Source/Editor/LandscapeEditor/LandscapeEditor.Build.cs b/Engine/Source/Editor/LandscapeEditor/LandscapeEditor.Build.cs
--- a/Engine/Source/Editor/LandscapeEditor/LandscapeEditor.Build.cs
+++ b/Engine/Source/Editor/LandscapeEditor/LandscapeEditor.Build.cs
@@ -19,6 +19,7 @@ public class LandscapeEditor : ModuleRules
                 "RenderCore",
                 "RHI",
                 "InputCore",
+				"ImageCore",
 				"EditorFramework",
 				"UnrealEd",
 				"PropertyEditor",
diff --git a/Engine/Source/Editor/LandscapeEditor/Private/LandscapeEdModeBrushes.cpp b/Engine/Source/Editor/LandscapeEditor/Private/LandscapeEdModeBrushes.cpp
--- a/Engine/Source/Editor/LandscapeEditor/Private/LandscapeEdModeBrushes.cpp
+++ b/Engine/Source/Editor/LandscapeEditor/Private/LandscapeEdModeBrushes.cpp
@@ -46,6 +46,8 @@ void FLandscapeBrush::EndStroke()
 
 class FLandscapeBrushCircle : public FLandscapeBrush
 {
+	using Super = FLandscapeBrush;
+
 	TSet<ULandscapeComponent*> BrushMaterialComponents;
 	TArray<UMaterialInstanceDynamic*> BrushMaterialFreeInstances;
 
@@ -119,10 +121,61 @@ public:
 
 	virtual void BeginStroke(float LandscapeX, float LandscapeY, FLandscapeTool* CurrentTool) override
 	{
-		FLandscapeBrush::BeginStroke(LandscapeX, LandscapeY, CurrentTool);
+		Super::BeginStroke(LandscapeX, LandscapeY, CurrentTool);
 		LastMousePosition = FVector2f(LandscapeX, LandscapeY);
 	}
 
+	virtual bool CanPaint(const TSet<ULandscapeComponent*>& InAffectedComponents, bool bInHasUnloadedComponents) const 
+	{
+		if (bInHasUnloadedComponents)
+		{
+			return false;
+		}
+
+		if (!EdMode->CanEditLayer())
+		{
+			return false;
+		}
+
+		for (const ULandscapeComponent* Component : InAffectedComponents)
+		{
+			const ALandscapeProxy* LandscapeProxy = Component->GetLandscapeProxy();
+			const ULandscapeLayerInfoObject* LayerInfo = EdMode->CurrentToolTarget.LayerInfo.Get();
+
+			if (((EdMode->CurrentToolTarget.TargetType == ELandscapeToolTargetType::Weightmap) || (EdMode->CurrentToolTarget.TargetType == ELandscapeToolTargetType::Visibility))
+				&& (EdMode->UISettings->PaintingRestriction != ELandscapeLayerPaintingRestriction::None))
+			{
+				if ((EdMode->UISettings->PaintingRestriction == ELandscapeLayerPaintingRestriction::UseComponentAllowList) 
+					&& !Component->LayerAllowList.Contains(LayerInfo))
+				{
+					return false;
+				}
+				else
+				{
+					const TArray<FWeightmapLayerAllocationInfo>& ComponentWeightmapLayerAllocations = Component->GetWeightmapLayerAllocations(true);
+
+					bool bExisting = ComponentWeightmapLayerAllocations.ContainsByPredicate([LayerInfo](const FWeightmapLayerAllocationInfo& Allocation) { return Allocation.LayerInfo == LayerInfo; });
+					if (!bExisting)
+					{
+						if (((EdMode->UISettings->PaintingRestriction == ELandscapeLayerPaintingRestriction::ExistingOnly) || (EdMode->UISettings->PaintingRestriction == ELandscapeLayerPaintingRestriction::UseMaxLayers))
+							&& (LandscapeProxy->MaxPaintedLayersPerComponent > 0) 
+							&& (ComponentWeightmapLayerAllocations.Num() >= LandscapeProxy->MaxPaintedLayersPerComponent))
+						{
+							return false;
+						}
+					}
+				}
+			}
+
+			if (EdMode->CurrentToolTarget.TargetType == ELandscapeToolTargetType::Visibility && !Component->IsLandscapeHoleMaterialValid())
+			{
+				return false;
+			}
+		}
+
+		return true;
+	}
+
 	virtual void Tick(FEditorViewportClient* ViewportClient, float DeltaTime) override
 	{
 		ULandscapeInfo* LandscapeInfo = EdMode->CurrentToolTarget.LandscapeInfo.Get();
@@ -182,6 +235,9 @@ public:
 		// Set params for brush material.
 		FVector WorldLocation = Proxy->LandscapeActorToWorld().TransformPosition(FVector(LastMousePosition.X, LastMousePosition.Y, 0));
 
+		// Update whether the brush can actually paint : 
+		bCanPaint = CanPaint(NewComponents, bHasUnloadedComponents);
+
 		for (const auto& BrushMaterialInstancePair : BrushMaterialInstanceMap)
 		{
 			ULandscapeComponent* const Component = BrushMaterialInstancePair.Key;
@@ -200,48 +256,6 @@ public:
 			MaterialInstance->SetScalarParameterValue(FName(TEXT("LocalRadius")), Radius);
 			MaterialInstance->SetScalarParameterValue(FName(TEXT("LocalFalloff")), Falloff);
 			MaterialInstance->SetDoubleVectorParameterValue(FName(TEXT("WorldPosition")), FVector4(WorldLocation.X, WorldLocation.Y, WorldLocation.Z, ScaleXY));
-
-			bCanPaint = !bHasUnloadedComponents;
-			if (bCanPaint)
-			{
-				const ALandscapeProxy* LandscapeProxy = Component->GetLandscapeProxy();
-				const ULandscapeLayerInfoObject* LayerInfo = EdMode->CurrentToolTarget.LayerInfo.Get();
-
-				if (!EdMode->CanEditLayer())
-				{
-					bCanPaint = false;
-				}
-				else if ((EdMode->CurrentToolTarget.TargetType == ELandscapeToolTargetType::Weightmap || EdMode->CurrentToolTarget.TargetType == ELandscapeToolTargetType::Visibility) &&
-					EdMode->UISettings->PaintingRestriction != ELandscapeLayerPaintingRestriction::None)
-				{
-					if (EdMode->UISettings->PaintingRestriction == ELandscapeLayerPaintingRestriction::UseComponentAllowList &&
-						!Component->LayerAllowList.Contains(LayerInfo))
-					{
-						bCanPaint = false;
-					}
-					else
-					{
-						const TArray<FWeightmapLayerAllocationInfo>& ComponentWeightmapLayerAllocations = Component->GetWeightmapLayerAllocations(true);
-
-						bool bExisting = ComponentWeightmapLayerAllocations.ContainsByPredicate([LayerInfo](const FWeightmapLayerAllocationInfo& Allocation) { return Allocation.LayerInfo == LayerInfo; });
-						if (!bExisting)
-						{
-							if (EdMode->UISettings->PaintingRestriction == ELandscapeLayerPaintingRestriction::ExistingOnly ||
-								(EdMode->UISettings->PaintingRestriction == ELandscapeLayerPaintingRestriction::UseMaxLayers &&
-									LandscapeProxy->MaxPaintedLayersPerComponent > 0 && ComponentWeightmapLayerAllocations.Num() >= LandscapeProxy->MaxPaintedLayersPerComponent))
-							{
-								bCanPaint = false;
-							}
-						}
-					}
-				}
-
-				if (bCanPaint && EdMode->CurrentToolTarget.TargetType == ELandscapeToolTargetType::Visibility && !Component->IsLandscapeHoleMaterialValid())
-				{
-					bCanPaint = false;
-				}
-			}
-
 			MaterialInstance->SetScalarParameterValue("CanPaint", bCanPaint ? 1.0f : 0.0f);
 		}
 	}
@@ -366,6 +380,8 @@ public:
 
 class FLandscapeBrushComponent : public FLandscapeBrush
 {
+	using Super = FLandscapeBrush;
+
 	TSet<ULandscapeComponent*> BrushMaterialComponents;
 
 	virtual const TCHAR* GetBrushName() override { return TEXT("Component"); }
@@ -409,7 +425,7 @@ public:
 
 	virtual void BeginStroke(float LandscapeX, float LandscapeY, FLandscapeTool* CurrentTool) override
 	{
-		FLandscapeBrush::BeginStroke(LandscapeX, LandscapeY, CurrentTool);
+		Super::BeginStroke(LandscapeX, LandscapeY, CurrentTool);
 		LastMousePosition = FVector2D(LandscapeX, LandscapeY);
 	}
 
@@ -568,6 +584,8 @@ public:
 
 class FLandscapeBrushGizmo : public FLandscapeBrush
 {
+	using Super = FLandscapeBrush;
+
 	TSet<ULandscapeComponent*> BrushMaterialComponents;
 
 	const TCHAR* GetBrushName() override { return TEXT("Gizmo"); }
@@ -817,6 +835,8 @@ public:
 //
 class FLandscapeBrushSplines : public FLandscapeBrush
 {
+	using Super = FLandscapeBrush;
+
 public:
 	const TCHAR* GetBrushName() override { return TEXT("Splines"); }
 	virtual FText GetDisplayName() override { return NSLOCTEXT("UnrealEd", "LandscapeMode_Brush_Splines", "Splines"); };
@@ -849,6 +869,8 @@ public:
 //
 class FLandscapeBrushDummy : public FLandscapeBrush
 {
+	using Super = FLandscapeBrush;
+
 public:
 	const TCHAR* GetBrushName() override { return TEXT("Circle_Dummy"); }
 	virtual FText GetDisplayName() override { return NSLOCTEXT("UnrealEd", "LandscapeMode_Brush_None", "None"); };
@@ -878,6 +900,8 @@ public:
 
 class FLandscapeBrushCircle_Linear : public FLandscapeBrushCircle
 {
+	using Super = FLandscapeBrushCircle;
+
 protected:
 	/** Protected so only subclasses can create instances. */
 	FLandscapeBrushCircle_Linear(FEdModeLandscape* InEdMode, UMaterialInterface* InBrushMaterial)
@@ -905,6 +929,8 @@ public:
 
 class FLandscapeBrushCircle_Smooth : public FLandscapeBrushCircle_Linear
 {
+	using Super = FLandscapeBrushCircle_Linear;
+
 protected:
 	FLandscapeBrushCircle_Smooth(FEdModeLandscape* InEdMode, UMaterialInterface* InBrushMaterial)
 		: FLandscapeBrushCircle_Linear(InEdMode, InBrushMaterial)
@@ -913,9 +939,9 @@ protected:
 
 	virtual float CalculateFalloff(float Distance, float Radius, float Falloff) override
 	{
-		float y = FLandscapeBrushCircle_Linear::CalculateFalloff(Distance, Radius, Falloff);
+		float Y = Super::CalculateFalloff(Distance, Radius, Falloff);
 		// Smooth-step it
-		return y*y*(3 - 2 * y);
+		return Y * Y * (3 - 2 * Y);
 	}
 
 public:
@@ -931,6 +957,8 @@ public:
 
 class FLandscapeBrushCircle_Spherical : public FLandscapeBrushCircle
 {
+	using Super = FLandscapeBrushCircle;
+
 protected:
 	FLandscapeBrushCircle_Spherical(FEdModeLandscape* InEdMode, UMaterialInterface* InBrushMaterial)
 		: FLandscapeBrushCircle(InEdMode, InBrushMaterial)
@@ -965,6 +993,8 @@ public:
 
 class FLandscapeBrushCircle_Tip : public FLandscapeBrushCircle
 {
+	using Super = FLandscapeBrushCircle_Tip;
+
 protected:
 	FLandscapeBrushCircle_Tip(FEdModeLandscape* InEdMode, UMaterialInterface* InBrushMaterial)
 		: FLandscapeBrushCircle(InEdMode, InBrushMaterial)
@@ -1001,6 +1031,8 @@ public:
 // FLandscapeBrushAlphaBase
 class FLandscapeBrushAlphaBase : public FLandscapeBrushCircle_Smooth
 {
+	using Super = FLandscapeBrushCircle_Smooth;
+
 public:
 	FLandscapeBrushAlphaBase(FEdModeLandscape* InEdMode, UMaterialInterface* InBrushMaterial)
 		: FLandscapeBrushCircle_Smooth(InEdMode, InBrushMaterial)
@@ -1009,8 +1041,14 @@ public:
 
 	float GetAlphaSample(float SampleX, float SampleY)
 	{
+		if (!EdMode->UISettings->HasValidAlphaTextureData())
+		{
+			return 1.0f;
+		}
+
 		int32 SizeX = EdMode->UISettings->AlphaTextureSizeX;
 		int32 SizeY = EdMode->UISettings->AlphaTextureSizeY;
+		check((SizeX > 0) && (SizeY > 0));
 
 		// Bilinear interpolate the values from the alpha texture
 		int32 SampleX0 = FMath::FloorToInt(SampleX);
@@ -1018,6 +1056,7 @@ public:
 		int32 SampleY0 = FMath::FloorToInt(SampleY);
 		int32 SampleY1 = (SampleY0 + 1) % SizeY;
 
+		check(!EdMode->UISettings->AlphaTextureData.IsEmpty());
 		const uint8* AlphaData = EdMode->UISettings->AlphaTextureData.GetData();
 
 		float Alpha00 = (float)AlphaData[SampleX0 + SampleY0 * SizeX] / 255.0f;
@@ -1039,6 +1078,8 @@ public:
 //
 class FLandscapeBrushAlphaPattern : public FLandscapeBrushAlphaBase
 {
+	using Super = FLandscapeBrushAlphaBase;
+
 protected:
 	FLandscapeBrushAlphaPattern(FEdModeLandscape* InEdMode, UMaterialInterface* InBrushMaterial)
 		: FLandscapeBrushAlphaBase(InEdMode, InBrushMaterial)
@@ -1054,6 +1095,16 @@ public:
 
 	virtual ELandscapeBrushType GetBrushType() override { return ELandscapeBrushType::Alpha; }
 
+	virtual bool CanPaint(const TSet<ULandscapeComponent*>& InAffectedComponents, bool bInHasUnloadedComponents) const override
+	{
+		if (!Super::CanPaint(InAffectedComponents, bInHasUnloadedComponents))
+		{
+			return false;
+		}
+
+		return EdMode->UISettings->HasValidAlphaTextureData();
+	}
+
 	virtual FLandscapeBrushData ApplyBrush(const TArray<FLandscapeToolInteractorPosition>& InteractorPositions) override
 	{
 		if (!bCanPaint)
@@ -1069,6 +1120,7 @@ public:
 
 		int32 SizeX = EdMode->UISettings->AlphaTextureSizeX;
 		int32 SizeY = EdMode->UISettings->AlphaTextureSizeY;
+		check(EdMode->UISettings->HasValidAlphaTextureData() && (SizeX > 0) && (SizeY > 0)); // See CanPaint() above : if we can paint, these must all be valid
 
 		FIntRect Bounds;
 		Bounds.Min.X = FMath::FloorToInt(LastMousePosition.X - TotalRadius);
@@ -1167,23 +1219,30 @@ public:
 
 	virtual void Tick(FEditorViewportClient* ViewportClient, float DeltaTime) override
 	{
-		FLandscapeBrushCircle::Tick(ViewportClient, DeltaTime);
+		Super::Tick(ViewportClient, DeltaTime);
 
 		ALandscapeProxy* Proxy = EdMode->CurrentToolTarget.LandscapeInfo.IsValid() ? EdMode->CurrentToolTarget.LandscapeInfo->GetLandscapeProxy() : nullptr;
 		if (Proxy)
 		{
 			const float ScaleXY = FMath::Abs(EdMode->CurrentToolTarget.LandscapeInfo->DrawScale.X);
+
+			const bool bHasValidAlphaTextureData = EdMode->UISettings->HasValidAlphaTextureData();
 			int32 SizeX = EdMode->UISettings->AlphaTextureSizeX;
 			int32 SizeY = EdMode->UISettings->AlphaTextureSizeY;
+			check(!bHasValidAlphaTextureData || ((SizeX > 0) && (SizeY > 0)));
 
-			FLinearColor AlphaScaleBias;
+			FLinearColor AlphaScaleBias(ForceInitToZero);
 			float Angle;
 			if (EdMode->UISettings->bUseWorldSpacePatternBrush)
 			{
 				FVector2D LocalOrigin = -FVector2D(Proxy->LandscapeActorToWorld().InverseTransformPosition(FVector(EdMode->UISettings->WorldSpacePatternBrushSettings.Origin, 0.0f)));
-				const FVector2D Scale = FVector2D(
-					ScaleXY / (EdMode->UISettings->WorldSpacePatternBrushSettings.RepeatSize * ((float)SizeX / SizeY)),
-					ScaleXY / (EdMode->UISettings->WorldSpacePatternBrushSettings.RepeatSize));
+				FVector2D Scale(ForceInitToZero);
+				if (bHasValidAlphaTextureData)
+				{
+					Scale = FVector2D(
+						ScaleXY / (EdMode->UISettings->WorldSpacePatternBrushSettings.RepeatSize * ((float)SizeX / SizeY)),
+						ScaleXY / (EdMode->UISettings->WorldSpacePatternBrushSettings.RepeatSize));
+				}
 				LocalOrigin *= Scale;
 				Angle = -EdMode->UISettings->WorldSpacePatternBrushSettings.Rotation;
 				if (EdMode->UISettings->WorldSpacePatternBrushSettings.bCenterTextureOnOrigin)
@@ -1196,11 +1255,12 @@ public:
 					LocalOrigin.X,
 					LocalOrigin.Y);
 			}
-			else
+			else 
 			{
+				const FVector2D Scale(EdMode->UISettings->AlphaBrushScale * SizeX, EdMode->UISettings->AlphaBrushScale * SizeY);
 				AlphaScaleBias = FLinearColor(
-					1.0f / (EdMode->UISettings->AlphaBrushScale * SizeX),
-					1.0f / (EdMode->UISettings->AlphaBrushScale * SizeY),
+					FMath::IsNearlyZero(Scale.X) ? 1.0f : 1.0f / (Scale.X),
+					FMath::IsNearlyZero(Scale.Y) ? 1.0f : 1.0f / (Scale.Y),
 					EdMode->UISettings->AlphaBrushPanU,
 					EdMode->UISettings->AlphaBrushPanV);
 				Angle = EdMode->UISettings->AlphaBrushRotation;
@@ -1210,16 +1270,19 @@ public:
 			FVector LandscapeLocation = Proxy->LandscapeActorToWorld().GetTranslation();
 			FVector4 LandscapeLocationParam(LandscapeLocation.X, LandscapeLocation.Y, LandscapeLocation.Z, Angle);
 
-			int32 Channel = EdMode->UISettings->AlphaTextureChannel;
-			FLinearColor AlphaTextureMask(Channel == 0 ? 1 : 0, Channel == 1 ? 1 : 0, Channel == 2 ? 1 : 0, Channel == 3 ? 1 : 0);
+			FLinearColor AlphaTextureMask(
+				EdMode->UISettings->AlphaTextureChannel == ELandscapeTextureColorChannel::Red ? 1 : 0, 
+				EdMode->UISettings->AlphaTextureChannel == ELandscapeTextureColorChannel::Green ? 1 : 0, 
+				EdMode->UISettings->AlphaTextureChannel == ELandscapeTextureColorChannel::Blue ? 1 : 0, 
+				EdMode->UISettings->AlphaTextureChannel == ELandscapeTextureColorChannel::Alpha ? 1 : 0);
 
 			for (const auto& BrushMaterialInstancePair : BrushMaterialInstanceMap)
 			{
 				UMaterialInstanceDynamic* const MaterialInstance = BrushMaterialInstancePair.Value;
-				MaterialInstance->SetVectorParameterValue(FName(TEXT("AlphaScaleBias")),    AlphaScaleBias);
+				MaterialInstance->SetVectorParameterValue(FName(TEXT("AlphaScaleBias")), AlphaScaleBias);
 				MaterialInstance->SetDoubleVectorParameterValue(FName(TEXT("LandscapeLocation")), LandscapeLocationParam);
-				MaterialInstance->SetVectorParameterValue(FName(TEXT("AlphaTextureMask")),  AlphaTextureMask);
-				MaterialInstance->SetTextureParameterValue(FName(TEXT("AlphaTexture")),     EdMode->UISettings->AlphaTexture);
+				MaterialInstance->SetVectorParameterValue(FName(TEXT("AlphaTextureMask")), AlphaTextureMask);
+				MaterialInstance->SetTextureParameterValue(FName(TEXT("AlphaTexture")), EdMode->UISettings->AlphaTexture);
 			}
 		}
 	}
@@ -1234,6 +1297,8 @@ public:
 //
 class FLandscapeBrushAlpha : public FLandscapeBrushAlphaBase
 {
+	using Super = FLandscapeBrushAlphaBase;
+
 	float LastMouseAngle;
 	FVector2f OldMousePosition;	// a previous mouse position, kept until we move a certain distance away, for smoothing deltas
 	double LastMouseSampleTime;
@@ -1254,6 +1319,57 @@ public:
 		return new FLandscapeBrushAlpha(InEdMode, AlphaBrushMaterial);
 	}
 
+	bool ComputeAlphaBrushScaleAndRadius(float& OutAlphaBrushScale, float& OutRadius, int32& OutSizeX, int32& OutSizeY) const
+	{
+		ULandscapeInfo* LandscapeInfo = EdMode->CurrentToolTarget.LandscapeInfo.Get();
+
+		int32 SizeX = EdMode->UISettings->AlphaTextureSizeX;
+		int32 SizeY = EdMode->UISettings->AlphaTextureSizeY;
+		bool bIsValidTextureSize = (SizeX > 0) && (SizeY > 0);
+
+		float ScaleXY = FMath::Abs(LandscapeInfo->DrawScale.X);
+		float Radius = FMath::IsNearlyZero(ScaleXY) ? 0.0f : EdMode->UISettings->GetCurrentToolBrushRadius() / ScaleXY;
+
+		float MaxSize = 2.0f * FMath::Sqrt(FMath::Square(Radius) / 2.0f);
+		float FinalScale = bIsValidTextureSize ? MaxSize / (float)FMath::Max<int32>(SizeX, SizeY) : 0.0f;
+
+		if (FMath::IsNearlyZero(FinalScale))
+		{
+			return false;
+		}
+
+		check(bIsValidTextureSize);
+		OutAlphaBrushScale = FinalScale;
+		OutRadius = Radius;
+		OutSizeX = SizeX;
+		OutSizeY = SizeY;
+		return true;
+	}
+
+	virtual bool CanPaint(const TSet<ULandscapeComponent*>& InAffectedComponents, bool bInHasUnloadedComponents) const override
+	{
+		if (!Super::CanPaint(InAffectedComponents, bInHasUnloadedComponents))
+		{
+			return false;
+		}
+
+		if (!EdMode->UISettings->HasValidAlphaTextureData())
+		{
+			return false;
+		}
+
+		float AlphaBrushScale; 
+		float Radius;
+		int32 SizeX;
+		int32 SizeY;
+		if (!ComputeAlphaBrushScaleAndRadius(AlphaBrushScale, Radius, SizeX, SizeY))
+		{
+			return false;
+		}
+
+		return true;
+	}
+
 	virtual FLandscapeBrushData ApplyBrush(const TArray<FLandscapeToolInteractorPosition>& InteractorPositions) override
 	{
 		if (!bCanPaint)
@@ -1271,21 +1387,20 @@ public:
 		}
 		else
 		{
-			float ScaleXY = FMath::Abs(LandscapeInfo->DrawScale.X);
-			float Radius = EdMode->UISettings->GetCurrentToolBrushRadius() / ScaleXY;
-			int32 SizeX = EdMode->UISettings->AlphaTextureSizeX;
-			int32 SizeY = EdMode->UISettings->AlphaTextureSizeY;
-			float MaxSize = 2.0f * FMath::Sqrt(FMath::Square(Radius) / 2.0f);
-			float AlphaBrushScale = MaxSize / (float)FMath::Max<int32>(SizeX, SizeY);
-			const float BrushAngle = EdMode->UISettings->bAlphaBrushAutoRotate ? LastMouseAngle : FMath::DegreesToRadians(EdMode->UISettings->AlphaBrushRotation);
+			float AlphaBrushScale;
+			float Radius;
+			int32 SizeX;
+			int32 SizeY;
+			bool bIsValid = ComputeAlphaBrushScaleAndRadius(AlphaBrushScale, Radius, SizeX, SizeY);
+			check(bIsValid && !FMath::IsNearlyZero(AlphaBrushScale) && !FMath::IsNearlyZero(Radius) && (SizeX > 0) && (SizeY > 0)); // See CanPaint() function above : if bCanPaint, AlphaBrushScale should be non-zero
 
+			const float BrushAngle = EdMode->UISettings->bAlphaBrushAutoRotate ? LastMouseAngle : FMath::DegreesToRadians(EdMode->UISettings->AlphaBrushRotation);
 			FIntRect Bounds;
 			Bounds.Min.X = FMath::FloorToInt(LastMousePosition.X - Radius);
 			Bounds.Min.Y = FMath::FloorToInt(LastMousePosition.Y - Radius);
 			Bounds.Max.X = FMath::CeilToInt( LastMousePosition.X + Radius);
 			Bounds.Max.Y = FMath::CeilToInt( LastMousePosition.Y + Radius);
 
-
 			// Clamp to landscape bounds
 			int32 MinX, MaxX, MinY, MaxY;
 			if (!ensure(LandscapeInfo->GetLandscapeExtent(MinX, MinY, MaxX, MaxY)))
@@ -1328,7 +1443,7 @@ public:
 							if (EdMode->CurrentTool && EdMode->CurrentTool->GetToolType() != ELandscapeToolType::Mask
 								&& EdMode->UISettings->bUseSelectedRegion && LandscapeInfo->SelectedRegion.Num() > 0)
 							{
-								float MaskValue = LandscapeInfo->SelectedRegion.FindRef(FIntPoint(X, Y));
+								float MaskValue = LandscapeInfo->SelectedRegion.FindRef(VertexKey);
 								if (EdMode->UISettings->bUseNegativeMask)
 								{
 									MaskValue = 1.0f - MaskValue;
@@ -1348,7 +1463,7 @@ public:
 
 	virtual void MouseMove(float LandscapeX, float LandscapeY) override
 	{
-		FLandscapeBrushAlphaBase::MouseMove(LandscapeX, LandscapeY);
+		Super::MouseMove(LandscapeX, LandscapeY);
 
 		if (EdMode->UISettings->bAlphaBrushAutoRotate)
 		{
@@ -1370,27 +1485,35 @@ public:
 
 	virtual void Tick(FEditorViewportClient* ViewportClient, float DeltaTime) override
 	{
-		FLandscapeBrushCircle::Tick(ViewportClient, DeltaTime);
+		Super::Tick(ViewportClient, DeltaTime);
 
 		ULandscapeInfo* LandscapeInfo = EdMode->CurrentToolTarget.LandscapeInfo.Get();
 		if (LandscapeInfo)
 		{
-			float ScaleXY = FMath::Abs(LandscapeInfo->DrawScale.X);
-			int32 SizeX = EdMode->UISettings->AlphaTextureSizeX;
-			int32 SizeY = EdMode->UISettings->AlphaTextureSizeY;
-			float Radius = EdMode->UISettings->GetCurrentToolBrushRadius() / ScaleXY;
-			float MaxSize = 2.0f * FMath::Sqrt(FMath::Square(Radius) / 2.0f);
-			float AlphaBrushScale = MaxSize / (float)FMath::Max<int32>(SizeX, SizeY);
+			FVector2f Scale(1.0f, 1.0f);
+			float AlphaBrushScale;
+			float Radius;
+			int32 SizeX;
+			int32 SizeY;
+			bool bIsValid = ComputeAlphaBrushScaleAndRadius(AlphaBrushScale, Radius, SizeX, SizeY);
+			if (bIsValid)
+			{
+				check(!FMath::IsNearlyZero(AlphaBrushScale) && !FMath::IsNearlyZero(Radius) && (SizeX > 0) && (SizeY > 0)); // See CanPaint() function above : if bCanPaint, AlphaBrushScale should be non-zero
+				Scale = FVector2f(1.0f / (AlphaBrushScale * SizeX), 1.0f / (AlphaBrushScale * SizeY));
+			}
 
 			FLinearColor BrushScaleRot(
-				1.0f / (AlphaBrushScale * SizeX),
-				1.0f / (AlphaBrushScale * SizeY),
+				Scale.X,
+				Scale.Y,
 				0.0f,
 				EdMode->UISettings->bAlphaBrushAutoRotate ? LastMouseAngle : FMath::DegreesToRadians(EdMode->UISettings->AlphaBrushRotation)
 				);
 
-			int32 Channel = EdMode->UISettings->AlphaTextureChannel;
-			FLinearColor AlphaTextureMask(Channel == 0 ? 1 : 0, Channel == 1 ? 1 : 0, Channel == 2 ? 1 : 0, Channel == 3 ? 1 : 0);
+			FLinearColor AlphaTextureMask(EdMode->UISettings->AlphaTextureChannel == ELandscapeTextureColorChannel::Red ? 1 : 0,
+				EdMode->UISettings->AlphaTextureChannel == ELandscapeTextureColorChannel::Green ? 1 : 0,
+				EdMode->UISettings->AlphaTextureChannel == ELandscapeTextureColorChannel::Blue ? 1 : 0,
+				EdMode->UISettings->AlphaTextureChannel == ELandscapeTextureColorChannel::Alpha ? 1 : 0);
+
 
 			for (const auto& BrushMaterialInstancePair : BrushMaterialInstanceMap)
 			{
diff --git a/Engine/Source/Editor/LandscapeEditor/Private/LandscapeEditorDetailCustomization_AlphaBrush.cpp b/Engine/Source/Editor/LandscapeEditor/Private/LandscapeEditorDetailCustomization_AlphaBrush.cpp
--- a/Engine/Source/Editor/LandscapeEditor/Private/LandscapeEditorDetailCustomization_AlphaBrush.cpp
+++ b/Engine/Source/Editor/LandscapeEditor/Private/LandscapeEditorDetailCustomization_AlphaBrush.cpp
@@ -61,8 +61,12 @@ public:
 		TextureChannel = InArgs._TextureChannel;
 
 		FAssetToolsModule& AssetToolsModule = FModuleManager::LoadModuleChecked<FAssetToolsModule>(TEXT("AssetTools"));
-		TWeakPtr<IAssetTypeActions> AssetTypeActions = AssetToolsModule.Get().GetAssetTypeActionsForClass(Texture.Get()->GetClass());
-		FLinearColor AssetColor = AssetTypeActions.Pin()->GetTypeColor();
+		FLinearColor AssetColor(ForceInitToZero);
+		if (UTexture2D* Texture2D = Texture.Get())
+		{
+			TWeakPtr<IAssetTypeActions> AssetTypeActions = AssetToolsModule.Get().GetAssetTypeActionsForClass(Texture2D->GetClass());
+			AssetColor = AssetTypeActions.Pin()->GetTypeColor();
+		}
 
 		ChildSlot
 		[
@@ -157,19 +161,20 @@ void FTextureMaskThumbnailViewportClient::Draw(FViewport* Viewport, FCanvas* Can
 		return;
 	}
 
-	UTexture2D* Texture = PinnedParent->Texture.Get();
-
 	Canvas->Clear( FLinearColor::Black);
 
-	// Fully stream in the texture before drawing it.
-	FTextureCompilingManager::Get().FinishCompilation({ Texture });
-	Texture->SetForceMipLevelsToBeResident(30.0f);
-	Texture->WaitForStreaming();
-
-	//Draw the selected texture, uses ColourChannelBlend mode parameter to filter colour channels and apply grayscale
-	FCanvasTileItem TileItem( FVector2D( 0.0f, 0.0f ), Texture->GetResource(), Viewport->GetSizeXY(), FLinearColor::White );
-	TileItem.BlendMode = (ESimpleElementBlendMode)(SE_BLEND_RGBA_MASK_START + (1<<PinnedParent->TextureChannel.Get()) + 16);
-	Canvas->DrawItem( TileItem );
+	if (UTexture2D* Texture = PinnedParent->Texture.Get())
+	{
+		// Fully stream in the texture before drawing it.
+		FTextureCompilingManager::Get().FinishCompilation({ Texture });
+		Texture->SetForceMipLevelsToBeResident(30.0f);
+		Texture->WaitForStreaming();
+
+		//Draw the selected texture, uses ColourChannelBlend mode parameter to filter colour channels and apply grayscale
+		FCanvasTileItem TileItem( FVector2D( 0.0f, 0.0f ), Texture->GetResource(), Viewport->GetSizeXY(), FLinearColor::White );
+		TileItem.BlendMode = (ESimpleElementBlendMode)(SE_BLEND_RGBA_MASK_START + (1<<PinnedParent->TextureChannel.Get()) + 16);
+		Canvas->DrawItem( TileItem );
+	}
 }
 
 //////////////////////////////////////////////////////////////////////////
diff --git a/Engine/Source/Editor/LandscapeEditor/Private/LandscapeEditorObject.cpp b/Engine/Source/Editor/LandscapeEditor/Private/LandscapeEditorObject.cpp
--- a/Engine/Source/Editor/LandscapeEditor/Private/LandscapeEditorObject.cpp
+++ b/Engine/Source/Editor/LandscapeEditor/Private/LandscapeEditorObject.cpp
@@ -14,9 +14,13 @@
 #include "LandscapeMaterialInstanceConstant.h"
 #include "Misc/ConfigCacheIni.h"
 #include "EngineUtils.h"
+#include "RenderUtils.h"
 #include "Misc/MessageDialog.h"
 
-//#define LOCTEXT_NAMESPACE "LandscapeEditor"
+static TAutoConsoleVariable<bool> CVarLandscapeSimulateAlphaBrushTextureLoadFailure(
+	TEXT("landscape.SimulateAlphaBrushTextureLoadFailure"),
+	false,
+	TEXT("Debug utility to simulate a loading failure (e.g. invalid source data, which can happen in cooked editor or with a badly virtualized texture) when loading the alpha brush texture"));
 
 ULandscapeEditorObject::ULandscapeEditorObject(const FObjectInitializer& ObjectInitializer)
 	: Super(ObjectInitializer)
@@ -105,7 +109,7 @@ ULandscapeEditorObject::ULandscapeEditorObject(const FObjectInitializer& ObjectI
 	, bUseWorldSpacePatternBrush(false)
 	, WorldSpacePatternBrushSettings(FVector2D::ZeroVector, 0.0f, false, 3200)
 	, AlphaTexture(nullptr)
-	, AlphaTextureChannel(EColorChannel::Red)
+	, AlphaTextureChannel(ELandscapeTextureColorChannel::Red)
 	, AlphaTextureSizeX(1)
 	, AlphaTextureSizeY(1)
 
@@ -126,10 +130,7 @@ ULandscapeEditorObject::ULandscapeEditorObject(const FObjectInitializer& ObjectI
 	};
 	static FConstructorStatics ConstructorStatics;
 
-	if (!IsTemplate())
-	{
-		SetAlphaTexture(ConstructorStatics.AlphaTexture.Object, AlphaTextureChannel);
-	}
+	SetAlphaTexture(ConstructorStatics.AlphaTexture.Object, AlphaTextureChannel);
 }
 
 void ULandscapeEditorObject::PostEditChangeProperty(FPropertyChangedEvent& PropertyChangedEvent)
@@ -222,11 +223,16 @@ void ULandscapeEditorObject::Load()
 	GConfig->GetBool(TEXT("LandscapeEdit"), TEXT("WorldSpacePatternBrushSettings.bCenterTextureOnOrigin"), WorldSpacePatternBrushSettings.bCenterTextureOnOrigin, GEditorPerProjectIni);
 	GConfig->GetFloat(TEXT("LandscapeEdit"), TEXT("WorldSpacePatternBrushSettings.RepeatSize"), WorldSpacePatternBrushSettings.RepeatSize, GEditorPerProjectIni);
 	FString AlphaTextureName = (AlphaTexture != nullptr) ? AlphaTexture->GetPathName() : FString();
-	int32 InAlphaTextureChannel = AlphaTextureChannel;
+	int32 InAlphaTextureChannel = static_cast<int32>(AlphaTextureChannel);
 	GConfig->GetString(TEXT("LandscapeEdit"), TEXT("AlphaTextureName"), AlphaTextureName, GEditorPerProjectIni);
 	GConfig->GetInt(TEXT("LandscapeEdit"), TEXT("AlphaTextureChannel"), InAlphaTextureChannel, GEditorPerProjectIni);
-	AlphaTextureChannel = (EColorChannel::Type)InAlphaTextureChannel;
-	SetAlphaTexture(LoadObject<UTexture2D>(nullptr, *AlphaTextureName, nullptr, LOAD_NoWarn), AlphaTextureChannel);
+	AlphaTextureChannel = (ELandscapeTextureColorChannel)InAlphaTextureChannel;
+	UTexture2D* LoadedTexture = LoadObject<UTexture2D>(nullptr, *AlphaTextureName, nullptr, LOAD_NoWarn);
+	if ((LoadedTexture == nullptr) && !AlphaTextureName.IsEmpty())
+	{
+		UE_LOG(LogLandscapeTools, Error, TEXT("Cannot load alpha texture (%s)"), *AlphaTextureName);
+	}
+	SetAlphaTexture(LoadedTexture, AlphaTextureChannel);
 
 	int32 InFlattenMode = (int32)ELandscapeToolFlattenMode::Both;
 	GConfig->GetInt(TEXT("LandscapeEdit"), TEXT("FlattenMode"), InFlattenMode, GEditorPerProjectIni);
@@ -503,78 +509,115 @@ void ULandscapeEditorObject::SetGizmoSnapMode(ELandscapeGizmoSnapType InSnapMode
 	}
 }
 
-void ULandscapeEditorObject::SetAlphaTexture(UTexture2D* InTexture, EColorChannel::Type InTextureChannel)
+bool ULandscapeEditorObject::LoadAlphaTextureSourceData(UTexture2D* InTexture, TArray<uint8>& OutSourceData, int32& OutSourceDataSizeX, int32& OutSourceDataSizeY, ELandscapeTextureColorChannel& InOutTextureChannel)
 {
-	TArray64<uint8> NewTextureData;
-	UTexture2D* NewAlphaTexture = nullptr;
-	int32 NumChannels = 0;
+	check(InTexture != nullptr);
 
-	// Validate that the input texture is valid, if not, we'll display an error message and use the default brush alpha texture : 
-	if (InTexture != nullptr)
+	if (InTexture && InTexture->Source.IsValid() 
+		&& !CVarLandscapeSimulateAlphaBrushTextureLoadFailure.GetValueOnGameThread()) // For debug purposes, we can also simulate a loading failure for the alpha brush texture
 	{
-		if (!InTexture->Source.IsValid())
+		FImage SourceImage;
+		if (InTexture->Source.GetMipImage(SourceImage, 0) && (SourceImage.Format != ERawImageFormat::Invalid))
 		{
-			UE_LOG(LogLandscapeTools, Error, TEXT("Invalid source data detected for texture (%s), the default AlphaTexture (%s) will be used."), *InTexture->GetPathName(), *GetClass()->GetDefaultObject<ULandscapeEditorObject>()->AlphaTexture->GetPathName());
-		}
-		else
-		{
-			// Try to read the new texture data now : 
-			const bool bSourceDataIsG8 = (InTexture->Source.GetFormat() == TSF_G8);
-			NumChannels = bSourceDataIsG8 ? 1 : 4;
-			InTexture->Source.GetMipData(NewTextureData, 0);
-			if (NewTextureData.Num() == (NumChannels * InTexture->Source.GetSizeX() * InTexture->Source.GetSizeY()))
+			OutSourceDataSizeX = SourceImage.SizeX;
+			OutSourceDataSizeY = SourceImage.SizeY;
+			const int32 NumPixels = OutSourceDataSizeX * OutSourceDataSizeY;
+
+			// Handle the case where we're being asked to sample from a channel that is non-existent in the source image :
+			EPixelFormat PixelFormat = InTexture->GetPixelFormat();
+			EPixelFormatChannelFlags ValidTextureChannels = GetPixelFormatValidChannels(PixelFormat);
+			if (((InOutTextureChannel == ELandscapeTextureColorChannel::Green && !EnumHasAnyFlags(ValidTextureChannels, EPixelFormatChannelFlags::G)))
+				|| ((InOutTextureChannel == ELandscapeTextureColorChannel::Blue && !EnumHasAnyFlags(ValidTextureChannels, EPixelFormatChannelFlags::B)))
+				|| ((InOutTextureChannel == ELandscapeTextureColorChannel::Alpha && !EnumHasAnyFlags(ValidTextureChannels, EPixelFormatChannelFlags::A))))
 			{
-				// Valid new texture
-				NewAlphaTexture = InTexture;
+				// Fallback to Red
+				InOutTextureChannel = ELandscapeTextureColorChannel::Red;
 			}
-			else
+
+			// Convert/expand the image to BGRA8 : 
+			SourceImage.ChangeFormat(ERawImageFormat::BGRA8, EGammaSpace::Linear);
+
+			const int32 FinalDataSizeBytes = NumPixels;
+			check(SourceImage.RawData.Num() == FinalDataSizeBytes * 4);
+			OutSourceData.SetNum(FinalDataSizeBytes);
+
+			uint8* SrcPtr = SourceImage.RawData.GetData();
+			// Properly offset the source data as we're reading a single channel from a BGRA8 source :
+			switch (InOutTextureChannel)
 			{
-				UE_LOG(LogLandscapeTools, Error, TEXT("Invalid data size detected for texture (%s), the default AlphaTexture (%s) will be used."), *InTexture->GetPathName(), *GetClass()->GetDefaultObject<ULandscapeEditorObject>()->AlphaTexture->GetPathName());
+			case ELandscapeTextureColorChannel::Blue:
+				SrcPtr += 0;
+				break;
+			case ELandscapeTextureColorChannel::Green:
+				SrcPtr += 1;
+				break;
+			case ELandscapeTextureColorChannel::Red:
+				SrcPtr += 2;
+				break;
+			case ELandscapeTextureColorChannel::Alpha:
+				SrcPtr += 3;
+				break;
+			default:
+				check(false);
+				break;
 			}
-		}
-	}
 
-	// Load fallback if there's no texture or valid data
-	if (NewAlphaTexture == nullptr)
-	{
-		UTexture2D* DefaultAlphaTexture = GetClass()->GetDefaultObject<ULandscapeEditorObject>()->AlphaTexture;
-		check((DefaultAlphaTexture != nullptr) && DefaultAlphaTexture->Source.IsValid()); // The default texture should always be valid
-		DefaultAlphaTexture->Source.GetMipData(NewTextureData, 0);
-		NewAlphaTexture = DefaultAlphaTexture;
-		const bool bSourceDataIsG8 = (DefaultAlphaTexture->Source.GetFormat() == TSF_G8);
-		NumChannels = bSourceDataIsG8 ? 1 : 4;
+			for (int32 Index = 0; Index < NumPixels; Index++, SrcPtr += 4)
+			{
+				OutSourceData[Index] = *SrcPtr;
+			}
+
+			return true;
+		}
 	}
 
-	check((NewAlphaTexture != nullptr) && !NewTextureData.IsEmpty() && (NumChannels > 0));
+	return false;
+}
 
-	AlphaTexture = NewAlphaTexture;
-	AlphaTextureSizeX = NewAlphaTexture->Source.GetSizeX();
-	AlphaTextureSizeY = NewAlphaTexture->Source.GetSizeY();
-	AlphaTextureChannel = (NumChannels == 1) ? EColorChannel::Red : InTextureChannel;
-	AlphaTextureData.Empty(AlphaTextureSizeX * AlphaTextureSizeY);
+void ULandscapeEditorObject::SetAlphaTexture(UTexture2D* InTexture, ELandscapeTextureColorChannel InTextureChannel)
+{
+	AlphaTexture = nullptr;
+	AlphaTextureChannel = InTextureChannel;
+	AlphaTextureSizeX = 0;
+	AlphaTextureSizeY = 0;
 
-	uint8* SrcPtr;
-	switch (AlphaTextureChannel)
+	UTexture2D* DefaultAlphaTexture = GetClass()->GetDefaultObject<ULandscapeEditorObject>()->AlphaTexture;
+	// Try to read the texture data from the specified texture if any : 
+	if (InTexture != nullptr)
 	{
-	case 1:
-		SrcPtr = &((FColor*)NewTextureData.GetData())->G;
-		break;
-	case 2:
-		SrcPtr = &((FColor*)NewTextureData.GetData())->B;
-		break;
-	case 3:
-		SrcPtr = &((FColor*)NewTextureData.GetData())->A;
-		break;
-	default:
-		SrcPtr = &((FColor*)NewTextureData.GetData())->R;
-		break;
+		if (LoadAlphaTextureSourceData(InTexture, AlphaTextureData, AlphaTextureSizeX, AlphaTextureSizeY, AlphaTextureChannel))
+		{
+			AlphaTexture = InTexture;
+		}
+		else
+		{
+			UE_LOG(LogLandscapeTools, Error, TEXT("Invalid source data detected for texture (%s), the default AlphaTexture (%s) will be used."), *InTexture->GetPathName(), DefaultAlphaTexture ? *DefaultAlphaTexture->GetPathName() : TEXT("None"));
+		}
 	}
 
-	for (int32 i = 0; i < AlphaTextureSizeX * AlphaTextureSizeY; i++)
+	if (AlphaTexture == nullptr)
 	{
-		AlphaTextureData.Add(*SrcPtr);
-		SrcPtr += NumChannels;
+		if (DefaultAlphaTexture == nullptr)
+		{
+			UE_LOG(LogLandscapeTools, Error, TEXT("No default AlphaTexture specified : the alpha brush won't work as expected."));
+		}
+		else if (LoadAlphaTextureSourceData(DefaultAlphaTexture, AlphaTextureData, AlphaTextureSizeX, AlphaTextureSizeY, AlphaTextureChannel))
+		{
+			AlphaTexture = DefaultAlphaTexture;
+		}
+		else
+		{
+			UE_LOG(LogLandscapeTools, Error, TEXT("Invalid source data detected for default AlphaTexture (%s)"), *DefaultAlphaTexture->GetPathName());
+		}
 	}
+
+	// If the AlphaTexture was successfully loaded, all read data should be valid :
+	check ((AlphaTexture == nullptr) || HasValidAlphaTextureData());
+}
+
+bool ULandscapeEditorObject::HasValidAlphaTextureData() const
+{
+	return ((AlphaTextureSizeX > 0) && (AlphaTextureSizeY > 0) && !AlphaTextureData.IsEmpty());
 }
 
 void ULandscapeEditorObject::ChooseBestComponentSizeForImport()
diff --git a/Engine/Source/Editor/LandscapeEditor/Public/LandscapeEditorObject.h b/Engine/Source/Editor/LandscapeEditor/Public/LandscapeEditorObject.h
--- a/Engine/Source/Editor/LandscapeEditor/Public/LandscapeEditorObject.h
+++ b/Engine/Source/Editor/LandscapeEditor/Public/LandscapeEditorObject.h
@@ -127,10 +127,9 @@ enum class ELandscapeConvertMode : int8
 	Resample = 2,
 };
 
-UENUM()
 namespace EColorChannel
 {
-	enum Type : int
+	enum UE_DEPRECATED(5.2, "Use ELandscapeTextureColorChannel") Type : int
 	{
 		Red,
 		Green,
@@ -139,6 +138,15 @@ namespace EColorChannel
 	};
 }
 
+UENUM()
+enum class ELandscapeTextureColorChannel : int32
+{
+	Red,
+	Green,
+	Blue,
+	Alpha,
+};
+
 UENUM()
 enum class ELandscapeMirrorOperation : uint8
 {
@@ -647,7 +655,7 @@ public:
 
 	// Channel of Mask Texture to use
 	UPROPERTY(Category="Brush Settings", EditAnywhere, NonTransactional, meta=(DisplayName="Texture Channel", ShowForBrushes="BrushSet_Alpha,BrushSet_Pattern"))
-	TEnumAsByte<EColorChannel::Type> AlphaTextureChannel;
+	ELandscapeTextureColorChannel AlphaTextureChannel;
 
 	UPROPERTY(NonTransactional)
 	int32 AlphaTextureSizeX;
@@ -696,7 +704,8 @@ public:
 	void SetPasteMode(ELandscapeToolPasteMode InPasteMode);
 
 	// Alpha/Pattern Brush
-	void SetAlphaTexture(UTexture2D* InTexture, EColorChannel::Type InTextureChannel);
+	void SetAlphaTexture(UTexture2D* InTexture, ELandscapeTextureColorChannel InTextureChannel);
+	bool HasValidAlphaTextureData() const;
 
 	// New Landscape
 	FString LastImportPath;
@@ -790,4 +799,6 @@ private:
 		check(ParentMode);
 		return !(ParentMode->CurrentToolTarget.TargetType == ELandscapeToolTargetType::Heightmap || ParentMode->CurrentToolTarget.TargetType == ELandscapeToolTargetType::Visibility);
 	}
+
+	static bool LoadAlphaTextureSourceData(UTexture2D* InTexture, TArray<uint8>& OutSourceData, int32& OutSourceDataSizeX, int32& OutSourceDataSizeY, ELandscapeTextureColorChannel& InOutTextureChannel);
 };
diff --git a/Engine/Source/Runtime/RenderCore/Private/RenderUtils.cpp b/Engine/Source/Runtime/RenderCore/Private/RenderUtils.cpp
--- a/Engine/Source/Runtime/RenderCore/Private/RenderUtils.cpp
+++ b/Engine/Source/Runtime/RenderCore/Private/RenderUtils.cpp
@@ -268,8 +268,8 @@ EPixelFormatChannelFlags GetPixelFormatValidChannels(EPixelFormat InPixelFormat)
 		EPixelFormatChannelFlags::None,		// PF_Unknown,
 		EPixelFormatChannelFlags::RGBA,		// PF_A32B32G32R32F
 		EPixelFormatChannelFlags::RGBA,		// PF_B8G8R8A8
-		EPixelFormatChannelFlags::G,		// PF_G8
-		EPixelFormatChannelFlags::G,		// PF_G16
+		EPixelFormatChannelFlags::R,		// PF_G8
+		EPixelFormatChannelFlags::R,		// PF_G16
 		EPixelFormatChannelFlags::RGB,		// PF_DXT1
 		EPixelFormatChannelFlags::RGBA,		// PF_DXT3
 		EPixelFormatChannelFlags::RGBA,		// PF_DXT5
