LandscapeEditor: Disable Retopo tool in restrictive mode

diff --git a/Engine/Source/Editor/LandscapeEditor/Private/SLandscapeEditor.cpp b/Engine/Source/Editor/LandscapeEditor/Private/SLandscapeEditor.cpp
--- a/Engine/Source/Editor/LandscapeEditor/Private/SLandscapeEditor.cpp
+++ b/Engine/Source/Editor/LandscapeEditor/Private/SLandscapeEditor.cpp
@@ -217,7 +217,10 @@ void FLandscapeToolKit::BuildToolPalette(FName PaletteName, class FToolBarBuilde
 		ToolBarBuilder.AddToolBarButton(Commands.ErosionTool);
 		ToolBarBuilder.AddToolBarButton(Commands.HydroErosionTool);
 		ToolBarBuilder.AddToolBarButton(Commands.NoiseTool);
-		ToolBarBuilder.AddToolBarButton(Commands.RetopologizeTool);
+		if (!Settings->InRestrictiveMode())
+		{
+			ToolBarBuilder.AddToolBarButton(Commands.RetopologizeTool);
+		}
 		ToolBarBuilder.AddToolBarButton(Commands.VisibilityTool);
 
 		if (Settings->AreBlueprintToolsAllowed())
