* Added CPU traces to LandscapeTrace
* Fixed important slowdown when using a landscape brush with a large radius : sphere sweeps perform very poorly on landscapes with high resolutions. Replaced with a bounded series of concentric raycasts (hundreds of ms -> 1 or less ms in the worse case) and adjusted hit point extrapolation in order to be more consistent

diff --git a/Engine/Source/Editor/LandscapeEditor/Private/LandscapeEdMode.cpp b/Engine/Source/Editor/LandscapeEditor/Private/LandscapeEdMode.cpp
--- a/Engine/Source/Editor/LandscapeEditor/Private/LandscapeEdMode.cpp
+++ b/Engine/Source/Editor/LandscapeEditor/Private/LandscapeEdMode.cpp
@@ -38,6 +38,7 @@
 #include "LandscapeImportHelper.h"
 #include "LandscapeConfigHelper.h"
 #include "UObject/ObjectSaveContext.h"
+#include "Math/TransformCalculus3D.h"
 
 //Slate dependencies
 #include "Misc/FeedbackContext.h"
@@ -1154,6 +1155,7 @@ bool FEdModeLandscape::ProcessLandscapeTraceHits(const TArray<FHitResult>& InRes
 
 bool FEdModeLandscape::LandscapeTrace(const FVector& InRayOrigin, const FVector& InRayEnd, const FVector& InDirection, FVector& OutHitLocation)
 {
+	TRACE_CPUPROFILER_EVENT_SCOPE(FEdModeLandscape_LandscapeTrace);
 	if (!CurrentTool || !CurrentToolTarget.LandscapeInfo.IsValid())
 	{
 		return false;
@@ -1166,93 +1168,143 @@ bool FEdModeLandscape::LandscapeTrace(const FVector& InRayOrigin, const FVector&
 	UWorld* World = GetWorld();
 
 	// Check Tool Trace first
-	if (CurrentTool->HitTrace(Start, End, OutHitLocation))
 	{
-		return true;
+		TRACE_CPUPROFILER_EVENT_SCOPE(ToolTrace);
+		if (CurrentTool->HitTrace(Start, End, OutHitLocation))
+		{
+			return true;
+		}
 	}
 	
 	TArray<FHitResult> Results;
-	// Each landscape component has 2 collision shapes, 1 of them is specific to landscape editor
-	// Trace only ECC_Visibility channel, so we do hit only Editor specific shape
-	if (World->LineTraceMultiByObjectType(Results, Start, End, FCollisionObjectQueryParams(ECollisionChannel::ECC_Visibility), FCollisionQueryParams(SCENE_QUERY_STAT(LandscapeTrace), true)))
 	{
-		if (FProcessLandscapeTraceHitsResult ProcessResult; ProcessLandscapeTraceHits(Results, ProcessResult))
+		TRACE_CPUPROFILER_EVENT_SCOPE(LineTrace);
+		// Each landscape component has 2 collision shapes, 1 of them is specific to landscape editor
+		// Trace only ECC_Visibility channel, so we do hit only Editor specific shape
+		if (World->LineTraceMultiByObjectType(Results, Start, End, FCollisionObjectQueryParams(ECollisionChannel::ECC_Visibility), FCollisionQueryParams(SCENE_QUERY_STAT(LandscapeTrace), true)))
 		{
-			OutHitLocation = ProcessResult.LandscapeProxy->LandscapeActorToWorld().InverseTransformPosition(ProcessResult.HitLocation);
+			if (FProcessLandscapeTraceHitsResult ProcessResult; ProcessLandscapeTraceHits(Results, ProcessResult))
+			{
+				OutHitLocation = ProcessResult.LandscapeProxy->LandscapeActorToWorld().InverseTransformPosition(ProcessResult.HitLocation);
 			
-			UE_VLOG_SEGMENT_THICK(World, LogLandscapeEdMode, VeryVerbose, InRayOrigin, InRayOrigin + 20000.0f * InDirection, FColor(100,255,100), 4, TEXT("landscape:ray-hit"));
-			UE_VLOG_LOCATION(World, LogLandscapeEdMode, VeryVerbose,  ProcessResult.LandscapeProxy->LandscapeActorToWorld().TransformPosition(OutHitLocation), 20.0,  FColor(100,100,255), TEXT("landscape:point-hit"));
-			return true;
+				UE_VLOG_SEGMENT_THICK(World, LogLandscapeEdMode, VeryVerbose, InRayOrigin, InRayEnd, FColor(100,255,100), 4, TEXT("landscape:ray-hit"));
+				UE_VLOG_LOCATION(World, LogLandscapeEdMode, VeryVerbose,  ProcessResult.LandscapeProxy->LandscapeActorToWorld().TransformPosition(OutHitLocation), 20.0,  FColor(100,100,255), TEXT("landscape:point-hit"));
+				return true;
+			}
 		}
 	}
 
-	UE_VLOG_SEGMENT_THICK(World, LogLandscapeEdMode, VeryVerbose, InRayOrigin, InRayOrigin + 20000.0f * InDirection, FColor(255,100,100), 4, TEXT("landscape:ray-miss"));
+	UE_VLOG_SEGMENT_THICK(World, LogLandscapeEdMode, VeryVerbose, InRayOrigin, InRayEnd, FColor(255,100,100), 2, TEXT("landscape:ray-miss"));
 		
 	if (CurrentTool->UseSphereTrace())
 	{
+		TRACE_CPUPROFILER_EVENT_SCOPE(SphereLikeTrace);
+
 		// If there is no landscape directly under the mouse search for a landscape collision
-		// under the shape of the brush.
-		FCollisionShape SphereShape;
-		SphereShape.SetSphere(UISettings->GetCurrentToolBrushRadius());
-		if (World->SweepMultiByObjectType(Results, Start, End, FQuat::Identity, FCollisionObjectQueryParams(ECollisionChannel::ECC_Visibility), SphereShape, FCollisionQueryParams(SCENE_QUERY_STAT(LandscapeTrace), true)))
+		// under the shape of the brush. 
+		FVector UpVector = FVector::UpVector;
+		if (FMath::IsNearlyEqual(FMath::Abs(FVector::DotProduct(UpVector, InDirection.GetSafeNormal(SMALL_NUMBER, UpVector))), 1.0f))
 		{
-			if (FProcessLandscapeTraceHitsResult SweepProcessResult; ProcessLandscapeTraceHits(Results, SweepProcessResult))
-			{
-				UE_VLOG_LOCATION(World, LogLandscapeEdMode, VeryVerbose, SweepProcessResult.HitLocation, UISettings->GetCurrentToolBrushRadius(), FColor(255, 100, 100), TEXT("landscape:sweep-hit-location"));
-				FSphere HitSphere(SweepProcessResult.HitLocation, UISettings->GetCurrentToolBrushRadius());
+			/* If we're too close to being colinear then pick another vector to do the cross products with. */
+			UpVector = FVector::RightVector;
+		}
+		FVector RightVector = InDirection.Cross(UpVector).GetUnsafeNormal();
+		UpVector = RightVector.Cross(InDirection);
+		FMatrix RayTransform(InDirection, RightVector, UpVector, FVector::ZeroVector);
+		
+		// Add a slight offset so that we don't end up with the landscape brush appearing slightly on the border of the landscape when
+		//  the "projected' brush center is at a distance from the border that is close to the radius : 
+		const double AdditionalRadius = 50.0;
+		double BrushRadius = UISettings->GetCurrentToolBrushRadius() + AdditionalRadius;
 
-				float MeanHeight = 0.0f;
-				int32 Count = 0;
-				const int32 NumHeightSamples = 16;
+		auto LineTraceAroundCenter = [this, World, &Start, &End, &RayTransform] (float InAngle, float InRadius) -> TOptional<FProcessLandscapeTraceHitsResult>
+		{
+			TArray<FHitResult> Results;
+			FVector Offset(0.0, InRadius * FMath::Cos(InAngle), InRadius * FMath::Sin(InAngle));
+			FVector AdjustedRayStart = TransformPoint(RayTransform, Offset) + Start;
+			FVector AdjustedRayEnd = TransformPoint(RayTransform, Offset) + End;
+			double HitDistance = -1.0;
 
-				for (int32 Y = 0; Y < NumHeightSamples; ++Y)
+			if (World->LineTraceMultiByObjectType(Results, AdjustedRayStart, AdjustedRayEnd, FCollisionObjectQueryParams(ECollisionChannel::ECC_Visibility), FCollisionQueryParams(SCENE_QUERY_STAT(LandscapeTrace), true)))
+			{
+				if (FProcessLandscapeTraceHitsResult ProcessResult; ProcessLandscapeTraceHits(Results, ProcessResult))
 				{
-					float HY = (Y / (float)NumHeightSamples) * HitSphere.W * 2.0f + HitSphere.Center.Y - HitSphere.W;
-					for (int32 X = 0; X < NumHeightSamples; ++X)
-					{
-						float HX = (X / (float)NumHeightSamples) * HitSphere.W * 2.0f + HitSphere.Center.X - HitSphere.W;
-
-						FVector HeightSampleLocation(HX, HY, HitSphere.Center.Z);
-
-						TOptional<float> Height = SweepProcessResult.LandscapeProxy->GetHeightAtLocation(HeightSampleLocation, EHeightfieldSource::Editor);
+					UE_VLOG_SEGMENT_THICK(World, LogLandscapeEdMode, VeryVerbose, AdjustedRayStart, AdjustedRayEnd, FColor(100, 255, 100), 4, TEXT("landscape:ray-hit"));
+					UE_VLOG_LOCATION(World, LogLandscapeEdMode, VeryVerbose, ProcessResult.HitLocation, 20.0, FColor(100, 100, 255), TEXT("landscape:point-hit"));
 
-						if (!Height.IsSet())
-						{
-							continue;
-						}
-
-						UE_VLOG_LOCATION(World, LogLandscapeEdMode, VeryVerbose, FVector(HX, HY, Height.GetValue()), 2.0f, FColor(100, 100, 255), TEXT(""));
+					HitDistance = (ProcessResult.HitLocation - Start).Length();
+					return ProcessResult;
+				}
+			}
+			else
+			{
+				UE_VLOG_SEGMENT_THICK(World, LogLandscapeEdMode, VeryVerbose, AdjustedRayStart, AdjustedRayEnd, FColor(255, 100, 100), 2, TEXT("landscape:ray-miss"));
+			}
+			return TOptional<FProcessLandscapeTraceHitsResult>();
+		};
 
+		// Don't use a sphere cast as it can get very expensive for large brush sizes. Use several concentric ray casts instead : 
+		const int32 MaxConcentricSamples = 10;
+		const double RadiusIncrement = 1024.0;
+		// Bound the number of concentric circles so that we don't have too many ray casts at a given time even if the brush radius is very large : 
+		int32 NumConcentricSamples = FMath::Clamp(static_cast<int32>(BrushRadius / RadiusIncrement), 1, MaxConcentricSamples);
+		double ActualRadiusIncrement = BrushRadius / NumConcentricSamples;
+
+		const double DistanceBetweenAngularSamples = 1024.0;
+		// Bound the number of angular samples so that we don't have too many ray casts at a given time even if the brush radius is very large : 
+		const int32 MinAngularSamples = 6;
+		const int32 MaxAngularSamples = 12;
+		
+		float MeanHeight = 0.0f;
+		int32 Count = 0;
+		TOptional<FProcessLandscapeTraceHitsResult> BestHitResult;
+		int32 ConcentricIndex = 0;
+		for (double Radius = ActualRadiusIncrement; Radius < (BrushRadius + UE_SMALL_NUMBER); Radius += ActualRadiusIncrement, ++ConcentricIndex)
+		{
+			double Circumference = TWO_PI * Radius;
+			int32 NumAngularSamples = FMath::Clamp(static_cast<int32>(Circumference / DistanceBetweenAngularSamples), MinAngularSamples, MaxAngularSamples);
+			double AngleIncrement = TWO_PI / NumAngularSamples;
+			// Alternate by half the angle offset in order to have a better coverage over the entire radius : 
+			double InitialAngle = (ConcentricIndex % 2) * AngleIncrement / 2.0;
+			for (int32 AngularSampleIndex = 0; AngularSampleIndex < NumAngularSamples; ++AngularSampleIndex)
+			{
+				if (TOptional<FProcessLandscapeTraceHitsResult> HitResult = LineTraceAroundCenter(AngularSampleIndex * AngleIncrement + InitialAngle, Radius); HitResult.IsSet())
+				{ 
+					// Compute the mean height in order to extrapolate the landscape hit to a "fake" landscape plane that extends beyond the borders : 
+					if (TOptional<float> Height = HitResult->LandscapeProxy->GetHeightAtLocation(HitResult->HitLocation, EHeightfieldSource::Editor); Height.IsSet())
+					{
 						MeanHeight += Height.GetValue();
 						Count++;
+						BestHitResult = HitResult;
 					}
-				}
-
-				FVector PointOnPlane(HitSphere.Center.X, HitSphere.Center.Y, HitSphere.Center.Z);
 
-				if (Count > 0)
-				{
-					MeanHeight /= (float)Count;
-					PointOnPlane.Z = MeanHeight;
 				}
+			}
+		}
 
-				UE_VLOG_LOCATION(World, LogLandscapeEdMode, VeryVerbose, PointOnPlane, 10.0, FColor(100, 100, 255), TEXT("landscape:point-on-plane"));
+		if (Count > 0)
+		{
+			check(BestHitResult.IsSet());
+			FVector PointOnPlane(BestHitResult->HitLocation);
+			MeanHeight /= (float)Count;
+			PointOnPlane.Z = MeanHeight;
 
-				if (FMath::Abs(FVector::DotProduct(InDirection, FVector::ZAxisVector)) < SMALL_NUMBER)
-				{
-					// the ray and plane are nearly parallel, and won't intersect (or if they do it will be wildly far away)
-					return false;
-				}
+			UE_VLOG_LOCATION(World, LogLandscapeEdMode, VeryVerbose, PointOnPlane, 10.0, FColor(100, 100, 255), TEXT("landscape:point-on-plane"));
 
-				const FPlane Plane(PointOnPlane, FVector::ZAxisVector);
-				FVector EstimatedHitLocation = FMath::RayPlaneIntersection(Start, InDirection, Plane);
-				check(!EstimatedHitLocation.ContainsNaN());
+			if (FMath::Abs(FVector::DotProduct(InDirection, FVector::ZAxisVector)) < SMALL_NUMBER)
+			{
+				// the ray and plane are nearly parallel, and won't intersect (or if they do it will be wildly far away)
+				return false;
+			}
 
-				UE_VLOG_LOCATION(World, LogLandscapeEdMode, VeryVerbose, EstimatedHitLocation, 10.0, FColor(100, 100, 255), TEXT("landscape:estimated-hit-location"));
+			const FPlane Plane(PointOnPlane, FVector::ZAxisVector);
+			FVector EstimatedHitLocation = FMath::RayPlaneIntersection(Start, InDirection, Plane);
+			check(!EstimatedHitLocation.ContainsNaN());
 
-				OutHitLocation = SweepProcessResult.LandscapeProxy->LandscapeActorToWorld().InverseTransformPosition(EstimatedHitLocation);
-				return true;
-			}
+			UE_VLOG_LOCATION(World, LogLandscapeEdMode, VeryVerbose, EstimatedHitLocation, 10.0, FColor(100, 100, 255), TEXT("landscape:estimated-hit-location"));
+
+			OutHitLocation = BestHitResult->LandscapeProxy->LandscapeActorToWorld().InverseTransformPosition(EstimatedHitLocation);
+			return true;
 		}
 	}
 
