Added perf tags and markers for easier profiling of landscape tools + some code improvements

diff --git a/Engine/Source/Editor/LandscapeEditor/Private/LandscapeEdModeBPCustomTools.cpp b/Engine/Source/Editor/LandscapeEditor/Private/LandscapeEdModeBPCustomTools.cpp
--- a/Engine/Source/Editor/LandscapeEditor/Private/LandscapeEdModeBPCustomTools.cpp
+++ b/Engine/Source/Editor/LandscapeEditor/Private/LandscapeEdModeBPCustomTools.cpp
@@ -165,48 +165,7 @@ public:
 			}
 		}
 	}
-
-	
-protected:
-/*	float GetLocalZAtPoint(const ULandscapeInfo* LandscapeInfo, int32 x, int32 y) const
-	{
-		// try to find Z location
-		TSet<ULandscapeComponent*> Components;
-		LandscapeInfo->GetComponentsInRegion(x, y, x, y, Components);
-		for (ULandscapeComponent* Component : Components)
-		{
-			FLandscapeComponentDataInterface DataInterface(Component);
-			return LandscapeDataAccess::GetLocalHeight(DataInterface.GetHeight(x - Component->SectionBaseX, y - Component->SectionBaseY));
-		}
-		return 0.0f;
-	}
-*/
-
-public:
 };
-/*
-void FEdModeLandscape::ApplyMirrorTool()
-{
-	if (CurrentTool->GetToolName() == FName("Mirror"))
-	{
-		FLandscapeToolMirror* MirrorTool = (FLandscapeToolMirror*)CurrentTool;
-		MirrorTool->ApplyMirror();
-		GEditor->RedrawLevelEditingViewports();
-	}
-}
-
-void FEdModeLandscape::CenterMirrorTool()
-{
-	if (CurrentTool->GetToolName() == FName("Mirror"))
-	{
-		FLandscapeToolMirror* MirrorTool = (FLandscapeToolMirror*)CurrentTool;
-		MirrorTool->CenterMirrorPoint();
-		GEditor->RedrawLevelEditingViewports();
-	}
-}
-*/
-
-
 
 //
 // Toolset initialization
diff --git a/Engine/Source/Editor/LandscapeEditor/Private/LandscapeEdModeComponentTools.cpp b/Engine/Source/Editor/LandscapeEditor/Private/LandscapeEdModeComponentTools.cpp
--- a/Engine/Source/Editor/LandscapeEditor/Private/LandscapeEdModeComponentTools.cpp
+++ b/Engine/Source/Editor/LandscapeEditor/Private/LandscapeEdModeComponentTools.cpp
@@ -31,13 +31,15 @@
 //
 class FLandscapeToolStrokeSelect : public FLandscapeToolStrokeBase
 {
+	using Super = FLandscapeToolStrokeBase;
+
 	bool bInitializedComponentInvert;
 	bool bInvert;
 	bool bNeedsSelectionUpdate;
 
 public:
 	FLandscapeToolStrokeSelect(FEdModeLandscape* InEdMode, FEditorViewportClient* InViewportClient, const FLandscapeToolTarget& InTarget)
-		: FLandscapeToolStrokeBase(InEdMode, InViewportClient, InTarget)
+		: Super(InEdMode, InViewportClient, InTarget)
 		, bInitializedComponentInvert(false)
 		, bNeedsSelectionUpdate(false)
 		, Cache(InTarget)
@@ -62,6 +64,8 @@ public:
 
 	void Apply(FEditorViewportClient* ViewportClient, FLandscapeBrush* Brush, const ULandscapeEditorObject* UISettings, const TArray<FLandscapeToolInteractorPosition>& InteractorPositions)
 	{
+		TRACE_CPUPROFILER_EVENT_SCOPE(FLandscapeToolStrokeSelect_Apply);
+
 		if (LandscapeInfo)
 		{
 			LandscapeInfo->Modify();
@@ -125,9 +129,11 @@ protected:
 
 class FLandscapeToolSelect : public FLandscapeToolBase<FLandscapeToolStrokeSelect>
 {
+	using Super = FLandscapeToolBase<FLandscapeToolStrokeSelect>;
+
 public:
 	FLandscapeToolSelect(FEdModeLandscape* InEdMode)
-		: FLandscapeToolBase<FLandscapeToolStrokeSelect>(InEdMode)
+		: Super(InEdMode)
 	{
 	}
 
@@ -161,15 +167,19 @@ public:
 //
 class FLandscapeToolStrokeMask : public FLandscapeToolStrokeBase
 {
+	using Super = FLandscapeToolStrokeBase;
+
 public:
 	FLandscapeToolStrokeMask(FEdModeLandscape* InEdMode, FEditorViewportClient* InViewportClient, const FLandscapeToolTarget& InTarget)
-		: FLandscapeToolStrokeBase(InEdMode, InViewportClient, InTarget)
+		: Super(InEdMode, InViewportClient, InTarget)
 		, Cache(InTarget)
 	{
 	}
 
 	void Apply(FEditorViewportClient* ViewportClient, FLandscapeBrush* Brush, const ULandscapeEditorObject* UISettings, const TArray<FLandscapeToolInteractorPosition>& InteractorPositions)
 	{
+		TRACE_CPUPROFILER_EVENT_SCOPE(FLandscapeToolStrokeMask_Apply);
+
 		if (LandscapeInfo)
 		{
 			LandscapeInfo->Modify();
@@ -245,9 +255,11 @@ protected:
 
 class FLandscapeToolMask : public FLandscapeToolBase<FLandscapeToolStrokeMask>
 {
+	using Super = FLandscapeToolBase<FLandscapeToolStrokeMask>;
+
 public:
 	FLandscapeToolMask(FEdModeLandscape* InEdMode)
-		: FLandscapeToolBase<FLandscapeToolStrokeMask>(InEdMode)
+		: Super(InEdMode)
 	{
 	}
 
@@ -268,15 +280,19 @@ public:
 //
 class FLandscapeToolStrokeVisibility : public FLandscapeToolStrokeBase
 {
+	using Super = FLandscapeToolStrokeBase;
+
 public:
 	FLandscapeToolStrokeVisibility(FEdModeLandscape* InEdMode, FEditorViewportClient* InViewportClient, const FLandscapeToolTarget& InTarget)
-		: FLandscapeToolStrokeBase(InEdMode, InViewportClient, InTarget)
+		: Super(InEdMode, InViewportClient, InTarget)
 		, Cache(InTarget)
 	{
 	}
 
 	void Apply(FEditorViewportClient* ViewportClient, FLandscapeBrush* Brush, const ULandscapeEditorObject* UISettings, const TArray<FLandscapeToolInteractorPosition>& InteractorPositions)
 	{
+		TRACE_CPUPROFILER_EVENT_SCOPE(FLandscapeToolStrokeVisibility_Apply);
+
 		if (LandscapeInfo)
 		{
 			LandscapeInfo->Modify();
@@ -328,15 +344,17 @@ protected:
 
 class FLandscapeToolVisibility : public FLandscapeToolBase<FLandscapeToolStrokeVisibility>
 {
+	using Super = FLandscapeToolBase<FLandscapeToolStrokeVisibility>;
+
 public:
 	FLandscapeToolVisibility(FEdModeLandscape* InEdMode)
-		: FLandscapeToolBase<FLandscapeToolStrokeVisibility>(InEdMode)
+		: Super(InEdMode)
 	{
 	}
 
 	virtual bool BeginTool(FEditorViewportClient* ViewportClient, const FLandscapeToolTarget& InTarget, const FVector& InHitLocation) override
 	{
-		return FLandscapeToolBase<FLandscapeToolStrokeVisibility>::BeginTool(ViewportClient, InTarget, InHitLocation);
+		return Super::BeginTool(ViewportClient, InTarget, InHitLocation);
 	}
 
 	virtual const TCHAR* GetToolName() const override { return TEXT("Visibility"); }
@@ -357,14 +375,18 @@ public:
 //
 class FLandscapeToolStrokeMoveToLevel : public FLandscapeToolStrokeBase
 {
+	using Super = FLandscapeToolStrokeBase;
+
 public:
 	FLandscapeToolStrokeMoveToLevel(FEdModeLandscape* InEdMode, FEditorViewportClient* InViewportClient, const FLandscapeToolTarget& InTarget)
-		: FLandscapeToolStrokeBase(InEdMode, InViewportClient, InTarget)
+		: Super(InEdMode, InViewportClient, InTarget)
 	{
 	}
 
 	void Apply(FEditorViewportClient* ViewportClient, FLandscapeBrush* Brush, const ULandscapeEditorObject* UISettings, const TArray<FLandscapeToolInteractorPosition>& InteractorPositions)
 	{
+		TRACE_CPUPROFILER_EVENT_SCOPE(FLandscapeToolStrokeMoveToLevel_Apply);
+
 		ALandscape* Landscape = LandscapeInfo ? LandscapeInfo->LandscapeActor.Get() : nullptr;
 
 		if (Landscape)
@@ -494,9 +516,11 @@ public:
 
 class FLandscapeToolMoveToLevel : public FLandscapeToolBase<FLandscapeToolStrokeMoveToLevel>
 {
+	using Super = FLandscapeToolBase<FLandscapeToolStrokeMoveToLevel>;
+
 public:
 	FLandscapeToolMoveToLevel(FEdModeLandscape* InEdMode)
-		: FLandscapeToolBase<FLandscapeToolStrokeMoveToLevel>(InEdMode)
+		: Super(InEdMode)
 	{
 	}
 	virtual bool AffectsEditLayers() const override { return false; }
@@ -514,9 +538,11 @@ public:
 //
 class FLandscapeToolStrokeAddComponent : public FLandscapeToolStrokeBase
 {
+	using Super = FLandscapeToolStrokeBase;
+
 public:
 	FLandscapeToolStrokeAddComponent(FEdModeLandscape* InEdMode, FEditorViewportClient* InViewportClient, const FLandscapeToolTarget& InTarget)
-		: FLandscapeToolStrokeBase(InEdMode, InViewportClient, InTarget)
+		: Super(InEdMode, InViewportClient, InTarget)
 		, HeightCache(InTarget)
 		, XYOffsetCache(InTarget)
 	{
@@ -531,6 +557,8 @@ public:
 
 	virtual void Apply(FEditorViewportClient* ViewportClient, FLandscapeBrush* Brush, const ULandscapeEditorObject* UISettings, const TArray<FLandscapeToolInteractorPosition>& InteractorPositions)
 	{
+		TRACE_CPUPROFILER_EVENT_SCOPE(FLandscapeToolStrokeAddComponent_Apply);
+
 		if (LandscapeInfo)
 		{
 			check(Brush->GetBrushType() == ELandscapeBrushType::Component);
@@ -784,9 +812,11 @@ protected:
 
 class FLandscapeToolAddComponent : public FLandscapeToolBase<FLandscapeToolStrokeAddComponent>
 {
+	using Super = FLandscapeToolBase<FLandscapeToolStrokeAddComponent>;
+
 public:
 	FLandscapeToolAddComponent(FEdModeLandscape* InEdMode)
-		: FLandscapeToolBase<FLandscapeToolStrokeAddComponent>(InEdMode)
+		: Super(InEdMode)
 		, bIsToolActionResolutionCompliant(true)
 	{
 	}
@@ -804,7 +834,7 @@ public:
 
 	virtual bool CanToolBeActivated() const override
 	{ 
-		return FLandscapeToolBase<FLandscapeToolStrokeAddComponent>::CanToolBeActivated() && bIsToolActionResolutionCompliant;
+		return Super::CanToolBeActivated() && bIsToolActionResolutionCompliant;
 	}
 
 	virtual void Tick(FEditorViewportClient* ViewportClient, float DeltaTime) override
@@ -814,12 +844,12 @@ public:
 			bIsToolActionResolutionCompliant = EdMode->IsLandscapeResolutionCompliant();
 		}
 
-		FLandscapeToolBase<FLandscapeToolStrokeAddComponent>::Tick(ViewportClient, DeltaTime);
+		Super::Tick(ViewportClient, DeltaTime);
 	}
 	
 	virtual void EnterTool() override
 	{
-		FLandscapeToolBase<FLandscapeToolStrokeAddComponent>::EnterTool();
+		Super::EnterTool();
 		AddCollision.Reset();
 		if(ULandscapeInfo* LandscapeInfo = EdMode->CurrentToolTarget.LandscapeInfo.Get())
 		{
@@ -829,7 +859,7 @@ public:
 
 	virtual void ExitTool() override
 	{
-		FLandscapeToolBase<FLandscapeToolStrokeAddComponent>::ExitTool();
+		Super::ExitTool();
 
 		AddCollision.Reset();
 	}
@@ -958,14 +988,18 @@ private:
 //
 class FLandscapeToolStrokeDeleteComponent : public FLandscapeToolStrokeBase
 {
+	using Super = FLandscapeToolStrokeBase;
+
 public:
 	FLandscapeToolStrokeDeleteComponent(FEdModeLandscape* InEdMode, FEditorViewportClient* InViewportClient, const FLandscapeToolTarget& InTarget)
-		: FLandscapeToolStrokeBase(InEdMode, InViewportClient, InTarget)
+		: Super(InEdMode, InViewportClient, InTarget)
 	{
 	}
 
 	void Apply(FEditorViewportClient* ViewportClient, FLandscapeBrush* Brush, const ULandscapeEditorObject* UISettings, const TArray<FLandscapeToolInteractorPosition>& InteractorPositions)
 	{
+		TRACE_CPUPROFILER_EVENT_SCOPE(FLandscapeToolStrokeDeleteComponent_Apply);
+
 		if (LandscapeInfo)
 		{
 			auto SelectedComponents = LandscapeInfo->GetSelectedComponents();
@@ -994,9 +1028,11 @@ public:
 
 class FLandscapeToolDeleteComponent : public FLandscapeToolBase<FLandscapeToolStrokeDeleteComponent>
 {
+	using Super = FLandscapeToolBase<FLandscapeToolStrokeDeleteComponent>;
+
 public:
 	FLandscapeToolDeleteComponent(FEdModeLandscape* InEdMode)
-		: FLandscapeToolBase<FLandscapeToolStrokeDeleteComponent>(InEdMode)
+		: Super(InEdMode)
 	{
 	}
 
@@ -1016,9 +1052,11 @@ public:
 template<class ToolTarget>
 class FLandscapeToolStrokeCopy : public FLandscapeToolStrokeBase
 {
+	using Super = FLandscapeToolStrokeBase;
+
 public:
 	FLandscapeToolStrokeCopy(FEdModeLandscape* InEdMode, FEditorViewportClient* InViewportClient, const FLandscapeToolTarget& InTarget)
-		: FLandscapeToolStrokeBase(InEdMode, InViewportClient, InTarget)
+		: Super(InEdMode, InViewportClient, InTarget)
 		, Cache(InTarget)
 		, HeightCache(InTarget)
 		, WeightCache(InTarget)
@@ -1033,6 +1071,8 @@ public:
 
 	void Apply(FEditorViewportClient* ViewportClient, FLandscapeBrush* Brush, const ULandscapeEditorObject* UISettings, const TArray<FLandscapeToolInteractorPosition>& InteractorPositions)
 	{
+		TRACE_CPUPROFILER_EVENT_SCOPE(FLandscapeToolStrokeCopy_Apply);
+
 		//ULandscapeInfo* LandscapeInfo = EdMode->CurrentToolTarget.LandscapeInfo;
 		ALandscapeGizmoActiveActor* Gizmo = EdMode->CurrentGizmoActor.Get();
 		if (LandscapeInfo && Gizmo && Gizmo->GizmoTexture && Gizmo->GetRootComponent())
@@ -1321,9 +1361,11 @@ protected:
 template<class ToolTarget>
 class FLandscapeToolCopy : public FLandscapeToolBase<FLandscapeToolStrokeCopy<ToolTarget>>
 {
+	using Super = FLandscapeToolBase<FLandscapeToolStrokeCopy<ToolTarget>>;
+
 public:
 	FLandscapeToolCopy(FEdModeLandscape* InEdMode)
-		: FLandscapeToolBase<FLandscapeToolStrokeCopy<ToolTarget> >(InEdMode)
+		: Super(InEdMode)
 		, BackupCurrentBrush(nullptr)
 	{
 	}
@@ -1353,12 +1395,12 @@ public:
 		BackupCurrentBrush = this->EdMode->CurrentBrush;
 		this->EdMode->CurrentBrush = this->EdMode->GizmoBrush;
 
-		return FLandscapeToolBase<FLandscapeToolStrokeCopy<ToolTarget>>::BeginTool(ViewportClient, InTarget, InHitLocation);
+		return Super::BeginTool(ViewportClient, InTarget, InHitLocation);
 	}
 
 	virtual void EndTool(FEditorViewportClient* ViewportClient) override
 	{
-		FLandscapeToolBase<FLandscapeToolStrokeCopy<ToolTarget>>::EndTool(ViewportClient);
+		Super::EndTool(ViewportClient);
 
 		this->EdMode->CurrentBrush = BackupCurrentBrush;
 	}
@@ -1373,9 +1415,11 @@ protected:
 template<class ToolTarget>
 class FLandscapeToolStrokePaste : public FLandscapeToolStrokeBase
 {
+	using Super = FLandscapeToolStrokeBase;
+
 public:
 	FLandscapeToolStrokePaste(FEdModeLandscape* InEdMode, FEditorViewportClient* InViewportClient, const FLandscapeToolTarget& InTarget)
-		: FLandscapeToolStrokeBase(InEdMode, InViewportClient, InTarget)
+		: Super(InEdMode, InViewportClient, InTarget)
 		, Cache(InTarget)
 		, HeightCache(InTarget)
 		, WeightCache(InTarget)
@@ -1384,6 +1428,8 @@ public:
 
 	void Apply(FEditorViewportClient* ViewportClient, FLandscapeBrush* Brush, const ULandscapeEditorObject* UISettings, const TArray<FLandscapeToolInteractorPosition>& InteractorPositions)
 	{
+		TRACE_CPUPROFILER_EVENT_SCOPE(FLandscapeToolStrokePaste_Apply);
+
 		//ULandscapeInfo* LandscapeInfo = EdMode->CurrentToolTarget.LandscapeInfo;
 		ALandscapeGizmoActiveActor* Gizmo = EdMode->CurrentGizmoActor.Get();
 		// Cache and copy in Gizmo's region...
@@ -1639,9 +1685,11 @@ protected:
 template<class ToolTarget>
 class FLandscapeToolPaste : public FLandscapeToolBase<FLandscapeToolStrokePaste<ToolTarget>>
 {
+	using Super = FLandscapeToolBase<FLandscapeToolStrokePaste<ToolTarget>>;
+
 public:
 	FLandscapeToolPaste(FEdModeLandscape* InEdMode)
-		: FLandscapeToolBase<FLandscapeToolStrokePaste<ToolTarget>>(InEdMode)
+		: Super(InEdMode)
 		, bUseGizmoRegion(false)
 		, BackupCurrentBrush(nullptr)
 	{
@@ -1685,12 +1733,12 @@ public:
 			this->EdMode->CurrentBrush = this->EdMode->GizmoBrush;
 		}
 
-		return FLandscapeToolBase<FLandscapeToolStrokePaste<ToolTarget>>::BeginTool(ViewportClient, InTarget, InHitLocation);
+		return Super::BeginTool(ViewportClient, InTarget, InHitLocation);
 	}
 
 	virtual void EndTool(FEditorViewportClient* ViewportClient) override
 	{
-		FLandscapeToolBase<FLandscapeToolStrokePaste<ToolTarget>>::EndTool(ViewportClient);
+		Super::EndTool(ViewportClient);
 
 		if (bUseGizmoRegion)
 		{
diff --git a/Engine/Source/Editor/LandscapeEditor/Private/LandscapeEdModeErosionTools.cpp b/Engine/Source/Editor/LandscapeEditor/Private/LandscapeEdModeErosionTools.cpp
--- a/Engine/Source/Editor/LandscapeEditor/Private/LandscapeEdModeErosionTools.cpp
+++ b/Engine/Source/Editor/LandscapeEditor/Private/LandscapeEdModeErosionTools.cpp
@@ -17,9 +17,11 @@
 //
 class FLandscapeToolStrokeErosionBase : public FLandscapeToolStrokeBase
 {
+	using Super = FLandscapeToolStrokeBase;
+
 public:
 	FLandscapeToolStrokeErosionBase(FEdModeLandscape* InEdMode, FEditorViewportClient* InViewportClient, const FLandscapeToolTarget& InTarget)
-		: FLandscapeToolStrokeBase(InEdMode, InViewportClient, InTarget)
+		: Super(InEdMode, InViewportClient, InTarget)
 		, HeightCache(InTarget)
 		, LayerHeightDataCache(InTarget, this->HeightCache)
 		, WeightCache(InTarget)
@@ -37,9 +39,11 @@ protected:
 template<class TStrokeClass>
 class FLandscapeToolErosionBase : public FLandscapeToolBase<TStrokeClass>
 {
+	using Super = FLandscapeToolBase<TStrokeClass>;
+
 public:
 	FLandscapeToolErosionBase(FEdModeLandscape* InEdMode)
-		: FLandscapeToolBase<TStrokeClass>(InEdMode)
+		: Super(InEdMode)
 	{
 	}
 
@@ -55,14 +59,18 @@ public:
 
 class FLandscapeToolStrokeErosion : public FLandscapeToolStrokeErosionBase
 {
+	using Super = FLandscapeToolStrokeErosionBase;
+
 public:
 	FLandscapeToolStrokeErosion(FEdModeLandscape* InEdMode, FEditorViewportClient* InViewportClient, const FLandscapeToolTarget& InTarget)
-		: FLandscapeToolStrokeErosionBase(InEdMode, InViewportClient, InTarget)
+		: Super(InEdMode, InViewportClient, InTarget)
 	{
 	}
 
 	void Apply(FEditorViewportClient* ViewportClient, FLandscapeBrush* Brush, const ULandscapeEditorObject* UISettings, const TArray<FLandscapeToolInteractorPosition>& InteractorPositions)
 	{
+		TRACE_CPUPROFILER_EVENT_SCOPE(FLandscapeToolStrokeErosion_Apply);
+
 		if (!this->LandscapeInfo)
 		{
 			return;
@@ -272,9 +280,11 @@ public:
 
 class FLandscapeToolErosion : public FLandscapeToolErosionBase<FLandscapeToolStrokeErosion>
 {
+	using Super = FLandscapeToolErosionBase<FLandscapeToolStrokeErosion>;
+
 public:
 	FLandscapeToolErosion(FEdModeLandscape* InEdMode)
-		: FLandscapeToolErosionBase(InEdMode)
+		: Super(InEdMode)
 	{
 	}
 
@@ -290,14 +300,18 @@ public:
 
 class FLandscapeToolStrokeHydraErosion : public FLandscapeToolStrokeErosionBase
 {
+	using Super = FLandscapeToolStrokeErosionBase;
+
 public:
 	FLandscapeToolStrokeHydraErosion(FEdModeLandscape* InEdMode, FEditorViewportClient* InViewportClient, const FLandscapeToolTarget& InTarget)
-		: FLandscapeToolStrokeErosionBase(InEdMode, InViewportClient, InTarget)
+		: Super(InEdMode, InViewportClient, InTarget)
 	{
 	}
 
 	void Apply(FEditorViewportClient* ViewportClient, FLandscapeBrush* Brush, const ULandscapeEditorObject* UISettings, const TArray<FLandscapeToolInteractorPosition>& InteractorPositions)
 	{
+		TRACE_CPUPROFILER_EVENT_SCOPE(FLandscapeToolStrokeHydraErosion_Apply);
+
 		if (!this->LandscapeInfo)
 		{
 			return;
@@ -493,9 +507,11 @@ public:
 
 class FLandscapeToolHydraErosion : public FLandscapeToolErosionBase<FLandscapeToolStrokeHydraErosion>
 {
+	using Super = FLandscapeToolErosionBase<FLandscapeToolStrokeHydraErosion>;
+
 public:
 	FLandscapeToolHydraErosion(FEdModeLandscape* InEdMode)
-		: FLandscapeToolErosionBase(InEdMode)
+		: Super(InEdMode)
 	{
 	}
 
diff --git a/Engine/Source/Editor/LandscapeEditor/Private/LandscapeEdModePaintTools.cpp b/Engine/Source/Editor/LandscapeEditor/Private/LandscapeEdModePaintTools.cpp
--- a/Engine/Source/Editor/LandscapeEditor/Private/LandscapeEdModePaintTools.cpp
+++ b/Engine/Source/Editor/LandscapeEditor/Private/LandscapeEdModePaintTools.cpp
@@ -95,9 +95,11 @@ namespace
 template<class TToolTarget, class TStrokeClass>
 class FLandscapeToolPaintBase : public FLandscapeToolBase<TStrokeClass>
 {
+	using Super = FLandscapeToolBase<TStrokeClass>;
+
 public:
 	FLandscapeToolPaintBase(FEdModeLandscape* InEdMode)
-		: FLandscapeToolBase<TStrokeClass>(InEdMode)
+		: Super(InEdMode)
 	{
 	}
 
@@ -111,9 +113,11 @@ public:
 template<class ToolTarget>
 class FLandscapeToolStrokePaintBase : public FLandscapeToolStrokeBase
 {
+	using Super = FLandscapeToolStrokeBase;
+
 public:
 	FLandscapeToolStrokePaintBase(FEdModeLandscape* InEdMode, FEditorViewportClient* InViewportClient, const FLandscapeToolTarget& InTarget)
-		: FLandscapeToolStrokeBase(InEdMode, InViewportClient, InTarget)
+		: Super(InEdMode, InViewportClient, InTarget)
 		, Cache(InTarget)
 	{
 	}
@@ -127,7 +131,8 @@ protected:
 //
 class FLandscapeToolStrokePaint : public FLandscapeToolStrokePaintBase<FWeightmapToolTarget>
 {
-	typedef FWeightmapToolTarget ToolTarget;
+	using Super = FLandscapeToolStrokePaintBase<FWeightmapToolTarget>;
+	using ToolTarget = FWeightmapToolTarget;
 
 	TMap<FIntPoint, float> TotalInfluenceMap; // amount of time and weight the brush has spent on each vertex.
 
@@ -138,7 +143,7 @@ public:
 	enum { UseContinuousApply = false };
 
 	FLandscapeToolStrokePaint(FEdModeLandscape* InEdMode, FEditorViewportClient* InViewportClient, const FLandscapeToolTarget& InTarget)
-		: FLandscapeToolStrokePaintBase<FWeightmapToolTarget>(InEdMode, InViewportClient, InTarget)
+		: Super(InEdMode, InViewportClient, InTarget)
 		, bIsAllowListMode(EdMode->UISettings->PaintingRestriction == ELandscapeLayerPaintingRestriction::UseComponentAllowList &&
 		                   (InViewportClient->Viewport->KeyState(EKeys::Equals) || InViewportClient->Viewport->KeyState(EKeys::Hyphen)))
 		, bAddToAllowList(bIsAllowListMode && InViewportClient->Viewport->KeyState(EKeys::Equals))
@@ -147,6 +152,8 @@ public:
 
 	void Apply(FEditorViewportClient* ViewportClient, FLandscapeBrush* Brush, const ULandscapeEditorObject* UISettings, const TArray<FLandscapeToolInteractorPosition>& InteractorPositions)
 	{
+		TRACE_CPUPROFILER_EVENT_SCOPE(FLandscapeToolStrokePaint_Apply);
+
 		// Invert when holding Shift
 		bool bInvert = InteractorPositions.Last().bModifierPressed;
 		UE_LOG(LogLandscapeTools, VeryVerbose, TEXT("bInvert = %d"), bInvert);
@@ -311,9 +318,11 @@ public:
 
 class FLandscapeToolPaint : public FLandscapeToolPaintBase<FWeightmapToolTarget, FLandscapeToolStrokePaint>
 {
+	using Super = FLandscapeToolPaintBase<FWeightmapToolTarget, FLandscapeToolStrokePaint>;
+
 public:
 	FLandscapeToolPaint(FEdModeLandscape* InEdMode)
-		: FLandscapeToolPaintBase<FWeightmapToolTarget, FLandscapeToolStrokePaint>(InEdMode)
+		: Super(InEdMode)
 	{
 	}
 
@@ -328,13 +337,14 @@ public:
 			EdMode->UISettings->UpdateComponentLayerAllowList();
 		}
 
-		FLandscapeToolPaintBase::EnterTool();
+		Super::EnterTool();
 	}
 };
 
 class FLandscapeToolStrokeErase : public FLandscapeToolStrokePaintBase<FHeightmapToolTarget>
 {
-	typedef FHeightmapToolTarget ToolTarget;
+	using Super = FLandscapeToolStrokePaintBase<FHeightmapToolTarget>;
+	using ToolTarget = FHeightmapToolTarget;
 	const typename ToolTarget::CacheClass::DataType FlattenHeight;
 
 public:
@@ -342,13 +352,15 @@ public:
 	enum { UseContinuousApply = true };
 
 	FLandscapeToolStrokeErase(FEdModeLandscape* InEdMode, FEditorViewportClient* InViewportClient, const FLandscapeToolTarget& InTarget)
-		: FLandscapeToolStrokePaintBase<FHeightmapToolTarget>(InEdMode, InViewportClient, InTarget)
+		: Super(InEdMode, InViewportClient, InTarget)
 		, FlattenHeight(LandscapeDataAccess::GetTexHeight(0.f))
 	{
 	}
 
 	void Apply(FEditorViewportClient* ViewportClient, FLandscapeBrush* Brush, const ULandscapeEditorObject* UISettings, const TArray<FLandscapeToolInteractorPosition>& InteractorPositions)
 	{
+		TRACE_CPUPROFILER_EVENT_SCOPE(FLandscapeToolStrokeErase_Apply);
+
 		if (!this->LandscapeInfo) return;
 
 		// Get list of verts to update
@@ -410,19 +422,22 @@ public:
 //
 class FLandscapeToolStrokeSculpt : public FLandscapeToolStrokePaintBase<FHeightmapToolTarget>
 {
-	typedef FHeightmapToolTarget ToolTarget;
+	using Super = FLandscapeToolStrokePaintBase<FHeightmapToolTarget>;
+	using ToolTarget = FHeightmapToolTarget;
 
 public:
 	// Heightmap sculpt tool will continuously sculpt in the same location, weightmap paint tool doesn't
 	enum { UseContinuousApply = true };
 
 	FLandscapeToolStrokeSculpt(FEdModeLandscape* InEdMode, FEditorViewportClient* InViewportClient, const FLandscapeToolTarget& InTarget)
-		: FLandscapeToolStrokePaintBase<FHeightmapToolTarget>(InEdMode, InViewportClient, InTarget)
+		: Super(InEdMode, InViewportClient, InTarget)
 	{
 	}
 
 	void Apply(FEditorViewportClient* ViewportClient, FLandscapeBrush* Brush, const ULandscapeEditorObject* UISettings, const TArray<FLandscapeToolInteractorPosition>& InteractorPositions)
 	{
+		TRACE_CPUPROFILER_EVENT_SCOPE(FLandscapeToolStrokeSculpt_Apply);
+
 		// Invert when holding Shift
 		bool bInvert = InteractorPositions.Last().bModifierPressed;
 		UE_LOG(LogLandscapeTools, VeryVerbose, TEXT("bInvert = %d"), bInvert);
@@ -482,6 +497,8 @@ public:
 
 		if (bUseClayBrush)
 		{
+			TRACE_CPUPROFILER_EVENT_SCOPE(ComputeClayBrush);
+
 			// Calculate normals for brush verts in data space
 			Normals.Empty(SourceDataArrayPtr->Num());
 			Normals.AddZeroed(SourceDataArrayPtr->Num());
@@ -566,50 +583,54 @@ public:
 		}
 
 		// Apply the brush
-		for (int32 Y = BrushInfo.GetBounds().Min.Y; Y < BrushInfo.GetBounds().Max.Y; Y++)
 		{
-			const float* BrushScanline = BrushInfo.GetDataPtr(FIntPoint(0, Y));
-			auto* DataScanline = Data.GetData() + (Y - Y1) * (X2 - X1 + 1) + (0 - X1);
-			auto* SourceDataScanline = SourceDataArrayPtr->GetData() + (Y - Y1) * (X2 - X1 + 1) + (0 - X1);
+			TRACE_CPUPROFILER_EVENT_SCOPE(ApplyBrush);
 
-			for (int32 X = BrushInfo.GetBounds().Min.X; X < BrushInfo.GetBounds().Max.X; X++)
+			for (int32 Y = BrushInfo.GetBounds().Min.Y; Y < BrushInfo.GetBounds().Max.Y; Y++)
 			{
-				const FIntPoint Key = FIntPoint(X, Y);
-				const float BrushValue = BrushScanline[X];
-
-				float SculptAmount = BrushValue * SculptStrength;
-				auto& CurrentValue = DataScanline[X];
-				const auto& SourceValue = SourceDataScanline[X];
+				const float* BrushScanline = BrushInfo.GetDataPtr(FIntPoint(0, Y));
+				auto* DataScanline = Data.GetData() + (Y - Y1) * (X2 - X1 + 1) + (0 - X1);
+				auto* SourceDataScanline = SourceDataArrayPtr->GetData() + (Y - Y1) * (X2 - X1 + 1) + (0 - X1);
 
-				if (bUseClayBrush)
+				for (int32 X = BrushInfo.GetBounds().Min.X; X < BrushInfo.GetBounds().Max.X; X++)
 				{
-					// Brush application starts from original world location at start of stroke
-					FVector WorldLoc = ToWorld.TransformPosition(FVector(X, Y, SourceValue));
-
-					// Calculate new location on the brush plane
-					WorldLoc.Z = (BrushPlane.W - BrushPlane.X*WorldLoc.X - BrushPlane.Y*WorldLoc.Y) / BrushPlane.Z;
+					const FIntPoint Key = FIntPoint(X, Y);
+					const float BrushValue = BrushScanline[X];
 
-					// Painted amount lerps based on brush falloff.
-					float PaintValue = FMath::Lerp<float>((float)SourceValue, FromWorld.TransformPosition(WorldLoc).Z, BrushValue);
+					float SculptAmount = BrushValue * SculptStrength;
+					auto& CurrentValue = DataScanline[X];
+					const auto& SourceValue = SourceDataScanline[X];
 
-					if (bInvert)
-					{
-						CurrentValue = ToolTarget::CacheClass::ClampValue(FMath::Min<int32>(FMath::RoundToInt(PaintValue), CurrentValue));
-					}
-					else
+					if (bUseClayBrush)
 					{
-						CurrentValue = ToolTarget::CacheClass::ClampValue(FMath::Max<int32>(FMath::RoundToInt(PaintValue), CurrentValue));
-					}
-				}
-				else
-				{
-					if (bInvert)
-					{
-						CurrentValue = ToolTarget::CacheClass::ClampValue(FMath::Min<int32>(SourceValue - FMath::RoundToInt(SculptAmount), CurrentValue));
+						// Brush application starts from original world location at start of stroke
+						FVector WorldLoc = ToWorld.TransformPosition(FVector(X, Y, SourceValue));
+
+						// Calculate new location on the brush plane
+						WorldLoc.Z = (BrushPlane.W - BrushPlane.X*WorldLoc.X - BrushPlane.Y*WorldLoc.Y) / BrushPlane.Z;
+
+						// Painted amount lerps based on brush falloff.
+						float PaintValue = FMath::Lerp<float>((float)SourceValue, FromWorld.TransformPosition(WorldLoc).Z, BrushValue);
+
+						if (bInvert)
+						{
+							CurrentValue = ToolTarget::CacheClass::ClampValue(FMath::Min<int32>(FMath::RoundToInt(PaintValue), CurrentValue));
+						}
+						else
+						{
+							CurrentValue = ToolTarget::CacheClass::ClampValue(FMath::Max<int32>(FMath::RoundToInt(PaintValue), CurrentValue));
+						}
 					}
 					else
 					{
-						CurrentValue = ToolTarget::CacheClass::ClampValue(FMath::Max<int32>(SourceValue + FMath::RoundToInt(SculptAmount), CurrentValue));
+						if (bInvert)
+						{
+							CurrentValue = ToolTarget::CacheClass::ClampValue(FMath::Min<int32>(SourceValue - FMath::RoundToInt(SculptAmount), CurrentValue));
+						}
+						else
+						{
+							CurrentValue = ToolTarget::CacheClass::ClampValue(FMath::Max<int32>(SourceValue + FMath::RoundToInt(SculptAmount), CurrentValue));
+						}
 					}
 				}
 			}
@@ -622,9 +643,11 @@ public:
 
 class FLandscapeToolSculpt : public FLandscapeToolPaintBase<FHeightmapToolTarget, FLandscapeToolStrokeSculpt>
 {
+	using Super = FLandscapeToolPaintBase<FHeightmapToolTarget, FLandscapeToolStrokeSculpt>;
+
 public:
 	FLandscapeToolSculpt(FEdModeLandscape* InEdMode)
-		: FLandscapeToolPaintBase<FHeightmapToolTarget, FLandscapeToolStrokeSculpt>(InEdMode)
+		: Super(InEdMode)
 	{
 	}
 
@@ -635,9 +658,11 @@ public:
 
 class FLandscapeToolErase : public FLandscapeToolPaintBase<FHeightmapToolTarget, FLandscapeToolStrokeErase>
 {
+	using Super = FLandscapeToolPaintBase<FHeightmapToolTarget, FLandscapeToolStrokeErase>;
+
 public:
 	FLandscapeToolErase(FEdModeLandscape* InEdMode)
-		: FLandscapeToolPaintBase<FHeightmapToolTarget, FLandscapeToolStrokeErase>(InEdMode)
+		: Super(InEdMode)
 	{
 	}
 
@@ -653,11 +678,13 @@ public:
 template<class ToolTarget>
 class FLandscapeToolStrokeSmooth : public FLandscapeToolStrokePaintBase<ToolTarget>
 {
+	using Super = FLandscapeToolStrokePaintBase<ToolTarget>;
+
 	bool bTargetIsHeightmap;
 	FLandscapeLayerDataCache<ToolTarget> LayerDataCache;
 public:
 	FLandscapeToolStrokeSmooth(FEdModeLandscape* InEdMode, FEditorViewportClient* InViewportClient, const FLandscapeToolTarget& InTarget)
-		: FLandscapeToolStrokePaintBase<ToolTarget>(InEdMode, InViewportClient, InTarget)
+		: Super(InEdMode, InViewportClient, InTarget)
 		, bTargetIsHeightmap(InTarget.TargetType == ELandscapeToolTargetType::Heightmap)
 		, LayerDataCache(InTarget, this->Cache)
 	{
@@ -665,6 +692,8 @@ public:
 
 	void Apply(FEditorViewportClient* ViewportClient, FLandscapeBrush* Brush, const ULandscapeEditorObject* UISettings, const TArray<FLandscapeToolInteractorPosition>& InteractorPositions)
 	{
+		TRACE_CPUPROFILER_EVENT_SCOPE(FLandscapeToolStrokeSmooth_Apply);
+
 		if (!this->LandscapeInfo) return;
 
 		ALandscape* Landscape = this->LandscapeInfo->LandscapeActor.Get();
@@ -775,9 +804,11 @@ public:
 template<class ToolTarget>
 class FLandscapeToolSmooth : public FLandscapeToolPaintBase<ToolTarget, FLandscapeToolStrokeSmooth<ToolTarget>>
 {
+	using Super = FLandscapeToolPaintBase<ToolTarget, FLandscapeToolStrokeSmooth<ToolTarget>>;
+
 public:
 	FLandscapeToolSmooth(FEdModeLandscape* InEdMode)
-		: FLandscapeToolPaintBase<ToolTarget, FLandscapeToolStrokeSmooth<ToolTarget> >(InEdMode)
+		: Super(InEdMode)
 	{
 	}
 
@@ -1061,8 +1092,10 @@ public:
 };
 
 template<class ToolTarget>
-class FLandscapeToolFlatten : public FLandscapeToolPaintBase < ToolTarget, FLandscapeToolStrokeFlatten<ToolTarget> >
+class FLandscapeToolFlatten : public FLandscapeToolPaintBase<ToolTarget, FLandscapeToolStrokeFlatten<ToolTarget>>
 {
+	using Super = FLandscapeToolPaintBase<ToolTarget, FLandscapeToolStrokeFlatten<ToolTarget>>;
+
 protected:
 	UStaticMesh* HeightmapFlattenPlaneMesh;
 	UStaticMeshComponent* HeightmapFlattenPreviewComponent;
@@ -1072,7 +1105,7 @@ protected:
 
 public:
 	FLandscapeToolFlatten(FEdModeLandscape* InEdMode)
-		: FLandscapeToolPaintBase<ToolTarget, FLandscapeToolStrokeFlatten<ToolTarget>>(InEdMode)
+		: Super(InEdMode)
 		, HeightmapFlattenPlaneMesh(LoadObject<UStaticMesh>(nullptr, TEXT("/Engine/EditorLandscapeResources/FlattenPlaneMesh.FlattenPlaneMesh")))
 		, HeightmapFlattenPreviewComponent(nullptr)
 		, CanToolBeActivatedNextTick(false)
@@ -1117,7 +1150,7 @@ public:
 			CanToolBeActivatedNextTick = false;
 		}
 
-		FLandscapeToolPaintBase<ToolTarget, FLandscapeToolStrokeFlatten<ToolTarget>>::Tick(ViewportClient, DeltaTime);
+		Super::Tick(ViewportClient, DeltaTime);
 
 		if (HeightmapFlattenPreviewComponent != nullptr)
 		{
@@ -1128,7 +1161,7 @@ public:
 
 	virtual bool MouseMove(FEditorViewportClient* ViewportClient, FViewport* Viewport, int32 x, int32 y) override
 	{
-		bool bResult = FLandscapeToolPaintBase<ToolTarget, FLandscapeToolStrokeFlatten<ToolTarget>>::MouseMove(ViewportClient, Viewport, x, y);
+		bool bResult = Super::MouseMove(ViewportClient, Viewport, x, y);
 
 		if (ViewportClient->IsLevelEditorClient() && HeightmapFlattenPreviewComponent != nullptr)
 		{
@@ -1156,7 +1189,7 @@ public:
 
 	virtual void EnterTool() override
 	{
-		FLandscapeToolPaintBase<ToolTarget, FLandscapeToolStrokeFlatten<ToolTarget>>::EnterTool();
+		Super::EnterTool();
 		if (!this->EdMode->CurrentToolTarget.LandscapeInfo.Get())
 		{
 			return;
@@ -1184,7 +1217,7 @@ public:
 
 	virtual void ExitTool() override
 	{
-		FLandscapeToolPaintBase<ToolTarget, FLandscapeToolStrokeFlatten<ToolTarget>>::ExitTool();
+		Super::ExitTool();
 
 		if (HeightmapFlattenPreviewComponent != nullptr)
 		{
@@ -1201,14 +1234,18 @@ public:
 template<class ToolTarget>
 class FLandscapeToolStrokeNoise : public FLandscapeToolStrokePaintBase<ToolTarget>
 {
+	using Super = FLandscapeToolStrokePaintBase<ToolTarget>;
+
 public:
 	FLandscapeToolStrokeNoise(FEdModeLandscape* InEdMode, FEditorViewportClient* InViewportClient, const FLandscapeToolTarget& InTarget)
-		: FLandscapeToolStrokePaintBase<ToolTarget>(InEdMode, InViewportClient, InTarget)
+		: Super(InEdMode, InViewportClient, InTarget)
 	{
 	}
 
 	void Apply(FEditorViewportClient* ViewportClient, FLandscapeBrush* Brush, const ULandscapeEditorObject* UISettings, const TArray<FLandscapeToolInteractorPosition>& InteractorPositions)
 	{
+		TRACE_CPUPROFILER_EVENT_SCOPE(FLandscapeToolStrokeNoise_Apply);
+
 		if (!this->LandscapeInfo) return;
 
 		// Get list of verts to update
@@ -1306,11 +1343,13 @@ public:
 };
 
 template<class ToolTarget>
-class FLandscapeToolNoise : public FLandscapeToolPaintBase < ToolTarget, FLandscapeToolStrokeNoise<ToolTarget> >
+class FLandscapeToolNoise : public FLandscapeToolPaintBase<ToolTarget, FLandscapeToolStrokeNoise<ToolTarget>>
 {
+	using Super = FLandscapeToolPaintBase<ToolTarget, FLandscapeToolStrokeNoise<ToolTarget>>;
+
 public:
 	FLandscapeToolNoise(FEdModeLandscape* InEdMode)
-		: FLandscapeToolPaintBase<ToolTarget, FLandscapeToolStrokeNoise<ToolTarget> >(InEdMode)
+		: Super(InEdMode)
 	{
 	}
 
diff --git a/Engine/Source/Editor/LandscapeEditor/Private/LandscapeEdModeTools.h b/Engine/Source/Editor/LandscapeEditor/Private/LandscapeEdModeTools.h
--- a/Engine/Source/Editor/LandscapeEditor/Private/LandscapeEdModeTools.h
+++ b/Engine/Source/Editor/LandscapeEditor/Private/LandscapeEdModeTools.h
@@ -297,6 +297,8 @@ public:
 	// Note that this should maybe be called "ExtendDataCache" because the region here will be combined with the existing cached region, not loaded independently, giving a cached region that is the bounding box of previous and new
 	void CacheData(int32 X1, int32 Y1, int32 X2, int32 Y2, bool bCacheOriginalData = false)
 	{
+		TRACE_CPUPROFILER_EVENT_SCOPE(TLandscapeEditCache_CacheData);
+
 		if (!bIsValid)
 		{
 			if (Accessor::bUseInterp)
@@ -550,6 +552,8 @@ public:
 	// X2/Y2 Coordinates are "inclusive" max values
 	bool GetCachedData(int32 X1, int32 Y1, int32 X2, int32 Y2, TArray<AccessorType>& OutData)
 	{
+		TRACE_CPUPROFILER_EVENT_SCOPE(TLandscapeEditCache_GetCachedData);
+
 		const int32 XSize = (1 + X2 - X1);
 		const int32 YSize = (1 + Y2 - Y1);
 		const int32 NumSamples = XSize * YSize;
@@ -585,6 +589,8 @@ public:
 	// X2/Y2 Coordinates are "inclusive" max values
 	void SetCachedData(int32 X1, int32 Y1, int32 X2, int32 Y2, TArray<AccessorType>& Data, ELandscapeLayerPaintingRestriction PaintingRestriction = ELandscapeLayerPaintingRestriction::None, bool bUpdateData = true)
 	{
+		TRACE_CPUPROFILER_EVENT_SCOPE(TLandscapeEditCache_SetCachedData);
+
 		checkSlow(Data.Num() == (1 + Y2 - Y1) * (1 + X2 - X1));
 
 		// Update cache
@@ -607,6 +613,8 @@ public:
 	// X2/Y2 Coordinates are "inclusive" max values
 	void GetOriginalData(int32 X1, int32 Y1, int32 X2, int32 Y2, TArray<AccessorType>& OutOriginalData)
 	{
+		TRACE_CPUPROFILER_EVENT_SCOPE(TLandscapeEditCache_GetOriginalData);
+
 		int32 NumSamples = (1 + X2 - X1)*(1 + Y2 - Y1);
 		OutOriginalData.Empty(NumSamples);
 		OutOriginalData.AddUninitialized(NumSamples);
@@ -626,6 +634,7 @@ public:
 
 	void Flush()
 	{
+		TRACE_CPUPROFILER_EVENT_SCOPE(TLandscapeEditCache_Flush);
 		DataAccess.Flush();
 	}
 
@@ -1393,6 +1402,8 @@ protected:
 template<class TStrokeClass>
 class FLandscapeToolBase : public FLandscapeTool
 {
+	using Super = FLandscapeTool;
+
 public:
 	FLandscapeToolBase(FEdModeLandscape* InEdMode)
 		: LastInteractorPosition(FVector2D::ZeroVector)
@@ -1427,6 +1438,8 @@ public:
 
 	virtual bool BeginTool(FEditorViewportClient* ViewportClient, const FLandscapeToolTarget& InTarget, const FVector& InHitLocation) override
 	{
+		TRACE_BOOKMARK(TEXT("BeginTool - %s"), GetToolName());
+
 		if (ShouldUpdateEditingLayer())
 		{
 			ALandscape* Landscape = this->EdMode->GetLandscape();
@@ -1515,6 +1528,8 @@ public:
 				Landscape->SetGrassUpdateEnabled(true);
 			}
 		}
+
+		TRACE_BOOKMARK(TEXT("EndTool - %s"), GetToolName());
 	}
 
 	virtual bool MouseMove(FEditorViewportClient* ViewportClient, FViewport* Viewport, int32 x, int32 y) override
diff --git a/Engine/Source/Editor/LandscapeEditor/Private/LandscapeEdModeXYOffset.cpp b/Engine/Source/Editor/LandscapeEditor/Private/LandscapeEdModeXYOffset.cpp
--- a/Engine/Source/Editor/LandscapeEditor/Private/LandscapeEdModeXYOffset.cpp
+++ b/Engine/Source/Editor/LandscapeEditor/Private/LandscapeEdModeXYOffset.cpp
@@ -51,15 +51,19 @@ namespace
 
 class FLandscapeToolStrokeRetopologize : public FLandscapeToolStrokeBase
 {
+	using Super = FLandscapeToolStrokeBase;
+
 public:
 	FLandscapeToolStrokeRetopologize(FEdModeLandscape* InEdMode, FEditorViewportClient* InViewportClient, const FLandscapeToolTarget& InTarget)
-		: FLandscapeToolStrokeBase(InEdMode, InViewportClient, InTarget)
+		: Super(InEdMode, InViewportClient, InTarget)
 		, Cache(InTarget)
 	{
 	}
 
 	void Apply(FEditorViewportClient* ViewportClient, FLandscapeBrush* Brush, const ULandscapeEditorObject* UISettings, const TArray<FLandscapeToolInteractorPosition>& InteractorPositions)
 	{
+		TRACE_CPUPROFILER_EVENT_SCOPE(FLandscapeToolStrokeRetopologize_Apply);
+
 		if (!LandscapeInfo)
 		{
 			return;
@@ -551,9 +555,11 @@ protected:
 
 class FLandscapeToolRetopologize : public FLandscapeToolBase<FLandscapeToolStrokeRetopologize>
 {
+	using Super = FLandscapeToolBase<FLandscapeToolStrokeRetopologize>;
+
 public:
 	FLandscapeToolRetopologize(FEdModeLandscape* InEdMode)
-		: FLandscapeToolBase<FLandscapeToolStrokeRetopologize>(InEdMode)
+		: Super(InEdMode)
 	{
 	}
 
