Landscape : Calculate AddComponent action resolution increase by actually working out where components will be created

diff --git a/Engine/Source/Editor/LandscapeEditor/Private/LandscapeEdModeComponentTools.cpp b/Engine/Source/Editor/LandscapeEditor/Private/LandscapeEdModeComponentTools.cpp
--- a/Engine/Source/Editor/LandscapeEditor/Private/LandscapeEdModeComponentTools.cpp
+++ b/Engine/Source/Editor/LandscapeEditor/Private/LandscapeEdModeComponentTools.cpp
@@ -934,7 +934,9 @@ public:
 		TOptional<FVector2D> LastMousePosition = CurrentBrush->GetLastMousePosition();
 		FIntRect LandscapeIndices;
 
-		if (ToolTarget.LandscapeInfo.IsValid() && LastMousePosition.IsSet() && ToolTarget.LandscapeInfo->GetLandscapeXYComponentBounds(LandscapeIndices))
+		ULandscapeInfo* LandscapeInfo = EdMode->CurrentToolTarget.LandscapeInfo.Get();
+
+		if ( LandscapeInfo != nullptr && ToolTarget.LandscapeInfo.IsValid() && LastMousePosition.IsSet() && ToolTarget.LandscapeInfo->GetLandscapeXYComponentBounds(LandscapeIndices))
 		{
 			const int32 BrushSize = FMath::Max(EdMode->UISettings->BrushComponentSize, 0);
 			const int32 ComponentSizeQuads = ToolTarget.LandscapeInfo->ComponentSizeQuads;
@@ -942,20 +944,23 @@ public:
 			const float BrushOriginY = LastMousePosition.GetValue().Y / ComponentSizeQuads - (BrushSize - 1) / 2.0f;
 			const int32 ComponentIndexX = FMath::FloorToInt(BrushOriginX);
 			const int32 ComponentIndexY = FMath::FloorToInt(BrushOriginY);
-			FIntPoint CurrentResolution = ToolTarget.LandscapeInfo->GetLandscapeProxy()->GetBoundingRect().Size() + 1;
-
-			if ((ComponentIndexX < LandscapeIndices.Min.X) || (ComponentIndexX > LandscapeIndices.Max.X))
-			{
-				ResolutionDelta += CurrentResolution.Y * BrushSize * ComponentSizeQuads;
-			}
 
-			if ((ComponentIndexY < LandscapeIndices.Min.Y) || (ComponentIndexY > LandscapeIndices.Max.Y))
+			int32 NumNewComponents = 0;
+			
+			int32 HalfBrushSize = BrushSize / 2;
+			FIntRect BrushSupport{ ComponentIndexX - HalfBrushSize, ComponentIndexY - HalfBrushSize, ComponentIndexX + HalfBrushSize, ComponentIndexY + HalfBrushSize };
+			for (int32 Y = BrushSupport.Min.Y; Y <= BrushSupport.Max.Y; Y++)
 			{
-				ResolutionDelta += CurrentResolution.X * BrushSize * ComponentSizeQuads;
+				for (int32 X = BrushSupport.Min.X; X <= BrushSupport.Max.X; X++)
+				{
+					NumNewComponents += LandscapeInfo->XYtoComponentMap.FindRef(FIntPoint(X, Y)) == nullptr ? 1 : 0;
+				}
 			}
+			
+			return NumNewComponents * ComponentSizeQuads * ComponentSizeQuads;
 		}
 
-		return ResolutionDelta;
+		return 0;
 	}
 
 private:
