Added Include Border to the landscape tool's component brush. Since the border is shared with neighbors, oftentimes the user wants to paint within the bounds of the component but not on the border (e.g. to quickly paint weightmaps on components but not have the paint "bleed" onto neighboring components)

diff --git a/Engine/Source/Editor/LandscapeEditor/Private/LandscapeEdModeBrushes.cpp b/Engine/Source/Editor/LandscapeEditor/Private/LandscapeEdModeBrushes.cpp
--- a/Engine/Source/Editor/LandscapeEditor/Private/LandscapeEdModeBrushes.cpp
+++ b/Engine/Source/Editor/LandscapeEditor/Private/LandscapeEdModeBrushes.cpp
@@ -390,6 +390,8 @@ class FLandscapeBrushComponent : public FLandscapeBrush
 protected:
 	FVector2D LastMousePosition;
 	UMaterialInterface* BrushMaterial;
+	FIntRect BrushExtentsInclusive; // True extents of the brush (in landscape coordinates)
+
 public:
 	FEdModeLandscape* EdMode;
 
@@ -445,6 +447,8 @@ public:
 				const float BrushOriginY = LastMousePosition.Y / LandscapeInfo->ComponentSizeQuads - (BrushSize - 1) / 2.0f;
 				const int32 ComponentIndexX = FMath::FloorToInt(BrushOriginX);
 				const int32 ComponentIndexY = FMath::FloorToInt(BrushOriginY);
+				BrushExtentsInclusive.Min = FIntPoint(ComponentIndexX * LandscapeInfo->ComponentSizeQuads, ComponentIndexY * LandscapeInfo->ComponentSizeQuads);
+				BrushExtentsInclusive.Max = FIntPoint((ComponentIndexX + BrushSize) * LandscapeInfo->ComponentSizeQuads + 1, (ComponentIndexY + BrushSize) * LandscapeInfo->ComponentSizeQuads + 1);
 
 				for (int32 YIndex = 0; YIndex < BrushSize; ++YIndex)
 				{
@@ -557,20 +561,30 @@ public:
 
 			for (int32 X = Bounds.Min.X; X < Bounds.Max.X; X++)
 			{
-				float PaintAmount = 1.0f;
-				if (EdMode->CurrentTool && EdMode->CurrentTool->GetToolType() != ELandscapeToolType::Mask
-					&& EdMode->UISettings->bUseSelectedRegion && LandscapeInfo->SelectedRegion.Num() > 0)
+				const int32 ComponentIndexX = FMath::FloorToInt(static_cast<float>(X) / LandscapeInfo->ComponentSizeQuads);
+				const int32 ComponentIndexY = FMath::FloorToInt(static_cast<float>(Y) / LandscapeInfo->ComponentSizeQuads);
+				// Skip the pixels from the border if requested :
+				if (EdMode->UISettings->bBrushComponentIncludeBorder 
+					|| ((X > BrushExtentsInclusive.Min.X)
+						&& (X < (BrushExtentsInclusive.Max.X - 1))
+						&& (Y > BrushExtentsInclusive.Min.Y)
+						&& (Y < (BrushExtentsInclusive.Max.Y - 1))))
 				{
-					float MaskValue = LandscapeInfo->SelectedRegion.FindRef(FIntPoint(X, Y));
-					if (EdMode->UISettings->bUseNegativeMask)
+					float PaintAmount = 1.0f;
+					if (EdMode->CurrentTool && EdMode->CurrentTool->GetToolType() != ELandscapeToolType::Mask
+						&& EdMode->UISettings->bUseSelectedRegion && LandscapeInfo->SelectedRegion.Num() > 0)
 					{
-						MaskValue = 1.0f - MaskValue;
+						float MaskValue = LandscapeInfo->SelectedRegion.FindRef(FIntPoint(X, Y));
+						if (EdMode->UISettings->bUseNegativeMask)
+						{
+							MaskValue = 1.0f - MaskValue;
+						}
+						PaintAmount *= MaskValue;
 					}
-					PaintAmount *= MaskValue;
-				}
 
-				// Set the brush value for this vertex
-				Scanline[X] = PaintAmount;
+					// Set the brush value for this vertex
+					Scanline[X] = PaintAmount;
+				}
 			}
 		}
 
diff --git a/Engine/Source/Editor/LandscapeEditor/Public/LandscapeEditorObject.h b/Engine/Source/Editor/LandscapeEditor/Public/LandscapeEditorObject.h
--- a/Engine/Source/Editor/LandscapeEditor/Public/LandscapeEditorObject.h
+++ b/Engine/Source/Editor/LandscapeEditor/Public/LandscapeEditorObject.h
@@ -650,11 +650,11 @@ public:
 	FLandscapePatternBrushWorldSpaceSettings WorldSpacePatternBrushSettings;
 
 	// Mask texture to use
-	UPROPERTY(Category="Brush Settings", EditAnywhere, NonTransactional, meta=(DisplayName="Texture", ShowForBrushes="BrushSet_Alpha,BrushSet_Pattern"))
+	UPROPERTY(Category  ="Brush Settings", EditAnywhere, NonTransactional, meta=(DisplayName="Texture", ShowForBrushes="BrushSet_Alpha,BrushSet_Pattern"))
 	TObjectPtr<UTexture2D> AlphaTexture;
 
 	// Channel of Mask Texture to use
-	UPROPERTY(Category="Brush Settings", EditAnywhere, NonTransactional, meta=(DisplayName="Texture Channel", ShowForBrushes="BrushSet_Alpha,BrushSet_Pattern"))
+	UPROPERTY(Category = "Brush Settings", EditAnywhere, NonTransactional, meta=(DisplayName="Texture Channel", ShowForBrushes="BrushSet_Alpha,BrushSet_Pattern"))
 	ELandscapeTextureColorChannel AlphaTextureChannel;
 
 	UPROPERTY(NonTransactional)
@@ -667,9 +667,12 @@ public:
 	// Component Brush:
 
 	// Number of components X/Y to affect at once. 1 means 1x1, 2 means 2x2, etc
-	UPROPERTY(Category="Brush Settings", EditAnywhere, NonTransactional, meta=(DisplayName="Brush Size", ShowForBrushes="BrushSet_Component", ClampMin="1", ClampMax="128", UIMin="1", UIMax="64", SliderExponent="3"))
+	UPROPERTY(Category = "Brush Settings", EditAnywhere, NonTransactional, meta=(DisplayName="Brush Size", ShowForBrushes="BrushSet_Component", ClampMin="1", ClampMax="128", UIMin="1", UIMax="64", SliderExponent="3"))
 	int32 BrushComponentSize;
 
+	UPROPERTY(Category = "Brush Settings", EditAnywhere, NonTransactional, meta = (DisplayName = "Include Border", ShowForBrushes = "BrushSet_Component", ShowForTools = "Paint,Sculpt", ToolTip = "When true, the brush will affect all pixels within the component, including those on the border, which means neightboring components (which share a line/row with their neighbor) will be affected as well"))
+	bool bBrushComponentIncludeBorder = true;
+
 
 	// Target Layer Settings:
 
