Tool changes to use local edit layer instead of shared state

The upgraded tools are:

FLandscapeToolSculpt
FLandscapeToolSmooth
FLandscapeToolFlatten
FLandscapeToolRamp
FLandscapeToolErosion
FLandscapeToolHydraErosion

The helper class FLandscapeLayerDataCache was fully upgraded to use explicit local edit layers (all tools that made use of it were upgraded, so no need for back compat).

For the tools, I decided to pass down the explicit edit layer via a "SetEditingLayer()" function -- legacy tools simply ignore that function call and continue to use the landscape shared state,
and upgraded tools will use it to set up the layer they should work in.

diff --git a/Engine/Source/Editor/LandscapeEditor/Private/LandscapeEdModeErosionTools.cpp b/Engine/Source/Editor/LandscapeEditor/Private/LandscapeEdModeErosionTools.cpp
--- a/Engine/Source/Editor/LandscapeEditor/Private/LandscapeEdModeErosionTools.cpp
+++ b/Engine/Source/Editor/LandscapeEditor/Private/LandscapeEdModeErosionTools.cpp
@@ -67,6 +67,12 @@ public:
 	{
 	}
 
+	virtual void SetEditLayer(const FGuid& EditLayerGUID) override
+	{
+		LayerHeightDataCache.SetCacheEditingLayer(EditLayerGUID);
+		WeightCache.DataAccess.SetEditLayer(EditLayerGUID);
+	}
+
 	void Apply(FEditorViewportClient* ViewportClient, FLandscapeBrush* Brush, const ULandscapeEditorObject* UISettings, const TArray<FLandscapeToolInteractorPosition>& InteractorPositions)
 	{
 		TRACE_CPUPROFILER_EVENT_SCOPE(FLandscapeToolStrokeErosion_Apply);
@@ -308,6 +314,12 @@ public:
 	{
 	}
 
+	virtual void SetEditLayer(const FGuid& EditLayerGUID) override
+	{
+		LayerHeightDataCache.SetCacheEditingLayer(EditLayerGUID);
+		WeightCache.DataAccess.SetEditLayer(EditLayerGUID);
+	}
+
 	void Apply(FEditorViewportClient* ViewportClient, FLandscapeBrush* Brush, const ULandscapeEditorObject* UISettings, const TArray<FLandscapeToolInteractorPosition>& InteractorPositions)
 	{
 		TRACE_CPUPROFILER_EVENT_SCOPE(FLandscapeToolStrokeHydraErosion_Apply);
diff --git a/Engine/Source/Editor/LandscapeEditor/Private/LandscapeEdModePaintTools.cpp b/Engine/Source/Editor/LandscapeEditor/Private/LandscapeEdModePaintTools.cpp
--- a/Engine/Source/Editor/LandscapeEditor/Private/LandscapeEdModePaintTools.cpp
+++ b/Engine/Source/Editor/LandscapeEditor/Private/LandscapeEdModePaintTools.cpp
@@ -434,6 +434,11 @@ public:
 	{
 	}
 
+	virtual void SetEditLayer(const FGuid& EditLayerGUID) override
+	{
+		this->Cache.DataAccess.SetEditLayer(EditLayerGUID);
+	}
+
 	void Apply(FEditorViewportClient* ViewportClient, FLandscapeBrush* Brush, const ULandscapeEditorObject* UISettings, const TArray<FLandscapeToolInteractorPosition>& InteractorPositions)
 	{
 		TRACE_CPUPROFILER_EVENT_SCOPE(FLandscapeToolStrokeSculpt_Apply);
@@ -690,6 +695,11 @@ public:
 	{
 	}
 
+	virtual void SetEditLayer(const FGuid& EditLayerGUID) override
+	{
+		LayerDataCache.SetCacheEditingLayer(EditLayerGUID);
+	}
+
 	void Apply(FEditorViewportClient* ViewportClient, FLandscapeBrush* Brush, const ULandscapeEditorObject* UISettings, const TArray<FLandscapeToolInteractorPosition>& InteractorPositions)
 	{
 		TRACE_CPUPROFILER_EVENT_SCOPE(FLandscapeToolStrokeSmooth_Apply);
@@ -848,6 +858,11 @@ public:
 		}
 	}
 
+	virtual void SetEditLayer(const FGuid& EditLayerGUID) override
+	{
+		LayerDataCache.SetCacheEditingLayer(EditLayerGUID);
+	}
+
 	void Apply(FEditorViewportClient* ViewportClient, FLandscapeBrush* Brush, const ULandscapeEditorObject* UISettings, const TArray<FLandscapeToolInteractorPosition>& InteractorPositions)
  	{
 		if (!this->LandscapeInfo) return;
@@ -873,15 +888,17 @@ public:
 			float V11 = 0.f;
 			if (bCombinedLayerOperation)
 			{
-				// Can't rely on cache in this mode
-				FScopedSetLandscapeEditingLayer Scope(Landscape, FGuid());
+				// create a new accessor and point it read from the final runtime data
+				typename ToolTarget::CacheClass::AccessorClass RuntimeDataAccessor(this->Target);
+				RuntimeDataAccessor.SetEditLayer(FGuid());
+
 				typename ToolTarget::CacheClass::DataType P00, P10, P01, P11;
-				this->Cache.DataAccess.GetDataFast(FlattenHeightX,     FlattenHeightY,     FlattenHeightX,     FlattenHeightY,     &P00);
-				this->Cache.DataAccess.GetDataFast(FlattenHeightX + 1, FlattenHeightY,     FlattenHeightX + 1, FlattenHeightY,     &P10);
-				this->Cache.DataAccess.GetDataFast(FlattenHeightX,     FlattenHeightY + 1, FlattenHeightX,     FlattenHeightY + 1, &P01);
-				this->Cache.DataAccess.GetDataFast(FlattenHeightX + 1, FlattenHeightY + 1, FlattenHeightX + 1, FlattenHeightY + 1, &P11);
+				RuntimeDataAccessor.GetDataFast(FlattenHeightX, FlattenHeightY, FlattenHeightX, FlattenHeightY, &P00);
+				RuntimeDataAccessor.GetDataFast(FlattenHeightX + 1, FlattenHeightY, FlattenHeightX + 1, FlattenHeightY, &P10);
+				RuntimeDataAccessor.GetDataFast(FlattenHeightX, FlattenHeightY + 1, FlattenHeightX, FlattenHeightY + 1, &P01);
+				RuntimeDataAccessor.GetDataFast(FlattenHeightX + 1, FlattenHeightY + 1, FlattenHeightX + 1, FlattenHeightY + 1, &P11);
 				// Release Texture Mips that will be Locked by the next SynchronousUpdateComponentVisibilityForHeight (inside the LayerDataCache.Read call)
-				this->Cache.DataAccess.Flush();
+				RuntimeDataAccessor.Flush();
 				V00 = P00;
 				V10 = P10;
 				V01 = P01;
diff --git a/Engine/Source/Editor/LandscapeEditor/Private/LandscapeEdModeRampTool.cpp b/Engine/Source/Editor/LandscapeEditor/Private/LandscapeEdModeRampTool.cpp
--- a/Engine/Source/Editor/LandscapeEditor/Private/LandscapeEdModeRampTool.cpp
+++ b/Engine/Source/Editor/LandscapeEditor/Private/LandscapeEdModeRampTool.cpp
@@ -491,7 +491,7 @@ public:
 
 		FScopedTransaction Transaction(LOCTEXT("Ramp_Apply", "Landscape Editing: Add ramp"));
 		ALandscape* Landscape = EdMode->GetLandscape();
-		FScopedSetLandscapeEditingLayer Scope(Landscape, EdMode->GetCurrentLayerGuid(), [&] { if (Landscape) { Landscape->RequestLayersContentUpdate(ELandscapeLayerUpdateMode::Update_Heightmap_All); } });
+		FGuid EditLayerGUID = EdMode->GetCurrentLayerGuid();
 
 		const ULandscapeInfo* LandscapeInfo = EdMode->CurrentToolTarget.LandscapeInfo.Get();
 		const ALandscapeProxy* LandscapeProxy = LandscapeInfo->GetLandscapeProxy();
@@ -540,9 +540,12 @@ public:
 			return;
 		}
 
-		FLandscapeEditDataInterface LandscapeEdit(EdMode->CurrentToolTarget.LandscapeInfo.Get());
+		// construct the caches, and set them to work in the EditLayer
+		FLandscapeEditDataInterface LandscapeEdit(EdMode->CurrentToolTarget.LandscapeInfo.Get(), EditLayerGUID);
 		FLandscapeHeightCache HeightCache(EdMode->CurrentToolTarget);
 		FLandscapeLayerDataCache<FHeightmapToolTarget> LayerHeightDataCache(EdMode->CurrentToolTarget, HeightCache);
+		LayerHeightDataCache.SetCacheEditingLayer(EditLayerGUID);
+
 		const bool bCombinedLayerOperation = EdMode->UISettings->bCombinedLayersOperation && Landscape && Landscape->HasLayersContent();
 		LayerHeightDataCache.Initialize(EdMode->CurrentToolTarget.LandscapeInfo.Get(), bCombinedLayerOperation);
 
@@ -608,6 +611,10 @@ public:
 				}
 			}
 		}
+ 		if (Landscape)
+ 		{
+ 			Landscape->RequestLayersContentUpdate(ELandscapeLayerUpdateMode::Update_Heightmap_All);
+ 		}
 	}
 
 	bool CanApplyRamp()
diff --git a/Engine/Source/Editor/LandscapeEditor/Private/LandscapeEdModeTools.h b/Engine/Source/Editor/LandscapeEditor/Private/LandscapeEdModeTools.h
--- a/Engine/Source/Editor/LandscapeEditor/Private/LandscapeEdModeTools.h
+++ b/Engine/Source/Editor/LandscapeEditor/Private/LandscapeEdModeTools.h
@@ -285,6 +285,8 @@ struct TLandscapeEditCache
 {
 public:
 	typedef AccessorType DataType;
+	typedef Accessor AccessorClass;
+
 	Accessor DataAccess;
 
 	TLandscapeEditCache(const FLandscapeToolTarget& InTarget)
@@ -917,6 +919,7 @@ public:
 	FLandscapeLayerDataCache(const FLandscapeToolTarget& InTarget, typename ToolTarget::CacheClass& Cache)
 		: LandscapeInfo(nullptr)
 		, Landscape(nullptr)
+		, EditingLayerGuid()
 		, EditingLayerIndex(MAX_uint8)
 		, bIsInitialized(false)
 		, bCombinedLayerOperation(false)
@@ -927,8 +930,16 @@ public:
 	{
 	}
 
+	void SetCacheEditingLayer(const FGuid& InEditLayerGUID)
+	{
+		CacheUpToEditingLayer.DataAccess.SetEditLayer(InEditLayerGUID);
+		CacheBottomLayers.DataAccess.SetEditLayer(InEditLayerGUID);
+		EditingLayerGuid = InEditLayerGUID;
+	}
+
 	void Initialize(ULandscapeInfo* InLandscapeInfo, bool InCombinedLayerOperation)
 	{
+		check(EditingLayerGuid.IsSet());	// you must call SetCacheEditingLayer before Initialize
 		if (!bIsInitialized)
 		{
 			LandscapeInfo = InLandscapeInfo;
@@ -936,12 +947,11 @@ public:
 			bCombinedLayerOperation = Landscape && Landscape->HasLayersContent() && InCombinedLayerOperation && bTargetIsHeightmap;
 			if (bCombinedLayerOperation)
 			{
-				EditingLayerGuid = Landscape->GetEditingLayer();
 				for (int i = 0; i < Landscape->GetLayerCount(); ++i)
 				{
 					FLandscapeLayer* CurrentLayer = Landscape->GetLayer(i);
 					BackupLayerVisibility.Add(CurrentLayer->bVisible);
-					if (CurrentLayer->Guid == EditingLayerGuid)
+					if (CurrentLayer->Guid == EditingLayerGuid.GetValue())
 					{
 						EditingLayerIndex = i;
 					}
@@ -952,6 +962,7 @@ public:
 		}
 	}
 
+	// read values in the specified rectangle into the array
 	void Read(int32 X1, int32 Y1, int32 X2, int32 Y2, TArray<typename ToolTarget::CacheClass::DataType>& Data)
 	{
 		check(bIsInitialized);
@@ -987,7 +998,10 @@ public:
 				return DesiredCacheBounds;
 			};
 
-			FScopedSetLandscapeEditingLayer Scope(Landscape, FGuid());
+			// temporarily switch to working on the final runtime data, so we can gather the combined layer data into the caches
+			FGuid PreviousLayerGUID = EditingLayerGuid.GetValue();
+			SetCacheEditingLayer(FGuid());
+
 			CacheUpToEditingLayer.GetDataAndCache(X1, Y1, X2, Y2, Data, OnCacheUpdating);
 			// Release Texture Mips that will be Locked by the next SynchronousUpdateComponentVisibilityForHeight
 			CacheUpToEditingLayer.DataAccess.Flush();
@@ -997,6 +1011,9 @@ public:
 			CacheBottomLayers.GetDataAndCache(X1, Y1, X2, Y2, BottomLayersData, OnCacheUpdating);
 			// Do the same here for consistency
 			CacheBottomLayers.DataAccess.Flush();
+			
+			SetCacheEditingLayer(PreviousLayerGUID);
+			check(PreviousLayerGUID == CacheUpToEditingLayer.DataAccess.GetEditLayer());
 		}
 		else
 		{
@@ -1021,9 +1038,11 @@ public:
 				float Contribution = (LandscapeDataAccess::GetLocalHeight(Data[i]) - LandscapeDataAccess::GetLocalHeight(BottomLayersData[i])) * InverseAlpha;
 				DataContribution[i] = LandscapeDataAccess::GetTexHeight(Contribution);
 			}
-			checkf(EditingLayerGuid == Landscape->GetEditingLayer(), TEXT("EditingLayer has changed between Initialize and Write. Was: %s (%s). Is now: %s (%s)"), 
-				Landscape->GetLayer(EditingLayerGuid) ? *(Landscape->GetLayer(EditingLayerGuid)->Name.ToString()) : TEXT("<unknown>"), *EditingLayerGuid.ToString(),
-				Landscape->GetLayer(Landscape->GetEditingLayer()) ? *(Landscape->GetLayer(Landscape->GetEditingLayer())->Name.ToString()) : TEXT("<unknown>"), *Landscape->GetEditingLayer().ToString());
+
+			FGuid CacheAccessorLayerGuid = CacheUpToEditingLayer.DataAccess.GetEditLayer();
+ 			checkf(EditingLayerGuid.GetValue() == CacheAccessorLayerGuid, TEXT("Editing Layer has changed between Initialize and Write. Was: %s (%s). Is now: %s (%s)"),
+ 				Landscape->GetLayer(*EditingLayerGuid) ? *(Landscape->GetLayer(*EditingLayerGuid)->Name.ToString()) : TEXT("<unknown>"), *EditingLayerGuid->ToString(),
+ 				Landscape->GetLayer(CacheAccessorLayerGuid) ? *(Landscape->GetLayer(CacheAccessorLayerGuid)->Name.ToString()) : TEXT("<unknown>"), *CacheAccessorLayerGuid.ToString());
 
 			// Restore layers visibility
 			SetLayersVisibility(BackupLayerVisibility);
@@ -1087,7 +1106,7 @@ private:
 
 	ULandscapeInfo* LandscapeInfo;
 	ALandscape* Landscape;
-	FGuid EditingLayerGuid;
+	TOptional<FGuid> EditingLayerGuid;
 	uint8 EditingLayerIndex;
 	TArray<bool> BackupLayerVisibility;
 	TArray<typename ToolTarget::CacheClass::DataType> BottomLayersData;
@@ -1141,6 +1160,16 @@ struct FFullWeightmapAccessor
 		}
 	}
 
+	void SetEditLayer(const FGuid& InEditLayerGUID)
+	{
+		LandscapeEdit.SetEditLayer(InEditLayerGUID);
+	}
+
+	FGuid GetEditLayer() const
+	{
+		return LandscapeEdit.GetEditLayer();
+	}
+
 	void GetData(int32& X1, int32& Y1, int32& X2, int32& Y2, TMap<FIntPoint, TArray<uint8>>& Data)
 	{
 		// Do not Support for interpolation....
@@ -1376,6 +1405,13 @@ public:
 	{
 	}
 
+	virtual void SetEditLayer(const FGuid& EditLayerGUID)
+	{
+		// if this function is not overridden, then the tool uses the old method of getting the edit layer (using the shared EditingLayer on ALandscape)
+		// we should migrate tools to use SetEditLayer() and deprecate the reliance on the shared EditingLayer
+		// once all tools use SetEditLayer we can move this to be part of the constructor
+	}
+
 	// Signature of Apply() method for derived strokes
 	// void Apply(FEditorViewportClient* ViewportClient, FLandscapeBrush* Brush, const ULandscapeEditorObject* UISettings, const TArray<FLandscapeToolMousePosition>& MousePositions);
 
@@ -1447,7 +1483,7 @@ public:
 			if (Landscape)
 			{
 				Landscape->RequestLayersContentUpdate(GetBeginToolContentUpdateFlag());
-				Landscape->SetEditingLayer(this->EdMode->GetCurrentLayerGuid());
+				Landscape->SetEditingLayer(this->EdMode->GetCurrentLayerGuid());	// legacy way to set the edit layer, via Landscape state
 				Landscape->SetGrassUpdateEnabled(false);
 			}
 		}
@@ -1457,13 +1493,14 @@ public:
 			InteractorPositions.Empty(1);
 		}
 
-		if( !IsToolActive() )
+		if (ensure(!IsToolActive()))
 		{
-			ToolStroke.Emplace( EdMode, ViewportClient, InTarget );
+			ToolStroke.Emplace( EdMode, ViewportClient, InTarget );				// construct the tool stroke class
+			ToolStroke->SetEditLayer(this->EdMode->GetCurrentLayerGuid());		// set the edit layer explicitly (if the tool supports this path)
 			EdMode->CurrentBrush->BeginStroke( InHitLocation.X, InHitLocation.Y, this );
 		}
 
-		// Save the mouse position
+		// Save the mouse position  
 		LastInteractorPosition = FVector2D(InHitLocation);
 		InteractorPositions.Emplace(LastInteractorPosition, ViewportClient ? IsModifierPressed(ViewportClient) : false); // Copy tool sometimes activates without a specific viewport via ctrl+c hotkey
 		TimeSinceLastInteractorMove = 0.0f;
@@ -1515,7 +1552,7 @@ public:
 			InteractorPositions.Empty(1);
 		}
 
-		ToolStroke.Reset();
+		ToolStroke.Reset();		// destruct the tool stroke class
 		EdMode->CurrentBrush->EndStroke();
 		EdMode->UpdateLayerUsageInformation(&EdMode->CurrentToolTarget.LayerInfo);
 
