Fixed typos

diff --git a/Engine/Source/Editor/LandscapeEditor/Private/LandscapeEditorDetailCustomization_TargetLayers.cpp b/Engine/Source/Editor/LandscapeEditor/Private/LandscapeEditorDetailCustomization_TargetLayers.cpp
--- a/Engine/Source/Editor/LandscapeEditor/Private/LandscapeEditorDetailCustomization_TargetLayers.cpp
+++ b/Engine/Source/Editor/LandscapeEditor/Private/LandscapeEditorDetailCustomization_TargetLayers.cpp
@@ -340,7 +340,7 @@ TSharedRef<SWidget> FLandscapeEditorCustomNodeBuilder_TargetLayers::GetTargetLay
 {
 	FMenuBuilder MenuBuilder(/*bInShouldCloseWindowAfterMenuSelection=*/true, nullptr, nullptr, /*bCloseSelfOnly=*/ true);
 
-	MenuBuilder.BeginSection("TargetLayerUnusedType", LOCTEXT("UnusedTypeHeading", "Layer Visilibity"));
+	MenuBuilder.BeginSection("TargetLayerUnusedType", LOCTEXT("UnusedTypeHeading", "Layer Visibility"));
 	{
 		MenuBuilder.AddMenuEntry(
 			LOCTEXT("TargetLayerShowUnusedLayer", "Show all layers"),
