Landscape : Guard against missing HitProxy when adding control points to a hidden landscape.

diff --git a/Engine/Source/Editor/LandscapeEditor/Private/LandscapeEdModeSplineTools.cpp b/Engine/Source/Editor/LandscapeEditor/Private/LandscapeEdModeSplineTools.cpp
--- a/Engine/Source/Editor/LandscapeEditor/Private/LandscapeEdModeSplineTools.cpp
+++ b/Engine/Source/Editor/LandscapeEditor/Private/LandscapeEdModeSplineTools.cpp
@@ -978,7 +978,8 @@ public:
 
 	virtual bool BeginTool(FEditorViewportClient* ViewportClient, const FLandscapeToolTarget& InTarget, const FVector& InHitLocation) override
 	{
-		if (ViewportClient->IsCtrlPressed())
+		HHitProxy* HitProxy = ViewportClient->Viewport->GetHitProxy(ViewportClient->Viewport->GetMouseX(), ViewportClient->Viewport->GetMouseY());
+		if (HitProxy && ViewportClient->IsCtrlPressed())
 		{
 			LandscapeInfo = InTarget.LandscapeInfo.Get();
 			ILandscapeSplineInterface* SplineOwner = nullptr;
@@ -1000,7 +1001,6 @@ public:
 			// Hit Test
 			if (!SplineOwner)
 			{
-				HHitProxy* HitProxy = ViewportClient->Viewport->GetHitProxy(ViewportClient->Viewport->GetMouseX(), ViewportClient->Viewport->GetMouseY());
 				if (HitProxy->IsA(HActor::StaticGetType()))
 				{
 					HActor* ActorProxy = (HActor*)HitProxy;
