Landscape: Fix crash when updating property permissions outside landscape mode

diff --git a/Engine/Source/Editor/LandscapeEditor/Private/LandscapeEditorDetails.cpp b/Engine/Source/Editor/LandscapeEditor/Private/LandscapeEditorDetails.cpp
--- a/Engine/Source/Editor/LandscapeEditor/Private/LandscapeEditorDetails.cpp
+++ b/Engine/Source/Editor/LandscapeEditor/Private/LandscapeEditorDetails.cpp
@@ -41,7 +41,11 @@ TSharedRef<IDetailCustomization> FLandscapeEditorDetails::MakeInstance()
 void FLandscapeEditorDetails::CustomizeDetails(IDetailLayoutBuilder& DetailBuilder)
 {
 	FEdModeLandscape* LandscapeEdMode = GetEditorMode();
-	CommandList = LandscapeEdMode->GetUICommandList();
+
+	if (LandscapeEdMode == nullptr)
+	{
+		return;
+	}
 
 	static const FLinearColor BorderColor = FLinearColor(0.2f, 0.2f, 0.2f, 0.2f);
 	static const FSlateBrush* BorderStyle = FAppStyle::GetBrush("DetailsView.GroupSection");
@@ -336,89 +340,6 @@ FSlateIcon FLandscapeEditorDetails::GetCurrentToolIcon() const
 	return FSlateIcon(FAppStyle::GetAppStyleSetName(), "Default");
 }
 
-TSharedRef<SWidget> FLandscapeEditorDetails::GetToolSelector()
-{
-	FEdModeLandscape* LandscapeEdMode = GetEditorMode();
-	if (LandscapeEdMode != NULL)
-	{
-		auto NameToCommandMap = FLandscapeEditorCommands::Get().NameToCommandMap;
-
-		FToolMenuBuilder MenuBuilder(true, CommandList);
-
-		if (LandscapeEdMode->CurrentToolMode->ToolModeName == "ToolMode_Manage")
-		{
-			MenuBuilder.BeginSection(NAME_None, LOCTEXT("NewLandscapeToolsTitle", "New Landscape"));
-			MenuBuilder.AddToolButton(NameToCommandMap.FindChecked("Tool_NewLandscape"), NAME_None, LOCTEXT("Tool.NewLandscape", "New Landscape"), LOCTEXT("Tool.NewLandscape.Tooltip", "Create or import a new landscape"));
-			MenuBuilder.EndSection();
-
-			MenuBuilder.BeginSection(NAME_None, LOCTEXT("ComponentToolsTitle", "Component Tools"));
-			MenuBuilder.AddToolButton(NameToCommandMap.FindChecked("Tool_Select"), NAME_None, LOCTEXT("Tool.SelectComponent", "Selection"), LOCTEXT("Tool.SelectComponent.Tooltip", "Select components to use with other tools"));
-			MenuBuilder.AddToolButton(NameToCommandMap.FindChecked("Tool_AddComponent"), NAME_None, LOCTEXT("Tool.AddComponent", "Add"), LOCTEXT("Tool.AddComponent.Tooltip", "Add components to the landscape"));
-			MenuBuilder.AddToolButton(NameToCommandMap.FindChecked("Tool_DeleteComponent"), NAME_None, LOCTEXT("Tool.DeleteComponent", "Delete"), LOCTEXT("Tool.DeleteComponent.Tooltip", "Delete components from the landscape, leaving a hole"));
-			MenuBuilder.AddToolButton(NameToCommandMap.FindChecked("Tool_MoveToLevel"), NAME_None, LOCTEXT("Tool.MoveToLevel", "Move to Level"), LOCTEXT("Tool.MoveToLevel.Tooltip", "Move landscape components to a landscape proxy in the currently active streaming level, so that they can be streamed in/out independently of the rest of the landscape"));
-			MenuBuilder.AddToolButton(NameToCommandMap.FindChecked("Tool_ResizeLandscape"), NAME_None, LOCTEXT("Tool.ResizeLandscape", "Change Component Size"), LOCTEXT("Tool.ResizeLandscape.Tooltip", "Change the size of the landscape components"));
-			MenuBuilder.EndSection();
-
-			MenuBuilder.BeginSection(NAME_None, LOCTEXT("SplineToolsTitle", "Spline Tools"));
-			MenuBuilder.AddToolButton(NameToCommandMap.FindChecked("Tool_Splines"), NAME_None, LOCTEXT("Tool.Spline", "Edit Splines"), LOCTEXT("Tool.Spline.Tooltip", "Ctrl+click to add control points\nHaving a control point selected when you ctrl+click will connect to the new control point with a segment\nSpline mesh settings can be found on the details panel when you have segments selected"));
-			MenuBuilder.EndSection();
-		}
-
-		if (LandscapeEdMode->CurrentToolMode->ToolModeName == "ToolMode_Sculpt")
-		{
-			MenuBuilder.BeginSection(NAME_None, LOCTEXT("SculptToolsTitle", "Sculpting Tools"));
-			MenuBuilder.AddToolButton(NameToCommandMap.FindChecked("Tool_Sculpt"), NAME_None, LOCTEXT("Tool.Sculpt", "Sculpt"), LOCTEXT("Tool.Sculpt.Tooltip", "Sculpt height data.\nClick to Raise, Shift+Click to lower"));
-			
-			if (LandscapeEdMode->CanHaveLandscapeLayersContent())
-			{
-				MenuBuilder.AddToolButton(NameToCommandMap.FindChecked("Tool_Erase"), NAME_None, LOCTEXT("Tool.Erase", "Erase"), LOCTEXT("Tool.Erase.Tooltip", "Erase height data."));
-			}
-			
-			MenuBuilder.AddToolButton(NameToCommandMap.FindChecked("Tool_Smooth"), NAME_None, LOCTEXT("Tool.Smooth", "Smooth"), LOCTEXT("Tool.Smooth.Tooltip", "Smooths heightmaps or blend layers"));
-			MenuBuilder.AddToolButton(NameToCommandMap.FindChecked("Tool_Flatten"), NAME_None, LOCTEXT("Tool.Flatten", "Flatten"), LOCTEXT("Tool.Flatten.Tooltip", "Flattens an area of heightmap or blend layer"));
-			MenuBuilder.AddToolButton(NameToCommandMap.FindChecked("Tool_Ramp"), NAME_None, LOCTEXT("Tool.Ramp", "Ramp"), LOCTEXT("Tool.Ramp.Tooltip", "Creates a ramp between two points"));
-			MenuBuilder.AddToolButton(NameToCommandMap.FindChecked("Tool_Erosion"), NAME_None, LOCTEXT("Tool.Erosion", "Erosion"), LOCTEXT("Tool.Erosion.Tooltip", "Thermal Erosion - Simulates erosion caused by the movement of soil from higher areas to lower areas"));
-			MenuBuilder.AddToolButton(NameToCommandMap.FindChecked("Tool_HydraErosion"), NAME_None, LOCTEXT("Tool.HydroErosion", "Hydro Erosion"), LOCTEXT("Tool.HydroErosion.Tooltip", "Hydro Erosion - Simulates erosion caused by rainfall"));
-			MenuBuilder.AddToolButton(NameToCommandMap.FindChecked("Tool_Noise"), NAME_None, LOCTEXT("Tool.Noise", "Noise"), LOCTEXT("Tool.Noise.Tooltip", "Adds noise to the heightmap or blend layer"));
-			MenuBuilder.AddToolButton(NameToCommandMap.FindChecked("Tool_Retopologize"), NAME_None, LOCTEXT("Tool.Retopologize", "Retopologize"), LOCTEXT("Tool.Retopologize.Tooltip", "Automatically adjusts landscape vertices with an X/Y offset map to improve vertex density on cliffs, reducing texture stretching.\nNote: An X/Y offset map makes the landscape slower to render and paint on with other tools, so only use if needed"));
-			MenuBuilder.AddToolButton(NameToCommandMap.FindChecked("Tool_Visibility"), NAME_None, LOCTEXT("Tool.Visibility", "Visibility"), LOCTEXT("Tool.Visibility.Tooltip", "Mask out individual quads in the landscape, leaving a hole."));
-
-			if (LandscapeEdMode->CanHaveLandscapeLayersContent())
-			{
-				MenuBuilder.AddToolButton(NameToCommandMap.FindChecked("Tool_BlueprintBrush"), NAME_None, LOCTEXT("Tool.SculptBlueprintBrush", "Blueprint Brushes"), LOCTEXT("Tool.SculptBlueprintBrush.Tooltip", "Custom sculpting tools created using Blueprint."));
-			}
-
-			MenuBuilder.EndSection();
-
-			MenuBuilder.BeginSection(NAME_None, LOCTEXT("RegionToolsTitle", "Region Tools"));
-			MenuBuilder.AddToolButton(NameToCommandMap.FindChecked("Tool_Mask"), NAME_None, LOCTEXT("Tool.RegionSelect", "Selection"), LOCTEXT("Tool.RegionSelect.Tooltip", "Select a region of landscape to use as a mask for other tools"));
-			MenuBuilder.AddToolButton(NameToCommandMap.FindChecked("Tool_CopyPaste"), NAME_None, LOCTEXT("Tool.RegionCopyPaste", "Copy/Paste"), LOCTEXT("Tool.RegionCopyPaste.Tooltip", "Copy/Paste areas of the landscape, or import/export a copied area of landscape from disk"));
-			MenuBuilder.AddToolButton(NameToCommandMap.FindChecked("Tool_Mirror"), NAME_None, LOCTEXT("Tool.Mirror", "Mirror"), LOCTEXT("Tool.Mirror.Tooltip", "Copies one side of a landscape to the other, to easily create a mirrored landscape."));
-			MenuBuilder.EndSection();
-		}
-
-		if (LandscapeEdMode->CurrentToolMode->ToolModeName == "ToolMode_Paint")
-		{
-			MenuBuilder.BeginSection(NAME_None, LOCTEXT("PaintToolsTitle", "Paint Tools"));
-			MenuBuilder.AddToolButton(NameToCommandMap.FindChecked("Tool_Paint"), NAME_None, LOCTEXT("Tool.Paint", "Paint"), LOCTEXT("Tool.Paint.Tooltip", "Paints weight data.\nClick to paint, Shift+Click to erase"));
-			MenuBuilder.AddToolButton(NameToCommandMap.FindChecked("Tool_Smooth"), NAME_None, LOCTEXT("Tool.Smooth", "Smooth"), LOCTEXT("Tool.Smooth.Tooltip", "Smooths heightmaps or blend layers"));
-			MenuBuilder.AddToolButton(NameToCommandMap.FindChecked("Tool_Flatten"), NAME_None, LOCTEXT("Tool.Flatten", "Flatten"), LOCTEXT("Tool.Flatten.Tooltip", "Flattens an area of heightmap or blend layer"));
-			MenuBuilder.AddToolButton(NameToCommandMap.FindChecked("Tool_Noise"), NAME_None, LOCTEXT("Tool.Noise", "Noise"), LOCTEXT("Tool.Noise.Tooltip", "Adds noise to the heightmap or blend layer"));
-
-			if (LandscapeEdMode->CanHaveLandscapeLayersContent())
-			{
-				MenuBuilder.AddToolButton(NameToCommandMap.FindChecked("Tool_BlueprintBrush"), NAME_None, LOCTEXT("Tool.PaintBlueprintBrush", "Blueprint Brushes"), LOCTEXT("Tool.PaintBlueprintBrush.Tooltip", "Custom painting tools created using Blueprint."));
-			}
-
-			MenuBuilder.EndSection();
-		}
-
-		return MenuBuilder.MakeWidget();
-	}
-
-	return SNullWidget::NullWidget;
-}
-
 bool FLandscapeEditorDetails::GetToolSelectorIsVisible() const
 {
 	FEdModeLandscape* LandscapeEdMode = GetEditorMode();
@@ -488,48 +409,6 @@ FSlateIcon FLandscapeEditorDetails::GetCurrentBrushIcon() const
 	return FSlateIcon(FAppStyle::GetAppStyleSetName(), "Default");
 }
 
-TSharedRef<SWidget> FLandscapeEditorDetails::GetBrushSelector()
-{
-	FEdModeLandscape* LandscapeEdMode = GetEditorMode();
-	if (LandscapeEdMode && LandscapeEdMode->CurrentTool)
-	{
-		auto NameToCommandMap = FLandscapeEditorCommands::Get().NameToCommandMap;
-
-		FToolMenuBuilder MenuBuilder(true, CommandList);
-		MenuBuilder.BeginSection(NAME_None, LOCTEXT("BrushesTitle", "Brushes"));
-
-		if (LandscapeEdMode->CurrentTool->ValidBrushes.Contains("BrushSet_Circle"))
-		{
-			MenuBuilder.AddToolButton(NameToCommandMap.FindChecked("BrushSet_Circle"), NAME_None, LOCTEXT("Brush.Circle", "Circle"), LOCTEXT("Brush.Circle.Brushtip", "Simple circular brush"));
-		}
-		if (LandscapeEdMode->CurrentTool->ValidBrushes.Contains("BrushSet_Alpha"))
-		{
-			MenuBuilder.AddToolButton(NameToCommandMap.FindChecked("BrushSet_Alpha"), NAME_None, LOCTEXT("Brush.Alpha.Alpha", "Alpha"), LOCTEXT("Brush.Alpha.Alpha.Tooltip", "Alpha brush, orients a mask image with the brush stroke"));
-		}
-		if (LandscapeEdMode->CurrentTool->ValidBrushes.Contains("BrushSet_Pattern"))
-		{
-			MenuBuilder.AddToolButton(NameToCommandMap.FindChecked("BrushSet_Pattern"), NAME_None, LOCTEXT("Brush.Alpha.Pattern", "Pattern"), LOCTEXT("Brush.Alpha.Pattern.Tooltip", "Pattern brush, tiles a mask image across the landscape"));
-		}
-		if (LandscapeEdMode->CurrentTool->ValidBrushes.Contains("BrushSet_Component"))
-		{
-			MenuBuilder.AddToolButton(NameToCommandMap.FindChecked("BrushSet_Component"), NAME_None, LOCTEXT("Brush.Component", "Component"), LOCTEXT("Brush.Component.Brushtip", "Work with entire landscape components"));
-		}
-		if (LandscapeEdMode->CurrentTool->ValidBrushes.Contains("BrushSet_Gizmo"))
-		{
-			MenuBuilder.AddToolButton(NameToCommandMap.FindChecked("BrushSet_Gizmo"), NAME_None, LOCTEXT("Brush.Gizmo", "Gizmo"), LOCTEXT("Brush.Gizmo.Brushtip", "Work with the landscape gizmo, used for copy/pasting landscape"));
-		}
-		//if (LandscapeEdMode->CurrentTool->ValidBrushes.Contains("BrushSet_Splines"))
-		//{
-		//	MenuBuilder.AddToolButton(NameToCommandMap.FindChecked("BrushSet_Splines"), NAME_None, LOCTEXT("Brush.Splines", "Splines"), LOCTEXT("Brush.Splines.Brushtip", "Edit Splines"));
-		//}
-		MenuBuilder.EndSection();
-
-		return MenuBuilder.MakeWidget();
-	}
-
-	return SNullWidget::NullWidget;
-}
-
 bool FLandscapeEditorDetails::GetBrushSelectorIsVisible() const
 {
 	FEdModeLandscape* LandscapeEdMode = GetEditorMode();
@@ -590,6 +469,7 @@ FSlateIcon FLandscapeEditorDetails::GetCurrentBrushFalloffIcon() const
 
 void FLandscapeEditorDetails::SetBrushCommand(FName InBrush)
 {
+	TSharedPtr<FUICommandList> CommandList = GetEditorMode()->GetUICommandList();
 	TSharedPtr<FUICommandInfo> BrushCommand = FLandscapeEditorCommands::Get().NameToCommandMap.FindRef(InBrush);
 	if (CommandList.IsValid() && BrushCommand.IsValid())
 	{
@@ -597,27 +477,6 @@ void FLandscapeEditorDetails::SetBrushCommand(FName InBrush)
 	}
 }
 
-TSharedRef<SWidget> FLandscapeEditorDetails::GetBrushFalloffSelector()
-{
-	FEdModeLandscape* LandscapeEdMode = GetEditorMode();
-	if (LandscapeEdMode && LandscapeEdMode->CurrentTool)
-	{
-		auto NameToCommandMap = FLandscapeEditorCommands::Get().NameToCommandMap;
-
-		FToolMenuBuilder MenuBuilder(true, CommandList);
-		MenuBuilder.BeginSection(NAME_None, LOCTEXT("FalloffTitle", "Falloff"));
-		MenuBuilder.AddToolButton(FLandscapeEditorCommands::Get().CircleBrush_Smooth,    NAME_None, LOCTEXT("Brush.Circle.Smooth", "Smooth"),       LOCTEXT("Brush.Circle.Smooth.Tooltip", "Smooth falloff"));
-		MenuBuilder.AddToolButton(FLandscapeEditorCommands::Get().CircleBrush_Linear,    NAME_None, LOCTEXT("Brush.Circle.Linear", "Linear"),       LOCTEXT("Brush.Circle.Linear.Tooltip", "Sharp, linear falloff"));
-		MenuBuilder.AddToolButton(FLandscapeEditorCommands::Get().CircleBrush_Spherical, NAME_None, LOCTEXT("Brush.Circle.Spherical", "Spherical"), LOCTEXT("Brush.Circle.Spherical.Tooltip", "Spherical falloff, smooth at the center and sharp at the edge"));
-		MenuBuilder.AddToolButton(FLandscapeEditorCommands::Get().CircleBrush_Tip,       NAME_None, LOCTEXT("Brush.Circle.Tip", "Tip"),             LOCTEXT("Brush.Circle.Tip.Tooltip", "Tip falloff, sharp at the center and smooth at the edge"));
-		MenuBuilder.EndSection();
-
-		return MenuBuilder.MakeWidget();
-	}
-
-	return SNullWidget::NullWidget;
-}
-
 bool FLandscapeEditorDetails::GetBrushFalloffSelectorIsVisible() const
 {
 	FEdModeLandscape* LandscapeEdMode = GetEditorMode();
diff --git a/Engine/Source/Editor/LandscapeEditor/Private/LandscapeEditorDetails.h b/Engine/Source/Editor/LandscapeEditor/Private/LandscapeEditorDetails.h
--- a/Engine/Source/Editor/LandscapeEditor/Private/LandscapeEditorDetails.h
+++ b/Engine/Source/Editor/LandscapeEditor/Private/LandscapeEditorDetails.h
@@ -54,21 +54,18 @@ protected:
 
 	FText GetCurrentToolName() const;
 	FSlateIcon GetCurrentToolIcon() const;
-	TSharedRef<SWidget> GetToolSelector();
 	bool GetToolSelectorIsVisible() const;
 	EVisibility GetToolSelectorVisibility() const;
 
 	FName GetCurrentBrushFName() const;
 	FText GetCurrentBrushName() const;
 	FSlateIcon GetCurrentBrushIcon() const;
-	TSharedRef<SWidget> GetBrushSelector();
 	bool GetBrushSelectorIsVisible() const;
 	EVisibility GetBrushSelectorVisibility() const;
 
 	FName GetCurrentBrushFalloffFName() const;
 	FText GetCurrentBrushFalloffName() const;
 	FSlateIcon GetCurrentBrushFalloffIcon() const;
-	TSharedRef<SWidget> GetBrushFalloffSelector();
 	bool GetBrushFalloffSelectorIsVisible() const;
 	EVisibility GetBrushFalloffSelectorVisibility() const;
 
@@ -76,8 +73,6 @@ protected:
 
 	bool IsBrushSetEnabled() const;
 
-	TSharedPtr<FUICommandList> CommandList;
-
 	TSharedPtr<FLandscapeEditorDetailCustomization_NewLandscape> Customization_NewLandscape;
 	TSharedPtr<FLandscapeEditorDetailCustomization_ResizeLandscape> Customization_ResizeLandscape;
 	TSharedPtr<FLandscapeEditorDetailCustomization_ImportExport> Customization_ImportExport;
