Fix incorrect Landscape tooltips

diff --git a/Engine/Source/Editor/LandscapeEditor/Private/LandscapeEditorCommands.cpp b/Engine/Source/Editor/LandscapeEditor/Private/LandscapeEditorCommands.cpp
--- a/Engine/Source/Editor/LandscapeEditor/Private/LandscapeEditorCommands.cpp
+++ b/Engine/Source/Editor/LandscapeEditor/Private/LandscapeEditorCommands.cpp
@@ -46,13 +46,13 @@ void FLandscapeEditorCommands::RegisterCommands()
 	UI_COMMAND(ImportExportTool, "Import", "Import or Export landscape data", EUserInterfaceActionType::RadioButton, FInputChord());
 	NameToCommandMap.Add("Tool_ImportExport", ImportExportTool);
 
-	UI_COMMAND(SculptTool, "Sculpt", "Sculpt height data.\n\nCtrl+Click to Raise, Ctrl+Shift_Click to lower.", EUserInterfaceActionType::RadioButton, FInputChord());
+	UI_COMMAND(SculptTool, "Sculpt", "Sculpt height data.\n\nClick to raise, Shift+Click to lower.", EUserInterfaceActionType::RadioButton, FInputChord());
 	NameToCommandMap.Add("Tool_Sculpt", SculptTool);
 
 	UI_COMMAND(EraseTool, "Erase", "Erase height data.", EUserInterfaceActionType::RadioButton, FInputChord());
 	NameToCommandMap.Add("Tool_Erase", EraseTool);
 
-	UI_COMMAND(PaintTool, "Paint", "Paint weight data.\n\nCtrl+Click to paint, Ctrl+Shift+Click to erase.", EUserInterfaceActionType::RadioButton, FInputChord());
+	UI_COMMAND(PaintTool, "Paint", "Paint weight data.\n\nClick to paint, Shift+Click to erase.", EUserInterfaceActionType::RadioButton, FInputChord());
 	NameToCommandMap.Add("Tool_Paint", PaintTool);
 
 	UI_COMMAND(SmoothTool, "Smooth", "Smooths heightmaps or blend layers.", EUserInterfaceActionType::RadioButton, FInputChord());
