Landscape: Create Large landscapes using World Partition

diff --git a/Engine/Plugins/Experimental/PCG/Source/PCG/Private/Grid/PCGLandscapeCache.cpp b/Engine/Plugins/Experimental/PCG/Source/PCG/Private/Grid/PCGLandscapeCache.cpp
--- a/Engine/Plugins/Experimental/PCG/Source/PCG/Private/Grid/PCGLandscapeCache.cpp
+++ b/Engine/Plugins/Experimental/PCG/Source/PCG/Private/Grid/PCGLandscapeCache.cpp
@@ -566,7 +566,7 @@ void UPCGLandscapeCache::PrimeCache()
 		if (IsValid(LandscapeInfo))
 		{
 			// Build per-component information
-			LandscapeInfo->ForAllLandscapeProxies([this, LandscapeInfo, &CacheEntriesToBuild](const ALandscapeProxy* LandscapeProxy)
+			LandscapeInfo->ForEachLandscapeProxy([this, LandscapeInfo, &CacheEntriesToBuild](const ALandscapeProxy* LandscapeProxy)
 			{
 				check(LandscapeProxy);
 				const FGuid LandscapeGuid = LandscapeProxy->GetLandscapeGuid();
@@ -586,6 +586,7 @@ void UPCGLandscapeCache::PrimeCache()
 						CacheEntriesToBuild.Emplace(LandscapeInfo, LandscapeComponent, ComponentKey);
 					}
 				}
+				return true;
 			});
 		}
 	}
diff --git a/Engine/Source/Editor/LandscapeEditor/LandscapeEditor.Build.cs b/Engine/Source/Editor/LandscapeEditor/LandscapeEditor.Build.cs
--- a/Engine/Source/Editor/LandscapeEditor/LandscapeEditor.Build.cs
+++ b/Engine/Source/Editor/LandscapeEditor/LandscapeEditor.Build.cs
@@ -10,6 +10,7 @@ public class LandscapeEditor : ModuleRules
 			new string[] {
 				"Core",
 				"CoreUObject",
+				"Json",
 				"ApplicationCore",
 				"Slate",
 				"SlateCore",
@@ -28,6 +29,9 @@ public class LandscapeEditor : ModuleRules
                 "Foliage",
 				"ToolMenus",
 				"ToolWidgets",
+				"SourceControl",
+				"DirectoryWatcher",
+				"DeveloperSettings"
 			}
 			);
 
diff --git a/Engine/Source/Editor/LandscapeEditor/Private/LandscapeEdMode.cpp b/Engine/Source/Editor/LandscapeEditor/Private/LandscapeEdMode.cpp
--- a/Engine/Source/Editor/LandscapeEditor/Private/LandscapeEdMode.cpp
+++ b/Engine/Source/Editor/LandscapeEditor/Private/LandscapeEdMode.cpp
@@ -20,6 +20,7 @@
 #include "LandscapeStreamingProxy.h"
 #include "LandscapeSubsystem.h"
 #include "LandscapeSettings.h"
+#include "LandscapeTiledImage.h"
 
 #include "EditorSupportDelegates.h"
 #include "ScopedTransaction.h"
@@ -3161,102 +3162,77 @@ void FEdModeLandscape::ReimportData(const FLandscapeTargetListInfo& TargetInfo)
 template<class T>
 void ImportDataInternal(ULandscapeInfo* LandscapeInfo, const FString& Filename, FName LayerName, bool bSingleFile, bool bFlipYAxis, const FIntRect& ImportRegionVerts, ELandscapeImportTransformType TransformType, FIntPoint Offset, TFunctionRef<void(int32, int32, int32, int32, const TArray<T>&)> SetDataFunc)
 {
-	if (LandscapeInfo)
+	if (!LandscapeInfo)
 	{
-		const FLandscapeImportResolution LandscapeResolution = { (uint32)(ImportRegionVerts.Width()), (uint32)(ImportRegionVerts.Height()) };
-		FLandscapeImportDescriptor OutImportDescriptor;
-		FText OutMessage;
-
-		ELandscapeImportResult ImportResult = FLandscapeImportHelper::GetImportDescriptor<T>(Filename, bSingleFile, bFlipYAxis, LayerName, OutImportDescriptor, OutMessage);
-		if (ImportResult == ELandscapeImportResult::Error)
-		{
-			FMessageDialog::Open(EAppMsgType::Ok, OutMessage);
-			return;
-		}
+		return;
+	}
 
-		check(OutImportDescriptor.ImportResolutions.Num() > 0);
+	FLandscapeTiledImage TiledImage;
+	TiledImage.Load(*Filename);
+	FIntPoint ImportResolution = TiledImage.GetResolution();
 
-		int32 DescriptorIndex = 0;
-		// If there is more than one possible res it needs to match perfectly (this constraint could probably removed)
-		if (OutImportDescriptor.ImportResolutions.Num() > 1)
-		{
-			DescriptorIndex = OutImportDescriptor.FindDescriptorIndex(LandscapeResolution.Width, LandscapeResolution.Height);
-			// if the file is a raw format with multiple possibly resolutions, only attempt import if one matches the current landscape
-			if (DescriptorIndex == INDEX_NONE)
-			{
-				FFormatNamedArguments Args;
-				Args.Add(TEXT("LandscapeSizeX"), LandscapeResolution.Width);
-				Args.Add(TEXT("LandscapeSizeY"), LandscapeResolution.Height);
+	bool bResolutionMismatch = false;
 
-				FMessageDialog::Open(EAppMsgType::Ok,
-					FText::Format(NSLOCTEXT("LandscapeEditor.Import", "Import_SizeMismatchRaw", "The file resolution does not match the current Landscape extent ({LandscapeSizeX}\u00D7{LandscapeSizeY}), and its exact resolution could not be determined"), Args));
+	if ((ImportResolution.X != ImportRegionVerts.Width() || ImportResolution.Y != ImportRegionVerts.Height()) && TransformType != ELandscapeImportTransformType::Subregion)
+	{
+		bResolutionMismatch = true;
 
-				return;
-			}
-		}
-				
-		// display warning message if there is one and allow user to cancel
-		if (ImportResult == ELandscapeImportResult::Warning)
-		{
-			auto Result = FMessageDialog::Open(EAppMsgType::OkCancel, OutMessage);
+		FFormatNamedArguments Args;
+		Args.Add(TEXT("ImportSizeX"), ImportResolution.X);
+		Args.Add(TEXT("ImportSizeY"), ImportResolution.Y);
+		Args.Add(TEXT("LandscapeSizeX"), ImportRegionVerts.Width());
+		Args.Add(TEXT("LandscapeSizeY"), ImportRegionVerts.Height());
 
-			if (Result != EAppReturnType::Ok)
-			{
-				return;
-			}
-		}
+		auto Result = FMessageDialog::Open(EAppMsgType::OkCancel,
+			FText::Format(NSLOCTEXT("LandscapeEditor.Import", "Import_SizeMismatch", "The import size ({ImportSizeX}\u00D7{ImportSizeY}) does not match the current Landscape extent ({LandscapeSizeX}\u00D7{LandscapeSizeY}), if you continue it will be padded/clipped to fit"), Args));
 
-		const FLandscapeImportResolution ImportResolution = OutImportDescriptor.ImportResolutions[DescriptorIndex];
-		bool bResolutionMismatch = false;
-		// if the file is a format with resolution information, warn the user if the resolution doesn't match the current landscape
-		// unlike for raw this is only a warning as we can pad/clip the data if we know what resolution it is
-		if (ImportResolution != LandscapeResolution)
+		if (Result != EAppReturnType::Ok)
 		{
-			bResolutionMismatch = true;
-
-			FFormatNamedArguments Args;
-			Args.Add(TEXT("ImportSizeX"), ImportResolution.Width);
-			Args.Add(TEXT("ImportSizeY"), ImportResolution.Height);
-			Args.Add(TEXT("LandscapeSizeX"), LandscapeResolution.Width);
-			Args.Add(TEXT("LandscapeSizeY"), LandscapeResolution.Height);
+			return;
+		}
+	}
 
-			auto Result = FMessageDialog::Open(EAppMsgType::OkCancel,
-				FText::Format(NSLOCTEXT("LandscapeEditor.Import", "Import_SizeMismatch", "The import size ({ImportSizeX}\u00D7{ImportSizeY}) does not match the current Landscape extent ({LandscapeSizeX}\u00D7{LandscapeSizeY}), if you continue it will be padded/clipped to fit"), Args));
+	TArray<T> ImportData;
+	if (TransformType == ELandscapeImportTransformType::Subregion)
+	{
+		TiledImage.ReadRegion<T>(ImportRegionVerts, ImportData, bFlipYAxis);
+	}
+	else if (bResolutionMismatch)
+	{
+		TiledImage.Read<T>(ImportData, bFlipYAxis);
+	}
+	else
+	{
+		const FIntRect RegionToLoad(0, 0, ImportRegionVerts.Width(), ImportRegionVerts.Height());
+		TiledImage.ReadRegion<T>(RegionToLoad, ImportData, bFlipYAxis);
+	}
 
-			if (Result != EAppReturnType::Ok)
-			{
-				return;
-			}
+	// Expand if necessary...
+	{
+		TArray<T> FinalData;
+		if (bResolutionMismatch)
+		{
+			FLandscapeImportHelper::TransformImportData<T>(ImportData, FinalData, FLandscapeImportResolution(ImportResolution.X, ImportResolution.Y), FLandscapeImportResolution(ImportRegionVerts.Width(), ImportRegionVerts.Height()), TransformType, Offset);
 		}
-
-		TArray<T> ImportData;
-		ImportResult = FLandscapeImportHelper::GetImportData<T>(OutImportDescriptor, DescriptorIndex, LayerName, ImportData, OutMessage);
-		if (ImportResult == ELandscapeImportResult::Error)
+		else if (TransformType == ELandscapeImportTransformType::Subregion)
 		{
-			FMessageDialog::Open(EAppMsgType::Ok, OutMessage);
-			return;
+			FinalData = MoveTemp(ImportData);
 		}
-
-		// Expand if necessary...
+		else
 		{
-			TArray<T> FinalData;
-			if (bResolutionMismatch)
-			{
-				FLandscapeImportHelper::TransformImportData<T>(ImportData, FinalData, ImportResolution, LandscapeResolution, TransformType, Offset);
-			}
-			else
-			{
-				FinalData = MoveTemp(ImportData);
-			}
-									
-			// Set Data is in Quads (so remove 1)
-			SetDataFunc(ImportRegionVerts.Min.X, ImportRegionVerts.Min.Y, ImportRegionVerts.Max.X-1, ImportRegionVerts.Max.Y-1, FinalData);
+			FinalData = MoveTemp(ImportData);
 		}
-	}			
+
+		// Set Data is in Quads (so remove 1)
+		SetDataFunc(ImportRegionVerts.Min.X, ImportRegionVerts.Min.Y, ImportRegionVerts.Max.X - 1, ImportRegionVerts.Max.Y - 1, FinalData);
+	}
+	
 }
 
 void FEdModeLandscape::ImportHeightData(ULandscapeInfo* LandscapeInfo, const FGuid& LayerGuid, const FString& Filename, const FIntRect& ImportRegionVerts, ELandscapeImportTransformType TransformType, FIntPoint Offset, const ELandscapeLayerPaintingRestriction& PaintRestriction, bool bFlipYAxis)
 {
+	TRACE_CPUPROFILER_EVENT_SCOPE(FEdModeLandscape::ImportHeightData);
+
 	ImportDataInternal<uint16>(LandscapeInfo, Filename, NAME_None, UseSingleFileImport(), bFlipYAxis, ImportRegionVerts, TransformType, Offset, [LandscapeInfo, LayerGuid, PaintRestriction](int32 MinX, int32 MinY, int32 MaxX, int32 MaxY, const TArray<uint16>& Data)
 	{
 		ALandscape* Landscape = LandscapeInfo->LandscapeActor.Get();
@@ -4303,13 +4279,8 @@ bool FEdModeLandscape::NeedToFillEmptyMaterialLayers() const
 
 	bool bCanFill = true;
 
-	CurrentToolTarget.LandscapeInfo->ForAllLandscapeProxies([&](ALandscapeProxy* Proxy)
+	CurrentToolTarget.LandscapeInfo->ForEachLandscapeProxy([&](ALandscapeProxy* Proxy)
 	{
-		if (!bCanFill)
-		{
-			return;
-		}
-
 		ALandscape* Landscape = Proxy->GetLandscapeActor();
 
 		if (Landscape != nullptr)
@@ -4327,13 +4298,14 @@ bool FEdModeLandscape::NeedToFillEmptyMaterialLayers() const
 							if (Alloc.LayerInfo != nullptr)
 							{
 								bCanFill = false;
-								return;
+								return false;
 							}
 						}
 					}
 				}
 			}
 		}
+		return true;
 	});	
 
 	return bCanFill;
diff --git a/Engine/Source/Editor/LandscapeEditor/Private/LandscapeEdModeComponentTools.cpp b/Engine/Source/Editor/LandscapeEditor/Private/LandscapeEdModeComponentTools.cpp
--- a/Engine/Source/Editor/LandscapeEditor/Private/LandscapeEdModeComponentTools.cpp
+++ b/Engine/Source/Editor/LandscapeEditor/Private/LandscapeEdModeComponentTools.cpp
@@ -536,6 +536,105 @@ public:
 //
 // FLandscapeToolAddComponent
 //
+
+namespace
+{
+	void AddComponents(ULandscapeInfo* LandscapeInfo, ULandscapeSubsystem* LandscapeSubsystem, const TArray<FIntPoint>& ComponentCoordinates)
+	{
+		TArray<ULandscapeComponent*> NewComponents;
+		LandscapeInfo->Modify();
+		for (const auto& ComponentCoordinate : ComponentCoordinates)
+		{
+			ULandscapeComponent* LandscapeComponent = LandscapeInfo->XYtoComponentMap.FindRef(ComponentCoordinate);
+			if (LandscapeComponent)
+			{
+				continue;
+			}
+
+			// Add New component...
+			FIntPoint ComponentBase = ComponentCoordinate * LandscapeInfo->ComponentSizeQuads;
+
+			ALandscapeProxy* LandscapeProxy = LandscapeSubsystem->FindOrAddLandscapeProxy(LandscapeInfo, ComponentBase);
+			if (!LandscapeProxy)
+			{
+				continue;
+			}
+
+			LandscapeComponent = NewObject<ULandscapeComponent>(LandscapeProxy, NAME_None, RF_Transactional);
+			NewComponents.Add(LandscapeComponent);
+			LandscapeComponent->Init(
+				ComponentBase.X, ComponentBase.Y,
+				LandscapeProxy->ComponentSizeQuads,
+				LandscapeProxy->NumSubsections,
+				LandscapeProxy->SubsectionSizeQuads
+			);
+
+			TArray<FColor> HeightData;
+			const int32 ComponentVerts = (LandscapeComponent->SubsectionSizeQuads + 1) * LandscapeComponent->NumSubsections;
+			HeightData.Empty(FMath::Square(ComponentVerts));
+			HeightData.AddZeroed(FMath::Square(ComponentVerts));
+			LandscapeComponent->InitHeightmapData(HeightData, true);
+			LandscapeComponent->UpdateMaterialInstances();
+
+			LandscapeInfo->XYtoComponentMap.Add(ComponentCoordinate, LandscapeComponent);
+			LandscapeInfo->XYtoAddCollisionMap.Remove(ComponentCoordinate);
+		}
+
+		// Need to register to use general height/xyoffset data update
+		for (int32 Idx = 0; Idx < NewComponents.Num(); Idx++)
+		{
+			NewComponents[Idx]->RegisterComponent();
+		}
+
+		const bool bHasXYOffset = false;
+		ALandscape* Landscape = LandscapeInfo->LandscapeActor.Get();
+
+		bool bHasLandscapeLayersContent = Landscape && Landscape->HasLayersContent();
+
+		if (bHasLandscapeLayersContent)
+		{
+			Landscape->RequestLayersInitialization();
+		}
+
+		for (ULandscapeComponent* NewComponent : NewComponents)
+		{
+			if (bHasLandscapeLayersContent)
+			{
+				TArray<ULandscapeComponent*> ComponentsUsingHeightmap;
+				ComponentsUsingHeightmap.Add(NewComponent);
+
+				for (const FLandscapeLayer& Layer : Landscape->LandscapeLayers)
+				{
+					// Since we do not share heightmap when adding new component, we will provided the required array, but they will only be used for 1 component
+					TMap<UTexture2D*, UTexture2D*> CreatedHeightmapTextures;
+					NewComponent->AddDefaultLayerData(Layer.Guid, ComponentsUsingHeightmap, CreatedHeightmapTextures);
+				}
+			}
+
+			// Update Collision
+			NewComponent->UpdateCachedBounds();
+			NewComponent->UpdateBounds();
+			NewComponent->MarkRenderStateDirty();
+
+			if (!bHasLandscapeLayersContent)
+			{
+				ULandscapeHeightfieldCollisionComponent* CollisionComp = NewComponent->GetCollisionComponent();
+				if (CollisionComp && !bHasXYOffset)
+				{
+					CollisionComp->MarkRenderStateDirty();
+					CollisionComp->RecreateCollision();
+				}
+			}
+		}
+
+		if (Landscape)
+		{
+			GEngine->BroadcastOnActorMoved(Landscape);
+		}
+
+	}
+}
+
 class FLandscapeToolStrokeAddComponent : public FLandscapeToolStrokeBase
 {
 	using Super = FLandscapeToolStrokeBase;
diff --git a/Engine/Source/Editor/LandscapeEditor/Private/LandscapeEditorDetailCustomization_ImportExport.cpp b/Engine/Source/Editor/LandscapeEditor/Private/LandscapeEditorDetailCustomization_ImportExport.cpp
--- a/Engine/Source/Editor/LandscapeEditor/Private/LandscapeEditorDetailCustomization_ImportExport.cpp
+++ b/Engine/Source/Editor/LandscapeEditor/Private/LandscapeEditorDetailCustomization_ImportExport.cpp
@@ -15,6 +15,16 @@
 #include "Widgets/Text/STextBlock.h"
 #include "Widgets/Notifications/SErrorText.h"
 #include "Widgets/SToolTip.h"
+
+#include "Misc/ScopedSlowTask.h"
+#include "Misc/MessageDialog.h"
+
+#include "LocationVolume.h"
+#include "SourceControlHelpers.h"
+#include "UObject/SavePackage.h"
+#include "WorldPartition/WorldPartition.h"
+#include "Internationalization/Regex.h"
+
 #include "Framework/Application/SlateApplication.h"
 #include "Framework/MultiBox/MultiBoxBuilder.h"
 
@@ -24,6 +34,12 @@
 #include "LandscapeEditorObject.h"
 #include "LandscapeEditorModule.h"
 #include "LandscapeImportHelper.h"
+#include "LandscapeSubsystem.h"
+#include "LandscapeStreamingProxy.h"
+#include "LandscapeImageFileCache.h"
+#include "LandscapeTiledImage.h"
+#include "LandscapeRegionUtils.h"
+#include "LandscapeEditorUtils.h"
 
 #define LOCTEXT_NAMESPACE "LandscapeEditor.ImportExport"
 
@@ -439,42 +455,18 @@ FReply FLandscapeEditorDetailCustomization_ImportExport::OnBrowseFilenameButtonC
 {
 	FEdModeLandscape* LandscapeEdMode = GetEditorMode();
 	check(LandscapeEdMode != nullptr);
+	ILandscapeEditorModule& LandscapeEditorModule = FModuleManager::GetModuleChecked<ILandscapeEditorModule>("LandscapeEditor");
+	const bool bIsImporting = IsImporting();
+	const FString DialogType = bIsImporting ? LandscapeEditorModule.GetHeightmapImportDialogTypeString() : LandscapeEditorModule.GetHeightmapExportDialogTypeString();
+	const FString DialogTitle = bIsImporting ? LOCTEXT("ImportHeightmap", "Import Heightmap").ToString() : LOCTEXT("ExportHeightmap", "Export Heightmap").ToString();
 
-	// Prompt the user for the Filenames
-	IDesktopPlatform* DesktopPlatform = FDesktopPlatformModule::Get();
-	if (DesktopPlatform != nullptr)
+	TOptional<FString> OptionalFilename = LandscapeEditorUtils::GetImportExportFilename(DialogTitle, LandscapeEdMode->UISettings->LastImportPath, DialogType);
+	if (OptionalFilename.IsSet())
 	{
-		bool bSuccess = false;
-		TArray<FString> Filenames;
-		ILandscapeEditorModule& LandscapeEditorModule = FModuleManager::GetModuleChecked<ILandscapeEditorModule>("LandscapeEditor");
-		if (IsImporting())
-		{
-			bSuccess = DesktopPlatform->OpenFileDialog(
-				FSlateApplication::Get().FindBestParentWindowHandleForDialogs(nullptr),
-				LOCTEXT("ImportHeightmap", "Import Heightmap").ToString(),
-				LandscapeEdMode->UISettings->LastImportPath,
-				TEXT(""),
-				LandscapeEditorModule.GetHeightmapImportDialogTypeString(),
-				EFileDialogFlags::None,
-				Filenames);
-		}
-		else
-		{
-			bSuccess = DesktopPlatform->SaveFileDialog(
-				FSlateApplication::Get().FindBestParentWindowHandleForDialogs(nullptr),
-				LOCTEXT("ExportHeightmap", "Export Heightmap").ToString(),
-				LandscapeEdMode->UISettings->LastImportPath,
-				TEXT(""),
-				LandscapeEditorModule.GetHeightmapExportDialogTypeString(),
-				EFileDialogFlags::None,
-				Filenames);
-		}
+		const FString& Filename = OptionalFilename.GetValue();
 
-		if (bSuccess)
-		{
-			ensure(PropertyHandle_Filename->SetValue(Filenames[0]) == FPropertyAccess::Success);
-			LandscapeEdMode->UISettings->LastImportPath = FPaths::GetPath(Filenames[0]);
-		}
+		ensure(PropertyHandle_Filename->SetValue(Filename) == FPropertyAccess::Success);
+		LandscapeEdMode->UISettings->LastImportPath = FPaths::GetPath(Filename);
 	}
 
 	return FReply::Handled();
@@ -552,13 +544,13 @@ FReply FLandscapeEditorDetailCustomization_ImportExport::OnImportExportButtonCli
 {
 	FEdModeLandscape* LandscapeEdMode = GetEditorMode();
 	check(LandscapeEdMode);
+
+	UWorld* World = LandscapeEdMode->GetWorld();
+
 	if (ULandscapeInfo* LandscapeInfo = LandscapeEdMode->CurrentToolTarget.LandscapeInfo.Get())
 	{
 		FIntRect LandscapeExtent;
-		if (!LandscapeInfo->GetLandscapeExtent(LandscapeExtent))
-		{
-			return FReply::Handled();
-		}
+		LandscapeInfo->GetLandscapeExtent(LandscapeExtent);
 	
 		if (IsImporting())
 		{
@@ -585,20 +577,87 @@ FReply FLandscapeEditorDetailCustomization_ImportExport::OnImportExportButtonCli
 				ImportRegion = FIntRect(LocalGizmoPoint.X, LocalGizmoPoint.Y, LocalGizmoPoint.X + LandscapeEdMode->UISettings->ImportLandscape_Width, LocalGizmoPoint.Y + LandscapeEdMode->UISettings->ImportLandscape_Height);
 			}
 			
-			// Import Heightmap
-			if (LandscapeEdMode->UISettings->bHeightmapSelected)
+			check(LandscapeEdMode->UISettings->ImportLandscape_HeightmapImportResult != ELandscapeImportResult::Error);
+			const bool bIsWorldPartition = LandscapeEdMode->GetWorld()->GetSubsystem<ULandscapeSubsystem>()->IsGridBased();
+
+			if (bIsWorldPartition && LandscapeEdMode->UISettings->ImportExportMode != ELandscapeImportExportMode::LoadedOnly) // If in world partition import region by region
 			{
-				check(LandscapeEdMode->UISettings->ImportLandscape_HeightmapImportResult != ELandscapeImportResult::Error);
-				LandscapeEdMode->ImportHeightData(LandscapeInfo, CurrentLayerGuid, LandscapeEdMode->UISettings->ImportLandscape_HeightmapFilename, ImportRegion, TransformType, ImportOffset, PaintRestriction, LandscapeEdMode->UISettings->bFlipYAxis);
-			}
+				TArray<ALocationVolume*> LandscapeRegions;
+				TArray<AActor*> Children;
+				LandscapeInfo->LandscapeActor->GetAttachedActors(Children);
+				for (AActor* Child : Children)
+				{
+					if (Child->IsA<ALocationVolume>())
+					{
+						LandscapeRegions.Add(Cast<ALocationVolume>(Child));
+					}
+				}
+
+				int32 NumRegions = LandscapeRegions.Num();
+					
+				FScopedSlowTask Progress(NumRegions, LOCTEXT("ImportingLandscapeRegions", "Importing Landscape Regions..."));
+				Progress.MakeDialog(/*bShowCancelButton = */ true);
 
-			for (const FLandscapeImportLayer& ImportLayer : LandscapeEdMode->UISettings->ImportLandscape_Layers)
+				auto RegionImporter = [&Progress, LandscapeInfo, CurrentLayerGuid, LandscapeEdMode, ImportRegion, TransformType, ImportOffset, PaintRestriction](const FBox& RegionBounds, const TArray<ALandscapeProxy*>& Proxies)
+				{
+					// todo need a better way of working out the landscape bounds of the StreamingProxies in this region.
+					// What if some of the landscape is loaded?
+					FIntRect LandscapeLoadedExtent;
+					LandscapeInfo->GetLandscapeExtent(LandscapeLoadedExtent);
+					LandscapeLoadedExtent.Max.X += 1;
+					LandscapeLoadedExtent.Max.Y += 1;
+
+					Progress.EnterProgressFrame(1.0f, LOCTEXT("ImportingLandscapeRegions", "Importing Landscape Regions"));
+
+					if (LandscapeEdMode->UISettings->bHeightmapSelected)
+					{
+						LandscapeEdMode->ImportHeightData(LandscapeInfo, CurrentLayerGuid, LandscapeEdMode->UISettings->ImportLandscape_HeightmapFilename, LandscapeLoadedExtent, TransformType, ImportOffset, PaintRestriction, LandscapeEdMode->UISettings->bFlipYAxis);
+					}
+
+					for (const FLandscapeImportLayer& ImportLayer : LandscapeEdMode->UISettings->ImportLandscape_Layers)
+					{
+						if (ImportLayer.bSelected)
+						{
+							check(ImportLayer.ImportResult != ELandscapeImportResult::Error);
+							LandscapeEdMode->ImportWeightData(LandscapeInfo, CurrentLayerGuid, ImportLayer.LayerInfo, ImportLayer.SourceFilePath, LandscapeLoadedExtent, TransformType, ImportOffset, PaintRestriction, LandscapeEdMode->UISettings->bFlipYAxis);
+						}
+					}
+
+					return !Progress.ShouldCancel();
+				};
+					
+				LandscapeRegionUtils::ForEachRegion_LoadProcessUnload(LandscapeInfo, ImportRegion, World, RegionImporter);
+
+				FLandscapeImageFileCache& LandscapeImageFileCache = FModuleManager::GetModuleChecked<ILandscapeEditorModule>("LandscapeEditor").GetImageFileCache();
+				LandscapeImageFileCache.Clear();
+			}
+			else
 			{
-				if (ImportLayer.bSelected)
+				FScopedSlowTask Progress(1 + LandscapeEdMode->UISettings->ImportLandscape_Layers.Num(), LOCTEXT("ImportingLandscape", "Importing Landscape"));
+				Progress.MakeDialog(/*bShowCancelButton = */ true);
+
+				if (LandscapeEdMode->UISettings->bHeightmapSelected)
+				{
+					Progress.EnterProgressFrame(1.0f, LOCTEXT("ImportingLandscapeHeight", "Importing Landscape Height"));
+					LandscapeEdMode->ImportHeightData(LandscapeInfo, CurrentLayerGuid, LandscapeEdMode->UISettings->ImportLandscape_HeightmapFilename, ImportRegion, TransformType, ImportOffset, PaintRestriction, LandscapeEdMode->UISettings->bFlipYAxis);
+				}
+
+				for (const FLandscapeImportLayer& ImportLayer : LandscapeEdMode->UISettings->ImportLandscape_Layers)
 				{
-					check(ImportLayer.ImportResult != ELandscapeImportResult::Error);
-					LandscapeEdMode->ImportWeightData(LandscapeInfo, CurrentLayerGuid, ImportLayer.LayerInfo, ImportLayer.SourceFilePath, ImportRegion, TransformType, ImportOffset, PaintRestriction, LandscapeEdMode->UISettings->bFlipYAxis);
+					if (Progress.ShouldCancel())
+					{
+						continue;
+					}
+
+					if (ImportLayer.bSelected)
+					{
+						Progress.EnterProgressFrame(1.0f, LOCTEXT("ImportingLandscapeWeight", "Importing Landscape Weight"));
+						check(ImportLayer.ImportResult != ELandscapeImportResult::Error);
+						LandscapeEdMode->ImportWeightData(LandscapeInfo, CurrentLayerGuid, ImportLayer.LayerInfo, ImportLayer.SourceFilePath, ImportRegion, TransformType, ImportOffset, PaintRestriction, LandscapeEdMode->UISettings->bFlipYAxis);
+					}
 				}
+
+				LandscapeInfo->ForceLayersFullUpdate();
 			}
 		}
 		else
@@ -648,8 +707,39 @@ FReply FLandscapeEditorDetailCustomization_ImportExport::OnImportExportButtonCli
 			}
 			else
 			{
+				FScopedSlowTask Progress(0, LOCTEXT("ExportingLandscapeRegions", "Exporting Landscape Regions..."));
+				Progress.MakeDialog(/*bShowCancelButton = */ true);
+
+				auto Exporter = [&Progress, LandscapeInfo, Landscape, LandscapeEdMode, LandscapeExtent, PerformExport, World](const FBox& RegionBounds, const TArray<ALandscapeProxy*>& Proxies)
+				{
+					Progress.EnterProgressFrame(0.0f, LOCTEXT("ExportingLandscapeRegions", "Exporting Landscape Regions"));
+					for (ALandscapeProxy* LandscapeProxy : Proxies)
+					{
+						FIntRect ExportRegion;
+
+						if (LandscapeInfo->GetLandscapeExtent(LandscapeProxy, ExportRegion))
+						{
+							const int32 YCoord = LandscapeEdMode->UISettings->bFlipYAxis ? LandscapeExtent.Max.Y - ExportRegion.Max.Y : ExportRegion.Min.Y - LandscapeExtent.Min.Y;
+							FIntPoint FileOffset = FIntPoint((ExportRegion.Min.X - LandscapeExtent.Min.X) / Landscape->GridSize,
+								YCoord / Landscape->GridSize);
+
+							// Remove the shared line/column that this proxy has with its neighbors because it 
+							// will be included by the neighbor or lost if there is none (that could become an option to avoid that loss)
+							ExportRegion.Max.X -= 1;
+							ExportRegion.Max.Y -= 1;
+
+							PerformExport(LandscapeEdMode->UISettings, LandscapeInfo, ExportRegion, FileOffset, /*bInUseOffset = */true);
+						}
+					}
+					return !Progress.ShouldCancel();
+					
+				};
+
+				LandscapeRegionUtils::ForEachRegion_LoadProcessUnload(LandscapeInfo, LandscapeExtent, World, Exporter);
+			
+
 				// For multiple file export, export each landscape proxy individually :
-				LandscapeInfo->ForAllLandscapeProxies([LandscapeInfo, Landscape, LandscapeEdMode, LandscapeExtent, BuildExportFileName, PerformExport](ALandscapeProxy* LandscapeProxy)
+				LandscapeInfo->ForEachLandscapeProxy([LandscapeInfo, Landscape, LandscapeEdMode, LandscapeExtent, BuildExportFileName, PerformExport](ALandscapeProxy* LandscapeProxy)
 				{
 					FIntRect ExportRegion;
 					if (LandscapeInfo->GetLandscapeExtent(LandscapeProxy, ExportRegion))
@@ -665,6 +755,7 @@ FReply FLandscapeEditorDetailCustomization_ImportExport::OnImportExportButtonCli
 
 						PerformExport(LandscapeEdMode->UISettings, LandscapeInfo, ExportRegion, FileOffset, /*bInUseOffset = */true);
 					}
+					return true;
 				});
 			}
 
@@ -804,4 +895,6 @@ FText FLandscapeEditorDetailCustomization_ImportExport::GetImportLandscapeResolu
 	return FText::GetEmpty();
 }
 
+
+
 #undef LOCTEXT_NAMESPACE
diff --git a/Engine/Source/Editor/LandscapeEditor/Private/LandscapeEditorDetailCustomization_ImportLayers.cpp b/Engine/Source/Editor/LandscapeEditor/Private/LandscapeEditorDetailCustomization_ImportLayers.cpp
--- a/Engine/Source/Editor/LandscapeEditor/Private/LandscapeEditorDetailCustomization_ImportLayers.cpp
+++ b/Engine/Source/Editor/LandscapeEditor/Private/LandscapeEditorDetailCustomization_ImportLayers.cpp
@@ -8,6 +8,8 @@
 #include "LandscapeEditorModule.h"
 #include "LandscapeEditorObject.h"
 #include "LandscapeUtils.h"
+#include "LandscapeTiledImage.h"
+#include "LandscapeEditorUtils.h"
 #include "SLandscapeEditor.h"
 
 #include "IDetailChildrenBuilder.h"
@@ -24,6 +26,7 @@
 #include "Widgets/SToolTip.h"
 #include "Framework/Application/SlateApplication.h"
 #include "Framework/MultiBox/MultiBoxBuilder.h"
+#include "Misc/MessageDialog.h"
 
 #define LOCTEXT_NAMESPACE "LandscapeEditor.ImportLayers"
 
@@ -228,41 +231,18 @@ FReply FLandscapeEditorStructCustomization_FLandscapeImportLayer::OnLayerFilenam
 	FEdModeLandscape* LandscapeEdMode = GetEditorMode();
 	check(LandscapeEdMode != nullptr);
 
-	// Prompt the user for the Filenames
-	IDesktopPlatform* DesktopPlatform = FDesktopPlatformModule::Get();
-	if (DesktopPlatform != nullptr)
+	ILandscapeEditorModule& LandscapeEditorModule = FModuleManager::GetModuleChecked<ILandscapeEditorModule>("LandscapeEditor");
+	const bool bIsImporting = IsImporting();
+	const FString DialogTypeString = bIsImporting  ? LandscapeEditorModule.GetWeightmapImportDialogTypeString() : LandscapeEditorModule.GetWeightmapExportDialogTypeString();
+	const FString DialogTitle = bIsImporting ? LOCTEXT("ImportLayer", "Import Layer").ToString() : LOCTEXT("ExportLayer", "Export Layer").ToString();
+
+	TOptional<FString> OptionalExportImportPath = LandscapeEditorUtils::GetImportExportFilename(DialogTitle, LandscapeEdMode->UISettings->LastImportPath, DialogTypeString);
+	if (OptionalExportImportPath.IsSet())
 	{
-		bool bSuccess = false;
-		TArray<FString> Filenames;
-		ILandscapeEditorModule& LandscapeEditorModule = FModuleManager::GetModuleChecked<ILandscapeEditorModule>("LandscapeEditor");
-		if (IsImporting())
-		{
-			bSuccess = DesktopPlatform->OpenFileDialog(
-				FSlateApplication::Get().FindBestParentWindowHandleForDialogs(nullptr),
-				LOCTEXT("ImportLayer", "Import Layer").ToString(),
-				LandscapeEdMode->UISettings->LastImportPath,
-				TEXT(""),
-				LandscapeEditorModule.GetWeightmapImportDialogTypeString(),
-				EFileDialogFlags::None,
-				Filenames);
-		}
-		else
-		{
-			bSuccess = DesktopPlatform->SaveFileDialog(
-				FSlateApplication::Get().FindBestParentWindowHandleForDialogs(nullptr),
-				LOCTEXT("ExportLayer", "Export Layer").ToString(),
-				LandscapeEdMode->UISettings->LastImportPath,
-				TEXT(""),
-				LandscapeEditorModule.GetWeightmapExportDialogTypeString(),
-				EFileDialogFlags::None,
-				Filenames);
-		}
+		const FString& Filename = OptionalExportImportPath.GetValue();
 
-		if (bSuccess)
-		{
-			ensure(PropertyHandle_LayerFilename->SetValue(Filenames[0]) == FPropertyAccess::Success);
-			LandscapeEdMode->UISettings->LastImportPath = FPaths::GetPath(Filenames[0]);
-		}
+		ensure(PropertyHandle_LayerFilename->SetValue(Filename) == FPropertyAccess::Success);
+		LandscapeEdMode->UISettings->LastImportPath = FPaths::GetPath(Filename);
 	}
 
 	return FReply::Handled();
diff --git a/Engine/Source/Editor/LandscapeEditor/Private/LandscapeEditorDetailCustomization_NewLandscape.cpp b/Engine/Source/Editor/LandscapeEditor/Private/LandscapeEditorDetailCustomization_NewLandscape.cpp
--- a/Engine/Source/Editor/LandscapeEditor/Private/LandscapeEditorDetailCustomization_NewLandscape.cpp
+++ b/Engine/Source/Editor/LandscapeEditor/Private/LandscapeEditorDetailCustomization_NewLandscape.cpp
@@ -21,9 +21,16 @@
 #include "LandscapeEditorModule.h"
 #include "LandscapeEditorObject.h"
 #include "Landscape.h"
+#include "LandscapeStreamingProxy.h"
+#include "LandscapeTiledImage.h"
+#include "LandscapeRegionUtils.h"
+#include "LandscapeEditorPrivate.h"
+#include "LandscapeEditorUtils.h"
+
 #include "LandscapeConfigHelper.h"
 #include "LandscapeImportHelper.h"
 #include "LandscapeSettings.h"
+#include "FileHelpers.h"
 
 #include "DetailLayoutBuilder.h"
 #include "IDetailChildrenBuilder.h"
@@ -36,7 +43,6 @@
 #include "Widgets/Input/SVectorInputBox.h"
 #include "Widgets/Input/SRotatorInputBox.h"
 #include "ScopedTransaction.h"
-#include "DesktopPlatformModule.h"
 #include "AssetRegistry/AssetRegistryModule.h"
 
 #include "TutorialMetaData.h"
@@ -47,9 +53,22 @@
 #include "Editor.h"
 #include "Editor/EditorEngine.h"
 #include "LandscapeSubsystem.h"
+#include "LocationVolume.h"
 #include "SPrimaryButton.h"
 #include "Widgets/Input/SSegmentedControl.h"
 
+#include "UObject/SavePackage.h"
+#include "Builders/CubeBuilder.h"
+#include "ActorFactories/ActorFactory.h"
+
+#include "SourceControlHelpers.h"
+#include "WorldPartition/WorldPartition.h"
+
+#include "Misc/ScopedSlowTask.h"
+
+#include "Math/Box.h"
+#include "LandscapeEditorUtils.h"
+
 #define LOCTEXT_NAMESPACE "LandscapeEditor.NewLandscape"
 
 
@@ -122,16 +141,20 @@ void FLandscapeEditorDetailCustomization_NewLandscape::CustomizeDetails(IDetailL
 	PropertyHandle_FlipYAxis->SetOnPropertyValueChanged(FSimpleDelegate::CreateLambda([this]() { OnImportHeightmapFilenameChanged(); }));
 
 	TSharedRef<IPropertyHandle> PropertyHandle_GridSize = DetailBuilder.GetProperty(GET_MEMBER_NAME_CHECKED(ULandscapeEditorObject, WorldPartitionGridSize));
+	TSharedRef<IPropertyHandle> PropertyHandle_RegionSize = DetailBuilder.GetProperty(GET_MEMBER_NAME_CHECKED(ULandscapeEditorObject, WorldPartitionRegionSize));
 	
 	FEdModeLandscape* LandscapeEdMode = GetEditorMode();
 	if (LandscapeEdMode->GetWorld()->GetSubsystem<ULandscapeSubsystem>()->IsGridBased())
 	{
 		NewLandscapeCategory.AddProperty(PropertyHandle_GridSize);
+		NewLandscapeCategory.AddProperty(PropertyHandle_RegionSize);
 	}
 	else
 	{
 		DetailBuilder.HideProperty(PropertyHandle_GridSize);
+		DetailBuilder.HideProperty(PropertyHandle_RegionSize);
 	}
+	
 	TSharedRef<IPropertyHandle> PropertyHandle_HeightmapFilename = DetailBuilder.GetProperty(GET_MEMBER_NAME_CHECKED(ULandscapeEditorObject, ImportLandscape_HeightmapFilename));
 	TSharedRef<IPropertyHandle> PropertyHandle_HeightmapImportResult = DetailBuilder.GetProperty(GET_MEMBER_NAME_CHECKED(ULandscapeEditorObject, ImportLandscape_HeightmapImportResult));
 	TSharedRef<IPropertyHandle> PropertyHandle_HeightmapErrorMessage = DetailBuilder.GetProperty(GET_MEMBER_NAME_CHECKED(ULandscapeEditorObject, ImportLandscape_HeightmapErrorMessage));
@@ -392,9 +415,9 @@ void FLandscapeEditorDetailCustomization_NewLandscape::CustomizeDetails(IDetailL
 			SNew(SNumericEntryBox<int32>)
 			.Font(DetailBuilder.GetDetailFont())
 			.MinValue(1)
-			.MaxValue(32)
+			.MaxValue_Lambda(LandscapeEditorUtils::GetMaxSizeInComponents)
 			.MinSliderValue(1)
-			.MaxSliderValue(32)
+			.MaxSliderValue_Lambda(LandscapeEditorUtils::GetMaxSizeInComponents)
 			.AllowSpin(true)
 			.UndeterminedString(NSLOCTEXT("PropertyEditor", "MultipleValues", "Multiple Values"))
 			.Value_Static(&FLandscapeEditorDetailCustomization_Base::OnGetValue<int32>, PropertyHandle_ComponentCount_X)
@@ -416,9 +439,9 @@ void FLandscapeEditorDetailCustomization_NewLandscape::CustomizeDetails(IDetailL
 			SNew(SNumericEntryBox<int32>)
 			.Font(DetailBuilder.GetDetailFont())
 			.MinValue(1)
-			.MaxValue(32)
+			.MaxValue_Lambda(LandscapeEditorUtils::GetMaxSizeInComponents)
 			.MinSliderValue(1)
-			.MaxSliderValue(32)
+			.MaxSliderValue_Lambda(LandscapeEditorUtils::GetMaxSizeInComponents)
 			.AllowSpin(true)
 			.UndeterminedString(NSLOCTEXT("PropertyEditor", "MultipleValues", "Multiple Values"))
 			.Value_Static(&FLandscapeEditorDetailCustomization_Base::OnGetValue<int32>, PropertyHandle_ComponentCount_Y)
@@ -875,133 +898,360 @@ FText FLandscapeEditorDetailCustomization_NewLandscape::GetNewLandscapeErrorText
 	return FText::GetEmpty();
 }
 
-FReply FLandscapeEditorDetailCustomization_NewLandscape::OnCreateButtonClicked()
+void FLandscapeEditorDetailCustomization_NewLandscape::AddComponents(ULandscapeInfo* InLandscapeInfo, ULandscapeSubsystem* InLandscapeSubsystem, const TArray<FIntPoint>& InComponentCoordinates, TArray<ALandscapeProxy*>& OutCreatedStreamingProxies)
 {
-	FEdModeLandscape* LandscapeEdMode = GetEditorMode();
-	if (LandscapeEdMode != nullptr && 
-		LandscapeEdMode->GetWorld() != nullptr && 
-		LandscapeEdMode->GetWorld()->GetCurrentLevel()->bIsVisible)
+	TRACE_CPUPROFILER_EVENT_SCOPE(AddComponents);
+	TArray<ULandscapeComponent*> NewComponents;
+	InLandscapeInfo->Modify();
+	for (const FIntPoint& ComponentCoordinate : InComponentCoordinates)
 	{
-		ULandscapeEditorObject* UISettings = LandscapeEdMode->UISettings;
-		const int32 ComponentCountX = UISettings->NewLandscape_ComponentCount.X;
-		const int32 ComponentCountY = UISettings->NewLandscape_ComponentCount.Y;
-		const int32 QuadsPerComponent = UISettings->NewLandscape_SectionsPerComponent * UISettings->NewLandscape_QuadsPerSection;
-		const int32 SizeX = ComponentCountX * QuadsPerComponent + 1;
-		const int32 SizeY = ComponentCountY * QuadsPerComponent + 1;
+		ULandscapeComponent* LandscapeComponent = InLandscapeInfo->XYtoComponentMap.FindRef(ComponentCoordinate);
+		if (LandscapeComponent)
+		{
+			continue;
+		}
 
-		TArray<FLandscapeImportLayerInfo> MaterialImportLayers;
-		ELandscapeImportResult LayerImportResult = LandscapeEdMode->NewLandscapePreviewMode == ENewLandscapePreviewMode::NewLandscape ? UISettings->CreateNewLayersInfo(MaterialImportLayers) : UISettings->CreateImportLayersInfo(MaterialImportLayers);
+		// Add New component...
+		FIntPoint ComponentBase = ComponentCoordinate * InLandscapeInfo->ComponentSizeQuads;
 
-		if (LayerImportResult == ELandscapeImportResult::Error)
+		ALandscapeProxy* LandscapeProxy = InLandscapeSubsystem->FindOrAddLandscapeProxy(InLandscapeInfo, ComponentBase);
+		if (!LandscapeProxy)
 		{
-			return FReply::Handled();
+			continue;
 		}
 
-		TMap<FGuid, TArray<uint16>> HeightDataPerLayers;
-		TMap<FGuid, TArray<FLandscapeImportLayerInfo>> MaterialLayerDataPerLayers;
+		OutCreatedStreamingProxies.Add(LandscapeProxy);
+
+		LandscapeComponent = NewObject<ULandscapeComponent>(LandscapeProxy, NAME_None, RF_Transactional);
+		NewComponents.Add(LandscapeComponent);
+		LandscapeComponent->Init(
+			ComponentBase.X, ComponentBase.Y,
+			LandscapeProxy->ComponentSizeQuads,
+			LandscapeProxy->NumSubsections,
+			LandscapeProxy->SubsectionSizeQuads
+		);
+
+		TArray<FColor> HeightData;
+		const int32 ComponentVerts = (LandscapeComponent->SubsectionSizeQuads + 1) * LandscapeComponent->NumSubsections;
+		const FColor PackedMidpoint = LandscapeDataAccess::PackHeight(LandscapeDataAccess::GetTexHeight(0.0f));
+		HeightData.Init(PackedMidpoint, FMath::Square(ComponentVerts));
+
+		LandscapeComponent->InitHeightmapData(HeightData, true);
+		LandscapeComponent->UpdateMaterialInstances();
+
+		InLandscapeInfo->XYtoComponentMap.Add(ComponentCoordinate, LandscapeComponent);
+		InLandscapeInfo->XYtoAddCollisionMap.Remove(ComponentCoordinate);
+	}
 
-		TArray<uint16> OutHeightData;
-		if (LandscapeEdMode->NewLandscapePreviewMode == ENewLandscapePreviewMode::NewLandscape)
+	// Need to register to use general height/xyoffset data update
+	for (int32 Idx = 0; Idx < NewComponents.Num(); Idx++)
+	{
+		NewComponents[Idx]->RegisterComponent();
+	}
+
+	const bool bHasXYOffset = false;
+	ALandscape* Landscape = InLandscapeInfo->LandscapeActor.Get();
+
+	bool bHasLandscapeLayersContent = Landscape && Landscape->HasLayersContent();
+
+	for (ULandscapeComponent* NewComponent : NewComponents)
+	{
+		if (bHasLandscapeLayersContent)
 		{
-			UISettings->InitializeDefaultHeightData(OutHeightData);
+			TArray<ULandscapeComponent*> ComponentsUsingHeightmap;
+			ComponentsUsingHeightmap.Add(NewComponent);
+
+			for (const FLandscapeLayer& Layer : Landscape->LandscapeLayers)
+			{
+				// Since we do not share heightmap when adding new component, we will provided the required array, but they will only be used for 1 component
+				TMap<UTexture2D*, UTexture2D*> CreatedHeightmapTextures;
+				NewComponent->AddDefaultLayerData(Layer.Guid, ComponentsUsingHeightmap, CreatedHeightmapTextures);
+			}
 		}
-		else
+
+		// Update Collision
+		NewComponent->UpdateCachedBounds();
+		NewComponent->UpdateBounds();
+		NewComponent->MarkRenderStateDirty();
+
+		if (!bHasLandscapeLayersContent)
 		{
-			UISettings->ExpandImportData(OutHeightData, MaterialImportLayers);
+			ULandscapeHeightfieldCollisionComponent* CollisionComp = NewComponent->GetCollisionComponent();
+			if (CollisionComp && !bHasXYOffset)
+			{
+				CollisionComp->MarkRenderStateDirty();
+				CollisionComp->RecreateCollision();
+			}
 		}
+	}
 
-		HeightDataPerLayers.Add(FGuid(), OutHeightData);
-		// ComputeHeightData will also modify/expand material layers data, which is why we create MaterialLayerDataPerLayers after calling ComputeHeightData
-		MaterialLayerDataPerLayers.Add(FGuid(), MoveTemp(MaterialImportLayers));
 
-		FScopedTransaction Transaction(LOCTEXT("Undo", "Creating New Landscape"));
+	if (Landscape)
+	{
+		GEngine->BroadcastOnActorMoved(Landscape);
+	}
+}
 
-		const FVector Offset = FTransform(UISettings->NewLandscape_Rotation, FVector::ZeroVector, UISettings->NewLandscape_Scale).TransformVector(FVector(-ComponentCountX * QuadsPerComponent / 2, -ComponentCountY * QuadsPerComponent / 2, 0));
-		
-		ALandscape* Landscape = LandscapeEdMode->GetWorld()->SpawnActor<ALandscape>(UISettings->NewLandscape_Location + Offset, UISettings->NewLandscape_Rotation);
-		Landscape->bCanHaveLayersContent = UISettings->bCanHaveLayersContent;
-		Landscape->LandscapeMaterial = UISettings->NewLandscape_Material.Get();
-		Landscape->SetActorRelativeScale3D(UISettings->NewLandscape_Scale);
-
-		// automatically calculate a lighting LOD that won't crash lightmass (hopefully)
-		// < 2048x2048 -> LOD0
-		// >=2048x2048 -> LOD1
-		// >= 4096x4096 -> LOD2
-		// >= 8192x8192 -> LOD3
-		Landscape->StaticLightingLOD = FMath::DivideAndRoundUp(FMath::CeilLogTwo((SizeX * SizeY) / (2048 * 2048) + 1), (uint32)2);
-
-		FString ReimportHeightmapFilePath;
-		if (LandscapeEdMode->NewLandscapePreviewMode == ENewLandscapePreviewMode::ImportLandscape)
+FReply FLandscapeEditorDetailCustomization_NewLandscape::OnCreateButtonClicked()
+{
+	TRACE_CPUPROFILER_EVENT_SCOPE(FLandscapeEditorDetailCustomization_NewLandscape::OnCreateButtonClicked);
+
+	FEdModeLandscape* LandscapeEdMode = GetEditorMode();
+	const bool bIsNewLandscape = LandscapeEdMode->NewLandscapePreviewMode == ENewLandscapePreviewMode::NewLandscape;
+	
+	UWorld* World = LandscapeEdMode->GetWorld();
+	const bool bIsTempPackage = FPackageName::IsTempPackage(World->GetPackage()->GetName());
+	
+	const bool bCreateLandscape = LandscapeEdMode != nullptr && 
+		World != nullptr &&
+		World->GetCurrentLevel()->bIsVisible;
+	
+	if (!bCreateLandscape)
+	{
+		return FReply::Handled();
+	}
+	
+	ULandscapeEditorObject* UISettings = LandscapeEdMode->UISettings;
+	
+	const bool bIsWorldPartition = World->GetSubsystem<ULandscapeSubsystem>()->IsGridBased();
+	const bool bLandscapeLargerThanRegion = static_cast<int32>(UISettings->WorldPartitionRegionSize) < UISettings->NewLandscape_ComponentCount.X || static_cast<int32>(UISettings->WorldPartitionRegionSize) < UISettings->NewLandscape_ComponentCount.Y;
+	const bool bNeedsLandscapeRegions =  bIsWorldPartition && bLandscapeLargerThanRegion;
+
+	// If we need to ensure the map is saved before proceeding to create a landscape with regions 
+	if (bIsTempPackage && bNeedsLandscapeRegions)
+	{
+		FString NewMapPackageName;
+		if (!FEditorFileUtils::SaveCurrentLevel())
 		{
-			ReimportHeightmapFilePath = UISettings->ImportLandscape_HeightmapFilename;
+			UE_LOG(LogLandscapeTools, Error, TEXT("Unable to save current level"));
+			return FReply::Handled();
 		}
+	}
+
+	const int32 QuadsPerSection = UISettings->NewLandscape_QuadsPerSection;
+	
+	const FIntPoint TotalLandscapeComponentSize { UISettings->NewLandscape_ComponentCount.X, UISettings->NewLandscape_ComponentCount.Y };
 
-		Landscape->Import(FGuid::NewGuid(), 0, 0, SizeX - 1, SizeY - 1, UISettings->NewLandscape_SectionsPerComponent, UISettings->NewLandscape_QuadsPerSection, HeightDataPerLayers, *ReimportHeightmapFilePath, MaterialLayerDataPerLayers, UISettings->ImportLandscape_AlphamapType);
+	const int32 ComponentCountX = bNeedsLandscapeRegions ? UISettings->WorldPartitionRegionSize : TotalLandscapeComponentSize.X;
+	const int32 ComponentCountY = bNeedsLandscapeRegions ? UISettings->WorldPartitionRegionSize : TotalLandscapeComponentSize.Y;
+	const int32 QuadsPerComponent = UISettings->NewLandscape_SectionsPerComponent * QuadsPerSection;
+	const int32 SizeX = ComponentCountX * QuadsPerComponent + 1;
+	const int32 SizeY = ComponentCountY * QuadsPerComponent + 1;
 
-		ULandscapeInfo* LandscapeInfo = Landscape->GetLandscapeInfo();
-		check(LandscapeInfo);
+	TArray<FLandscapeImportLayerInfo> MaterialImportLayers;
+	ELandscapeImportResult LayerImportResult = bIsNewLandscape  ? UISettings->CreateNewLayersInfo(MaterialImportLayers) : UISettings->CreateImportLayersInfo(MaterialImportLayers);
 
-		FActorLabelUtilities::SetActorLabelUnique(Landscape, ALandscape::StaticClass()->GetName());
+	if (LayerImportResult == ELandscapeImportResult::Error)
+	{
+		UE_LOG(LogLandscapeTools, Error, TEXT("Unable to import weight maps"));
+		return FReply::Handled();
+	}
 
-		LandscapeInfo->UpdateLayerInfoMap(Landscape);
+	TMap<FGuid, TArray<uint16>> HeightDataPerLayers;
+	TMap<FGuid, TArray<FLandscapeImportLayerInfo>> MaterialLayerDataPerLayers;
 
-		// Import doesn't fill in the LayerInfo for layers with no data, do that now
-		const TArray<FLandscapeImportLayer>& ImportLandscapeLayersList = UISettings->ImportLandscape_Layers;
-		const ULandscapeSettings* Settings = GetDefault<ULandscapeSettings>();
-		TSoftObjectPtr<ULandscapeLayerInfoObject> DefaultLayerInfoObject = Settings->GetDefaultLayerInfoObject().LoadSynchronous();
+	TArray<uint16> OutHeightData;
+	
+	if (bIsNewLandscape || bNeedsLandscapeRegions)
+	{
+		UISettings->InitializeDefaultHeightData(OutHeightData);
+	}
+	else
+	{
+		UISettings->ExpandImportData(OutHeightData, MaterialImportLayers);
+	}
+
+	HeightDataPerLayers.Add(FGuid(), OutHeightData);
+	// ComputeHeightData will also modify/expand material layers data, which is why we create MaterialLayerDataPerLayers after calling ComputeHeightData
+	MaterialLayerDataPerLayers.Add(FGuid(), MoveTemp(MaterialImportLayers));
 
-		for (int32 i = 0; i < ImportLandscapeLayersList.Num(); i++)
+	FScopedTransaction Transaction(bIsNewLandscape ? LOCTEXT("Undo_CreateLandscape", "Creating New Landscape") : LOCTEXT("Undo_CreateAndImportLandscape", "Creating and Importing New Landscape"));
+
+	const FVector Offset = FTransform(UISettings->NewLandscape_Rotation, FVector::ZeroVector, UISettings->NewLandscape_Scale).TransformVector(FVector(-UISettings->NewLandscape_ComponentCount.X * QuadsPerComponent / 2, -UISettings->NewLandscape_ComponentCount.Y * QuadsPerComponent / 2, 0));
+	
+	ALandscape* Landscape = LandscapeEdMode->GetWorld()->SpawnActor<ALandscape>(UISettings->NewLandscape_Location + Offset, UISettings->NewLandscape_Rotation);
+	Landscape->bCanHaveLayersContent = UISettings->bCanHaveLayersContent;
+	Landscape->LandscapeMaterial = UISettings->NewLandscape_Material.Get();
+	Landscape->SetActorRelativeScale3D(UISettings->NewLandscape_Scale);
+
+	// automatically calculate a lighting LOD that won't crash lightmass (hopefully)
+	// < 2048x2048 -> LOD0
+	// >=2048x2048 -> LOD1
+	// >= 4096x4096 -> LOD2
+	// >= 8192x8192 -> LOD3
+	Landscape->StaticLightingLOD = FMath::DivideAndRoundUp(FMath::CeilLogTwo((SizeX * SizeY) / (2048 * 2048) + 1), (uint32)2);
+
+	FString ReimportHeightmapFilePath;
+	if (LandscapeEdMode->NewLandscapePreviewMode == ENewLandscapePreviewMode::ImportLandscape)
+	{
+		ReimportHeightmapFilePath = UISettings->ImportLandscape_HeightmapFilename;
+	}
+
+	Landscape->Import(FGuid::NewGuid(), 0, 0, SizeX - 1, SizeY - 1, UISettings->NewLandscape_SectionsPerComponent, QuadsPerSection, HeightDataPerLayers, *ReimportHeightmapFilePath, MaterialLayerDataPerLayers, UISettings->ImportLandscape_AlphamapType);
+
+	ULandscapeInfo* LandscapeInfo = Landscape->GetLandscapeInfo();
+	check(LandscapeInfo);
+
+	FActorLabelUtilities::SetActorLabelUnique(Landscape, ALandscape::StaticClass()->GetName());
+
+	LandscapeInfo->UpdateLayerInfoMap(Landscape);
+
+	// Import doesn't fill in the LayerInfo for layers with no data, do that now
+	const TArray<FLandscapeImportLayer>& ImportLandscapeLayersList = UISettings->ImportLandscape_Layers;
+	const ULandscapeSettings* Settings = GetDefault<ULandscapeSettings>();
+	TSoftObjectPtr<ULandscapeLayerInfoObject> DefaultLayerInfoObject = Settings->GetDefaultLayerInfoObject().LoadSynchronous();
+
+	for (int32 i = 0; i < ImportLandscapeLayersList.Num(); i++)
+	{
+		ULandscapeLayerInfoObject* LayerInfo = ImportLandscapeLayersList[i].LayerInfo;
+		FName LayerName = ImportLandscapeLayersList[i].LayerName;
+
+		// If DefaultLayerInfoObject is set and LayerInfo does not exist, we will try to create the new LayerInfo by cloning DefaultLayerInfoObject. Except for VisibilityLayer which doesn't require an asset.
+		if (DefaultLayerInfoObject.IsValid() && (LayerInfo == nullptr) && (LayerName != ALandscapeProxy::VisibilityLayer->LayerName))
 		{
-			ULandscapeLayerInfoObject* LayerInfo = ImportLandscapeLayersList[i].LayerInfo;
-			FName LayerName = ImportLandscapeLayersList[i].LayerName;
+			LayerInfo = Landscape->CreateLayerInfo(*LayerName.ToString(), DefaultLayerInfoObject.Get());
 
-			// If DefaultLayerInfoObject is set and LayerInfo does not exist, we will try to create the new LayerInfo by cloning DefaultLayerInfoObject. Except for VisibilityLayer which doesn't require an asset.
-			if (DefaultLayerInfoObject.IsValid() && (LayerInfo == nullptr) && (LayerName != ALandscapeProxy::VisibilityLayer->LayerName))
+			if (LayerInfo != nullptr)
 			{
-				LayerInfo = Landscape->CreateLayerInfo(*LayerName.ToString(), DefaultLayerInfoObject.Get());
+				LayerInfo->LayerUsageDebugColor = LayerInfo->GenerateLayerUsageDebugColor();
+				LayerInfo->MarkPackageDirty();
+			}
+		}
 
-				if (LayerInfo != nullptr)
-				{
-					LayerInfo->LayerUsageDebugColor = LayerInfo->GenerateLayerUsageDebugColor();
-					LayerInfo->MarkPackageDirty();
-				}
+		if (LayerInfo != nullptr)
+		{
+			if (LandscapeEdMode->NewLandscapePreviewMode == ENewLandscapePreviewMode::ImportLandscape)
+			{
+				Landscape->EditorLayerSettings.Add(FLandscapeEditorLayerSettings(LayerInfo, ImportLandscapeLayersList[i].SourceFilePath));
+			}
+			else
+			{
+				Landscape->EditorLayerSettings.Add(FLandscapeEditorLayerSettings(LayerInfo));
 			}
 
-			if (LayerInfo != nullptr)
+			int32 LayerInfoIndex = LandscapeInfo->GetLayerInfoIndex(ImportLandscapeLayersList[i].LayerName);
+			if (ensure(LayerInfoIndex != INDEX_NONE))
 			{
-				if (LandscapeEdMode->NewLandscapePreviewMode == ENewLandscapePreviewMode::ImportLandscape)
-				{
-					Landscape->EditorLayerSettings.Add(FLandscapeEditorLayerSettings(LayerInfo, ImportLandscapeLayersList[i].SourceFilePath));
-				}
-				else
-				{
-					Landscape->EditorLayerSettings.Add(FLandscapeEditorLayerSettings(LayerInfo));
-				}
+				FLandscapeInfoLayerSettings& LayerSettings = LandscapeInfo->Layers[LayerInfoIndex];
+				LayerSettings.LayerInfoObj = LayerInfo;
+			}
+		}
+	}
+
+	LandscapeEdMode->UpdateLandscapeList();
+	LandscapeEdMode->SetLandscapeInfo(LandscapeInfo);
+	LandscapeEdMode->CurrentToolTarget.TargetType = ELandscapeToolTargetType::Heightmap;
+	LandscapeEdMode->SetCurrentTargetLayer(NAME_None, nullptr);
+	LandscapeEdMode->SetCurrentTool("Select"); // change tool so switching back to the manage mode doesn't give "New Landscape" again
+	LandscapeEdMode->SetCurrentTool("Sculpt"); // change to sculpting mode and tool
+	LandscapeEdMode->SetCurrentLayer(0);
 
-				int32 LayerInfoIndex = LandscapeInfo->GetLayerInfoIndex(ImportLandscapeLayersList[i].LayerName);
-				if (ensure(LayerInfoIndex != INDEX_NONE))
+	World->GetSubsystem<ULandscapeSubsystem>()->ChangeGridSize(LandscapeInfo, UISettings->WorldPartitionGridSize);
+
+	if (LandscapeEdMode->CurrentToolTarget.LandscapeInfo.IsValid())
+	{
+		ALandscapeProxy* LandscapeProxy = LandscapeEdMode->CurrentToolTarget.LandscapeInfo->GetLandscapeProxy();
+		LandscapeProxy->OnMaterialChangedDelegate().AddRaw(LandscapeEdMode, &FEdModeLandscape::OnLandscapeMaterialChangedDelegate);
+	}
+
+	UWorldPartition* WorldPartition = World->GetWorldPartition();
+	ULandscapeSubsystem* LandscapeSubsystem = World->GetSubsystem<ULandscapeSubsystem>();
+	ALandscapeProxy* LandscapeProxy = LandscapeEdMode->CurrentToolTarget.LandscapeInfo->GetLandscapeProxy();
+	
+	if (bNeedsLandscapeRegions)
+	{
+		ULevel* Level = LandscapeProxy->GetLevel();
+		UPackage* LevelPackage = Level->GetPackage();
+
+		TArray<FIntPoint> NewComponents;
+		NewComponents.Empty(TotalLandscapeComponentSize.X * TotalLandscapeComponentSize.Y);
+		for (int32 Y = 0; Y < TotalLandscapeComponentSize.Y; Y++)
+		{
+			for (int32 X = 0; X < TotalLandscapeComponentSize.X; X++)
+			{
+				NewComponents.Add(FIntPoint(X, Y));
+			}
+		}
+
+		int32 NumRegions = FMath::DivideAndRoundUp(TotalLandscapeComponentSize.X, static_cast<int32>(UISettings->WorldPartitionRegionSize)) * FMath::DivideAndRoundUp(TotalLandscapeComponentSize.Y, static_cast<int32>(UISettings->WorldPartitionRegionSize));
+
+
+		FScopedSlowTask Progress(NumRegions, LOCTEXT("CreateLandscapeRegions", "Creating Landscape Editor Regions..."));
+		Progress.MakeDialog();
+
+		TArray<ALocationVolume*> RegionVolumes;
+		FBox LandscapeBounds;
+		auto AddComponentsToRegion = [&Progress, NumRegions, bIsNewLandscape, LandscapeEdMode, WorldPartition, World, LandscapeProxy, &UISettings, QuadsPerSection, LandscapeInfo, LandscapeSubsystem, &RegionVolumes, &LandscapeBounds,&MaterialLayerDataPerLayers](const FIntPoint& RegionCoordinate, const TArray<FIntPoint>& NewComponents)
+		{
+			TRACE_CPUPROFILER_EVENT_SCOPE(AddComponentsToRegion);
+						
+			// Create a LocationVolume around the region 
+			int32 RegionSizeXTexels = QuadsPerSection * UISettings->WorldPartitionRegionSize;
+			int32 RegionSizeYTexels = QuadsPerSection * UISettings->WorldPartitionRegionSize;
+
+			double RegionSizeX = QuadsPerSection * LandscapeProxy->GetActorScale3D().X * UISettings->WorldPartitionRegionSize;
+			double RegionSizeY = QuadsPerSection * LandscapeProxy->GetActorScale3D().Y * UISettings->WorldPartitionRegionSize;
+			ALocationVolume* RegionVolume = LandscapeRegionUtils::CreateLandscapeRegionVolume(World, LandscapeProxy, RegionCoordinate, RegionSizeX);
+			RegionVolumes.Add(RegionVolume);
+			
+			TArray<ALandscapeProxy*> CreatedStreamingProxies;
+			AddComponents(LandscapeInfo, LandscapeSubsystem, NewComponents, CreatedStreamingProxies);
+			
+			FIntRect ImportRegion (RegionSizeXTexels * RegionCoordinate.X, RegionSizeYTexels* RegionCoordinate.Y, RegionSizeXTexels * (RegionCoordinate.X +1) + 1  , RegionSizeYTexels * (RegionCoordinate.Y +1) +1 );
+			FIntPoint ImportOffset(0, 0);
+
+			if (!bIsNewLandscape && !UISettings->ImportLandscape_HeightmapFilename.IsEmpty())
+			{
+				LandscapeEdMode->ImportHeightData(LandscapeInfo, LandscapeEdMode->GetCurrentLayerGuid(), UISettings->ImportLandscape_HeightmapFilename, ImportRegion, ELandscapeImportTransformType::Subregion, ImportOffset, ELandscapeLayerPaintingRestriction::None, LandscapeEdMode->UISettings->bFlipYAxis);
+
+				TArray<FLandscapeImportLayerInfo>& Weights = *MaterialLayerDataPerLayers.Find(FGuid());
+
+				for (const TPair<FGuid, TArray<FLandscapeImportLayerInfo>>& Layer : MaterialLayerDataPerLayers)
 				{
-					FLandscapeInfoLayerSettings& LayerSettings = LandscapeInfo->Layers[LayerInfoIndex];
-					LayerSettings.LayerInfoObj = LayerInfo;
+					for (const FLandscapeImportLayerInfo& WeightMap : Layer.Value)
+					{						
+						LandscapeEdMode->ImportWeightData(LandscapeInfo, LandscapeEdMode->GetCurrentLayerGuid(), WeightMap.LayerInfo, WeightMap.SourceFilePath, ImportRegion, ELandscapeImportTransformType::Subregion, ImportOffset, ELandscapeLayerPaintingRestriction::None, LandscapeEdMode->UISettings->bFlipYAxis);
+					}
 				}
 			}
-		}
 
-		LandscapeEdMode->UpdateLandscapeList();
-		LandscapeEdMode->SetLandscapeInfo(LandscapeInfo);
-		LandscapeEdMode->CurrentToolTarget.TargetType = ELandscapeToolTargetType::Heightmap;
-		LandscapeEdMode->SetCurrentTargetLayer(NAME_None, nullptr);
-		LandscapeEdMode->SetCurrentTool("Select"); // change tool so switching back to the manage mode doesn't give "New Landscape" again
-		LandscapeEdMode->SetCurrentTool("Sculpt"); // change to sculpting mode and tool
-		LandscapeEdMode->SetCurrentLayer(0);
+			// ensures all the final height textures have been updated.
+			LandscapeInfo->ForceLayersFullUpdate();
+			LandscapeEditorUtils::SaveLandscapeProxies(MakeArrayView(CreatedStreamingProxies), WorldPartition);
+			LandscapeBounds += LandscapeInfo->GetCompleteBounds();
+
+			Progress.EnterProgressFrame(1.0f , FText::Format(LOCTEXT("LandscapeCreateRegion", "Creating Landscape Editor Regions ({0}, {1})"), RegionCoordinate.X, RegionCoordinate.Y));
+			return true;
+		};
 
-		LandscapeEdMode->GetWorld()->GetSubsystem<ULandscapeSubsystem>()->ChangeGridSize(LandscapeInfo, UISettings->WorldPartitionGridSize);
+		LandscapeEditorUtils::SaveObjects(TArrayView<ALandscape*>(TArray<ALandscape*> { LandscapeInfo->LandscapeActor.Get() }));
 
-		if (LandscapeEdMode->CurrentToolTarget.LandscapeInfo.IsValid())
+		LandscapeRegionUtils::ForEachComponentByRegion(UISettings->WorldPartitionRegionSize, NewComponents, AddComponentsToRegion);
+
+		// update the zcomponent of the volumes 
+		const float ZScale =  LandscapeBounds.Max.Z - LandscapeBounds.Min.Z;
+		for (ALocationVolume* RegionVolume : RegionVolumes)
 		{
-			ALandscapeProxy* LandscapeProxy = LandscapeEdMode->CurrentToolTarget.LandscapeInfo->GetLandscapeProxy();
-			LandscapeProxy->OnMaterialChangedDelegate().AddRaw(LandscapeEdMode, &FEdModeLandscape::OnLandscapeMaterialChangedDelegate);
+			const FVector Scale = RegionVolume->GetActorScale();
+			const FVector NewScale{ Scale.X, Scale.Y, 1000000.0f}; //ZScale * 10.0 
+			RegionVolume->SetActorScale3D(NewScale);
 		}
+
+		LandscapeEditorUtils::SaveObjects(MakeArrayView(RegionVolumes));
+
+		TArray<ALandscapeProxy*> AllProxies;
+		
+		// save the initial region & unload it
+		LandscapeInfo->ForEachLandscapeProxy([&WorldPartition, &AllProxies](ALandscapeProxy* Proxy) {
+			if (Proxy->IsA<ALandscapeStreamingProxy>())
+			{		
+				AllProxies.Add(Proxy);
+				
+			}
+			return true;
+		});
+
+		LandscapeEditorUtils::SaveLandscapeProxies(MakeArrayView(AllProxies), WorldPartition);
 	}
 
 	return FReply::Handled();
@@ -1059,8 +1309,7 @@ bool FLandscapeEditorDetailCustomization_NewLandscape::GetImportButtonIsEnabled(
 	if (LandscapeEdMode != nullptr)
 	{
 		if (LandscapeEdMode->UISettings->ImportLandscape_HeightmapImportResult == ELandscapeImportResult::Error ||
-			LandscapeEdMode->UISettings->ImportLandscape_HeightmapFilename.IsEmpty() ||
-			LandscapeEdMode->UISettings->GetImportLandscapeData().IsEmpty())
+			LandscapeEdMode->UISettings->ImportLandscape_HeightmapFilename.IsEmpty()) 
 		{
 			return false;
 		}
@@ -1147,28 +1396,15 @@ FReply FLandscapeEditorDetailCustomization_NewLandscape::OnImportHeightmapFilena
 	FEdModeLandscape* LandscapeEdMode = GetEditorMode();
 	check(LandscapeEdMode != nullptr);
 
-	// Prompt the user for the Filenames
-	IDesktopPlatform* DesktopPlatform = FDesktopPlatformModule::Get();
-	if (DesktopPlatform != nullptr)
+	ILandscapeEditorModule& LandscapeEditorModule = FModuleManager::GetModuleChecked<ILandscapeEditorModule>("LandscapeEditor");
+	const FString FileTypes = LandscapeEditorModule.GetHeightmapImportDialogTypeString();
+
+	TOptional<FString> OptionalFilename = LandscapeEditorUtils::GetImportExportFilename(NSLOCTEXT("UnrealEd", "Import", "Import").ToString(), LandscapeEdMode->UISettings->LastImportPath, FileTypes);
+	if (OptionalFilename.IsSet())
 	{
-		ILandscapeEditorModule& LandscapeEditorModule = FModuleManager::GetModuleChecked<ILandscapeEditorModule>("LandscapeEditor");
-		const TCHAR* FileTypes = LandscapeEditorModule.GetHeightmapImportDialogTypeString();
-
-		TArray<FString> OpenFilenames;
-		bool bOpened = DesktopPlatform->OpenFileDialog(
-			FSlateApplication::Get().FindBestParentWindowHandleForDialogs(nullptr),
-			NSLOCTEXT("UnrealEd", "Import", "Import").ToString(),
-			LandscapeEdMode->UISettings->LastImportPath,
-			TEXT(""),
-			FileTypes,
-			EFileDialogFlags::None,
-			OpenFilenames);
-
-		if (bOpened)
-		{
-			ensure(PropertyHandle_HeightmapFilename->SetValue(OpenFilenames[0]) == FPropertyAccess::Success);
-			LandscapeEdMode->UISettings->LastImportPath = FPaths::GetPath(OpenFilenames[0]);
-		}
+		const FString& Filename = OptionalFilename.GetValue();
+		ensure(PropertyHandle_HeightmapFilename->SetValue(Filename) == FPropertyAccess::Success);
+		LandscapeEdMode->UISettings->LastImportPath = FPaths::GetPath(Filename);
 	}
 
 	return FReply::Handled();
diff --git a/Engine/Source/Editor/LandscapeEditor/Private/LandscapeEditorDetailCustomization_NewLandscape.h b/Engine/Source/Editor/LandscapeEditor/Private/LandscapeEditorDetailCustomization_NewLandscape.h
--- a/Engine/Source/Editor/LandscapeEditor/Private/LandscapeEditorDetailCustomization_NewLandscape.h
+++ b/Engine/Source/Editor/LandscapeEditor/Private/LandscapeEditorDetailCustomization_NewLandscape.h
@@ -13,6 +13,7 @@
 #include "IPropertyTypeCustomization.h"
 #include "PropertyHandle.h"
 #include "LandscapeEditorDetailCustomization_Base.h"
+#include "LandscapeSubsystem.h"
 
 class IDetailLayoutBuilder;
 
@@ -93,5 +94,9 @@ public:
 	bool ShouldShowResetMaterialToDefault(TSharedPtr<IPropertyHandle> InPropertyHandle);
 
 protected:
+
 	bool bUsingSlider;
+
+private:
+	static void AddComponents(ULandscapeInfo* InLandscapeInfo, ULandscapeSubsystem* InLandscapeSubsystem, const TArray<FIntPoint>& InComponentCoordinates, TArray<ALandscapeProxy*>& OutCreatedStreamingProxies);
 };
diff --git a/Engine/Source/Editor/LandscapeEditor/Private/LandscapeEditorModule.cpp b/Engine/Source/Editor/LandscapeEditor/Private/LandscapeEditorModule.cpp
--- a/Engine/Source/Editor/LandscapeEditor/Private/LandscapeEditorModule.cpp
+++ b/Engine/Source/Editor/LandscapeEditor/Private/LandscapeEditorModule.cpp
@@ -20,6 +20,7 @@
 #include "LandscapeFileFormatRaw.h"
 #include "Settings/EditorExperimentalSettings.h"
 #include "LandscapeEditorServices.h"
+#include "LandscapeImageFileCache.h"
 
 #include "Framework/MultiBox/MultiBoxBuilder.h"
 #include "PropertyEditorModule.h"
@@ -161,6 +162,8 @@ public:
 		ILandscapeModule& LandscapeModule = FModuleManager::GetModuleChecked<ILandscapeModule>("Landscape");
 		LandscapeEditorServices.Reset(new FLandscapeEditorServices);
 		LandscapeModule.SetLandscapeEditorServices(LandscapeEditorServices.Get());
+
+		LandscapeImageFileCache.Reset(new FLandscapeImageFileCache());
 	}
 
 	/**
@@ -198,6 +201,7 @@ public:
 			LandscapeModule.SetLandscapeEditorServices(nullptr);
 		}
 		LandscapeEditorServices.Reset();
+		LandscapeImageFileCache.Reset();
 	}
 
 	static void ConstructLandscapeViewportMenu(FMenuBuilder& MenuBuilder)
@@ -371,6 +375,8 @@ public:
 
 	virtual TSharedPtr<FUICommandList> GetLandscapeLevelViewportCommandList() const override;
 
+	FLandscapeImageFileCache& GetImageFileCache() const override;
+
 protected:
 	TSharedPtr<FExtender> ViewportMenuExtender;
 	TSharedPtr<FUICommandList> GlobalUICommandList;
@@ -381,6 +387,7 @@ protected:
 	mutable FString HeightmapExportDialogTypeString;
 	mutable FString WeightmapExportDialogTypeString;
 	TUniquePtr<ILandscapeEditorServices> LandscapeEditorServices;
+	TUniquePtr<FLandscapeImageFileCache> LandscapeImageFileCache;
 };
 
 IMPLEMENT_MODULE(FLandscapeEditorModule, LandscapeEditor);
@@ -549,4 +556,11 @@ TSharedPtr<FUICommandList> FLandscapeEditorModule::GetLandscapeLevelViewportComm
 	return GlobalUICommandList;
 }
 
+FLandscapeImageFileCache& FLandscapeEditorModule::GetImageFileCache() const
+{
+	check(LandscapeImageFileCache != nullptr);
+	return *LandscapeImageFileCache;
+}
+
+
 #undef LOCTEXT_NAMESPACE
diff --git a/Engine/Source/Editor/LandscapeEditor/Private/LandscapeEditorObject.cpp b/Engine/Source/Editor/LandscapeEditor/Private/LandscapeEditorObject.cpp
--- a/Engine/Source/Editor/LandscapeEditor/Private/LandscapeEditorObject.cpp
+++ b/Engine/Source/Editor/LandscapeEditor/Private/LandscapeEditorObject.cpp
@@ -11,6 +11,7 @@
 #include "LandscapeRender.h"
 #include "LandscapeSettings.h"
 #include "LandscapeImportHelper.h"
+#include "LandscapeTiledImage.h"
 #include "LandscapeMaterialInstanceConstant.h"
 #include "Misc/ConfigCacheIni.h"
 #include "EngineUtils.h"
@@ -155,7 +156,11 @@ void ULandscapeEditorObject::PostEditChangeProperty(FPropertyChangedEvent& Prope
 		PropertyChangedEvent.MemberProperty->GetFName() == GET_MEMBER_NAME_CHECKED(ULandscapeEditorObject, NewLandscape_SectionsPerComponent) ||
 		PropertyChangedEvent.MemberProperty->GetFName() == GET_MEMBER_NAME_CHECKED(ULandscapeEditorObject, NewLandscape_ComponentCount))
 	{
-		NewLandscape_ClampSize();
+		// Only clamp the landscape size if we're not in World Partition
+		if (!ParentMode->IsGridBased())
+		{
+			NewLandscape_ClampSize();
+		}
 	}
 
 	if (PropertyChangedEvent.MemberProperty == nullptr ||
@@ -637,6 +642,8 @@ bool ULandscapeEditorObject::UseSingleFileImport() const
 
 void ULandscapeEditorObject::RefreshImports()
 {
+	TRACE_CPUPROFILER_EVENT_SCOPE(ULandscapeEditorObject::RefreshImports);
+
 	ClearImportLandscapeData();
 	HeightmapImportDescriptorIndex = 0;
 	HeightmapImportDescriptor.Reset();
@@ -648,12 +655,12 @@ void ULandscapeEditorObject::RefreshImports()
 
 	if (!ImportLandscape_HeightmapFilename.IsEmpty())
 	{
-		ImportLandscape_HeightmapImportResult =
-			FLandscapeImportHelper::GetHeightmapImportDescriptor(ImportLandscape_HeightmapFilename, UseSingleFileImport(), bFlipYAxis, HeightmapImportDescriptor, ImportLandscape_HeightmapErrorMessage);
-		if (ImportLandscape_HeightmapImportResult != ELandscapeImportResult::Error)
+		FLandscapeTiledImage TiledImage;
+		
+		if (FLandscapeFileInfo FileInfo = TiledImage.Load(*ImportLandscape_HeightmapFilename); FileInfo.ResultCode == ELandscapeImportResult::Success)
 		{
-			ImportLandscape_Width = HeightmapImportDescriptor.ImportResolutions[HeightmapImportDescriptorIndex].Width;
-			ImportLandscape_Height = HeightmapImportDescriptor.ImportResolutions[HeightmapImportDescriptorIndex].Height;
+			ImportLandscape_Width = TiledImage.GetResolution().X;
+			ImportLandscape_Height = TiledImage.GetResolution().Y;
 			ChooseBestComponentSizeForImport();
 			ImportLandscapeData();
 		}
@@ -697,27 +704,16 @@ void ULandscapeEditorObject::RefreshLayerImport(FLandscapeImportLayer& ImportLay
 		}
 		else
 		{
-			ImportLayer.ImportResult = FLandscapeImportHelper::GetWeightmapImportDescriptor(ImportLayer.SourceFilePath, UseSingleFileImport(), bFlipYAxis, ImportLayer.LayerName, ImportLayer.ImportDescriptor, ImportLayer.ErrorMessage);
-			if (ImportLayer.ImportResult != ELandscapeImportResult::Error)
+			FLandscapeTiledImage TiledImage;
+			FLandscapeFileInfo FileInfo = TiledImage.Load(*ImportLayer.SourceFilePath);
+			ImportLayer.ImportResult = FileInfo.ResultCode;
+			ImportLayer.ErrorMessage = FileInfo.ErrorMessage;
+			if (FileInfo.ResultCode == ELandscapeImportResult::Success)
 			{
-				if (ImportLandscape_Height != 0 || ImportLandscape_Width != 0)
+				if (FileInfo.PossibleResolutions[0] != FLandscapeFileResolution(ImportLandscape_Height, ImportLandscape_Width))
 				{
-					// Use same import index as Heightmap
-					int32 FoundIndex = INDEX_NONE;
-					for (int32 Index = 0; Index < ImportLayer.ImportDescriptor.FileResolutions.Num(); ++Index)
-					{
-						if (ImportLayer.ImportDescriptor.ImportResolutions[Index] == HeightmapImportDescriptor.ImportResolutions[HeightmapImportDescriptorIndex])
-						{
-							FoundIndex = Index;
-							break;
-						}
-					}
-
-					if (FoundIndex == INDEX_NONE)
-					{
-						ImportLayer.ImportResult = ELandscapeImportResult::Error;
-						ImportLayer.ErrorMessage = NSLOCTEXT("LandscapeEditor.ImportLandscape", "Import_WeightHeightResolutionMismatch", "Weightmap import resolution isn't same as Heightmap resolution.");
-					}
+					ImportLayer.ImportResult = ELandscapeImportResult::Error;
+					ImportLayer.ErrorMessage = NSLOCTEXT("LandscapeEditor.ImportLandscape", "Import_WeightHeightResolutionMismatch", "Weightmap import resolution isn't same as Heightmap resolution.");
 				}
 			}
 		}
@@ -737,11 +733,15 @@ void ULandscapeEditorObject::OnChangeImportLandscapeResolution(int32 DescriptorI
 
 void ULandscapeEditorObject::ImportLandscapeData()
 {
-	ImportLandscape_HeightmapImportResult = FLandscapeImportHelper::GetHeightmapImportData(HeightmapImportDescriptor, HeightmapImportDescriptorIndex, ImportLandscape_Data, ImportLandscape_HeightmapErrorMessage);
-	if (ImportLandscape_HeightmapImportResult == ELandscapeImportResult::Error)
+	FLandscapeTiledImage TiledImage;
+	FLandscapeFileInfo FileInfo = TiledImage.Load(*ImportLandscape_HeightmapFilename);
+
+	if (FileInfo.ResultCode != ELandscapeImportResult::Success)
 	{
 		ImportLandscape_Data.Empty();
 	}
+
+	TiledImage.Read(ImportLandscape_Data, bFlipYAxis);
 }
 
 ELandscapeImportResult ULandscapeEditorObject::CreateImportLayersInfo(TArray<FLandscapeImportLayerInfo>& OutImportLayerInfos)
@@ -765,38 +765,18 @@ ELandscapeImportResult ULandscapeEditorObject::CreateImportLayersInfo(TArray<FLa
 
 		if (ImportLayer.LayerInfo != nullptr && !ImportLayer.SourceFilePath.IsEmpty())
 		{
-			UIImportLayer.ImportResult = FLandscapeImportHelper::GetWeightmapImportDescriptor(ImportLayer.SourceFilePath, UseSingleFileImport(), bFlipYAxis, ImportLayer.LayerName, UIImportLayer.ImportDescriptor, UIImportLayer.ErrorMessage);
-			if (UIImportLayer.ImportResult == ELandscapeImportResult::Error)
-			{
-				FMessageDialog::Open(EAppMsgType::Ok, UIImportLayer.ErrorMessage);
-				return ELandscapeImportResult::Error;
-			}
+			FLandscapeTiledImage LayerImage;
+			FLandscapeFileInfo LayerFileInfo = LayerImage.Load(*ImportLayer.SourceFilePath);
 
-			// Use same import index as Heightmap
-			int32 FoundIndex = INDEX_NONE;
-			for (int32 Index = 0; Index < UIImportLayer.ImportDescriptor.FileResolutions.Num(); ++Index)
-			{
-				if (UIImportLayer.ImportDescriptor.ImportResolutions[Index] == HeightmapImportDescriptor.ImportResolutions[HeightmapImportDescriptorIndex])
-				{
-					FoundIndex = Index;
-					break;
-				}
-			}
+			UIImportLayer.ImportResult = LayerFileInfo.ResultCode;
 
-			if (FoundIndex == INDEX_NONE)
-			{
-				UIImportLayer.ImportResult = ELandscapeImportResult::Error;
-				UIImportLayer.ErrorMessage = NSLOCTEXT("LandscapeEditor.NewLandscape", "Import_WeightHeightResolutionMismatch", "Weightmap import resolution isn't same as Heightmap resolution.");
-				FMessageDialog::Open(EAppMsgType::Ok, UIImportLayer.ErrorMessage);
-				return ELandscapeImportResult::Error;
-			}
-
-			UIImportLayer.ImportResult = FLandscapeImportHelper::GetWeightmapImportData(UIImportLayer.ImportDescriptor, FoundIndex, ImportLayer.LayerName, ImportLayer.LayerData, UIImportLayer.ErrorMessage);
 			if (UIImportLayer.ImportResult == ELandscapeImportResult::Error)
 			{
 				FMessageDialog::Open(EAppMsgType::Ok, UIImportLayer.ErrorMessage);
 				return ELandscapeImportResult::Error;
 			}
+
+			LayerImage.Read(ImportLayer.LayerData);
 		}
 	}
 
diff --git a/Engine/Source/Editor/LandscapeEditor/Private/LandscapeEditorUtils.cpp b/Engine/Source/Editor/LandscapeEditor/Private/LandscapeEditorUtils.cpp
new file mode 100644
--- /dev/null
+++ b/Engine/Source/Editor/LandscapeEditor/Private/LandscapeEditorUtils.cpp
@@ -0,0 +1,90 @@
+// Copyright Epic Games, Inc. All Rights Reserved.
+
+#include "LandscapeEditorUtils.h"
+#include "LandscapeSettings.h"
+
+#include "LandscapeEditorModule.h"
+#include "LandscapeEdit.h"
+#include "LandscapeEdMode.h"
+#include "LandscapeEditorObject.h"
+#include "LandscapeTiledImage.h"
+#include "LandscapeStreamingProxy.h"
+
+#include "DesktopPlatformModule.h"
+#include "EditorModeManager.h"
+#include "EditorModes.h"
+
+#include "WorldPartition/WorldPartition.h"
+
+namespace LandscapeEditorUtils
+{
+	int32 GetMaxSizeInComponents()
+	{
+		const ULandscapeSettings* Settings = GetDefault<ULandscapeSettings>();
+		return Settings->MaxComponents;
+	}
+
+
+	TOptional<FString> GetImportExportFilename(const FString& InDialogTitle, const FString& InStartPath, const FString& InDialogTypeString)
+	{
+		FString Filename;
+		IDesktopPlatform* DesktopPlatform = FDesktopPlatformModule::Get();
+
+		if (DesktopPlatform == nullptr)
+		{
+			return TOptional<FString>();
+		}
+
+		TArray<FString> Filenames;
+		ILandscapeEditorModule& LandscapeEditorModule = FModuleManager::GetModuleChecked<ILandscapeEditorModule>("LandscapeEditor");
+		FEdModeLandscape* LandscapeEdMode = static_cast<FEdModeLandscape*>(GLevelEditorModeTools().GetActiveMode(FBuiltinEditorModes::EM_Landscape));
+
+		const bool bSuccess = DesktopPlatform->OpenFileDialog(
+			FSlateApplication::Get().FindBestParentWindowHandleForDialogs(nullptr),
+			InDialogTitle,
+			InStartPath,
+			TEXT(""),
+			InDialogTypeString,
+			EFileDialogFlags::None,
+			Filenames);
+
+		if (bSuccess)
+		{
+			FString TiledFileNamePattern;
+
+			if (FLandscapeTiledImage::CheckTiledNamePath(Filenames[0], TiledFileNamePattern) && FMessageDialog::Open(EAppMsgType::YesNo, FText::FromString(FString::Format(TEXT("Use '{0}' Tiled Image?"), { TiledFileNamePattern }))) == EAppReturnType::Yes)
+			{
+				Filename = TiledFileNamePattern;
+			}
+			else
+			{
+				Filename = Filenames[0];
+			}
+
+		}
+
+		return TOptional<FString>(Filename);
+
+	}
+
+	void SaveLandscapeProxies(TArrayView<ALandscapeProxy*> Proxies, UWorldPartition* WorldPartition)
+	{
+		// Save the Proxies a side effect of this is storing all the Saved Actors in the WP DirtyActors list
+		{
+			TRACE_CPUPROFILER_EVENT_SCOPE(SaveCreatedActors);
+			LandscapeEditorUtils::SaveObjects(Proxies);
+		}
+		
+		// Hold References to the Proxies we're going to save, this prevents WorldPartition Pinning these Actors on Tick
+		TArray< FWorldPartitionReference> References;
+		for (ALandscapeProxy* Proxy : Proxies)
+		{
+			References.Add(FWorldPartitionReference(WorldPartition, Proxy->GetActorGuid()));
+		}
+
+		// Clear the Dirty Actors in WP
+		WorldPartition->Tick(0.0f);
+	}
+}
+
+
diff --git a/Engine/Source/Editor/LandscapeEditor/Private/LandscapeFileFormatRaw.cpp b/Engine/Source/Editor/LandscapeEditor/Private/LandscapeFileFormatRaw.cpp
--- a/Engine/Source/Editor/LandscapeEditor/Private/LandscapeFileFormatRaw.cpp
+++ b/Engine/Source/Editor/LandscapeEditor/Private/LandscapeFileFormatRaw.cpp
@@ -13,37 +13,80 @@
 #include "Misc/FileHelper.h"
 #include "Templates/UnrealTemplate.h"
 
+#include "Serialization/JsonReader.h"
+#include "Serialization/JsonWriter.h"
+#include "Serialization/JsonSerializer.h"
+
 #define LOCTEXT_NAMESPACE "LandscapeEditor.NewLandscape"
 
-TArray<FLandscapeFileResolution> CalculatePossibleRawResolutions(int64 FileSize)
+namespace
 {
-	TArray<FLandscapeFileResolution> PossibleResolutions;
-
-	// Find all possible heightmap sizes, between 8 and 8192 width/height
-	const int32 MinWidth = FMath::Max(8, (int32)FMath::DivideAndRoundUp(FileSize, (int64)8192));
-	const int32 MaxWidth = FMath::TruncToInt(FMath::Sqrt(static_cast<double>(FileSize)));
-	for (int32 Width = MinWidth; Width <= MaxWidth; Width++)
+	void WriteMetaData(const TCHAR* InMapFilename, FLandscapeFileResolution InDataResolution, int32 InBitsPerPixel)
 	{
-		if (FileSize % Width == 0)
+		FString RawMetadataFilename = FPaths::SetExtension(InMapFilename, ".json");
+		FString JsonStr;
+		
+		TSharedRef<TJsonWriter<TCHAR>> JsonWriter = TJsonWriterFactory<TCHAR>::Create(&JsonStr);
+		TSharedRef<FJsonObject> JsonObject = MakeShareable(new FJsonObject);;
+		JsonObject->SetField("width", MakeShared<FJsonValueNumber>(InDataResolution.Width));
+		JsonObject->SetField("height", MakeShared<FJsonValueNumber>(InDataResolution.Height));
+		JsonObject->SetField("bpp", MakeShared<FJsonValueNumber>(InBitsPerPixel));
+
+		if (FJsonSerializer::Serialize<TCHAR>(JsonObject, JsonWriter))
 		{
-			FLandscapeFileResolution ImportResolution;
-			ImportResolution.Width = Width;
-			ImportResolution.Height = FileSize / Width;
-			PossibleResolutions.Add(ImportResolution);
+			FFileHelper::SaveStringToFile(JsonStr, *RawMetadataFilename);
 		}
 	}
+}
 
-	for (int32 i = PossibleResolutions.Num() - 1; i >= 0; --i)
+bool GetRawResolution(const TCHAR* InFilename, FLandscapeFileResolution& OutResolution, int32& OutBitsPerPixel)
+{
+	FString Extension = FPaths::GetExtension(InFilename);
+	FString RawMetadataFilename = FPaths::SetExtension(InFilename, ".json");
+	
+	if (FPaths::FileExists(RawMetadataFilename))
 	{
-		FLandscapeFileResolution ImportResolution = PossibleResolutions[i];
-		if (ImportResolution.Width != ImportResolution.Height)
+		FString JsonStr;
+		FFileHelper::LoadFileToString(JsonStr, *RawMetadataFilename);
+		TSharedRef<TJsonReader<TCHAR>> JsonReader = TJsonReaderFactory<TCHAR>::Create(JsonStr);
+		TSharedPtr<FJsonObject> JsonObject;
+		if (FJsonSerializer::Deserialize(JsonReader, JsonObject))
 		{
-			Swap(ImportResolution.Width, ImportResolution.Height);
-			PossibleResolutions.Add(ImportResolution);
+			OutResolution.Width = JsonObject->GetIntegerField("width");
+			OutResolution.Height = JsonObject->GetIntegerField("height");
+
+			OutBitsPerPixel = JsonObject->GetIntegerField("bpp");
+			
+			return true;
 		}
+		
+		return false;
+	}
+
+	const bool bIs16Bit = Extension == FString("r16");
+	const bool bIs8Bit = Extension == FString("r8");
+
+	if (!(bIs16Bit || bIs8Bit))
+	{
+		return false;
+	}
+
+	const int32 BytesPerPixel = bIs16Bit ? 2 : 1;
+	int64 FileSize = IFileManager::Get().FileSize(InFilename);
+	
+	const int32 NumPixels = FileSize / BytesPerPixel;
+	int32 Dimension = FMath::TruncToInt(FMath::Sqrt(static_cast<double>(NumPixels)));
+	if (Dimension * Dimension == NumPixels)
+	{
+		OutResolution.Width = Dimension;
+		OutResolution.Height = Dimension;
+
+		OutBitsPerPixel = BytesPerPixel * 8;
+
+		return true;
 	}
 
-	return PossibleResolutions;
+	return false;
 }
 
 FLandscapeHeightmapFileFormat_Raw::FLandscapeHeightmapFileFormat_Raw()
@@ -72,9 +115,22 @@ FLandscapeFileInfo FLandscapeHeightmapFileFormat_Raw::Validate(const TCHAR* Heig
 	}
 	else
 	{
-		Result.PossibleResolutions = CalculatePossibleRawResolutions(ImportFileSize / 2);
-
-		if (Result.PossibleResolutions.Num() == 0)
+		FLandscapeFileResolution Resolution;
+		int32 BitsPerPixel;
+		if (GetRawResolution(HeightmapFilename, Resolution, BitsPerPixel))
+		{
+			if (BitsPerPixel != 16)
+			{
+				Result.ResultCode = ELandscapeImportResult::Error;
+				Result.ErrorMessage = LOCTEXT("Import_WeightmapFileBitsPerPixel", "Height file has an invalid number of bits per pixel");
+			}
+			else
+			{
+				Result.PossibleResolutions = { Resolution };
+			}
+			
+		}
+		else
 		{
 			Result.ResultCode = ELandscapeImportResult::Error;
 			Result.ErrorMessage = LOCTEXT("Import_HeightmapFileInvalidSize", "The heightmap file has an invalid size (possibly not 16-bit?)");
@@ -116,7 +172,15 @@ void FLandscapeHeightmapFileFormat_Raw::Export(const TCHAR* HeightmapFilename, F
 	TempData.AddUninitialized(DataResolution.Width * DataResolution.Height * 2);
 	FMemory::Memcpy(TempData.GetData(), Data.GetData(), DataResolution.Width * DataResolution.Height * 2);
 
-	FFileHelper::SaveArrayToFile(TempData, HeightmapFilename);
+	if (FFileHelper::SaveArrayToFile(TempData, HeightmapFilename))
+	{
+		FString Extension = FPaths::GetExtension(HeightmapFilename);
+
+		if (Extension == "raw")
+		{
+			WriteMetaData(HeightmapFilename, DataResolution, 16);
+		}
+	}
 }
 
 //////////////////////////////////////////////////////////////////////////
@@ -142,9 +206,22 @@ FLandscapeFileInfo FLandscapeWeightmapFileFormat_Raw::Validate(const TCHAR* Weig
 	}
 	else
 	{
-		Result.PossibleResolutions = CalculatePossibleRawResolutions(ImportFileSize);
-
-		if (Result.PossibleResolutions.Num() == 0)
+		FLandscapeFileResolution Resolution;
+		int32 BitsPerPixel;
+		if (GetRawResolution(WeightmapFilename, Resolution, BitsPerPixel))
+		{
+			if (BitsPerPixel != 8)
+			{
+				Result.ResultCode = ELandscapeImportResult::Error;
+				Result.ErrorMessage = LOCTEXT("Import_WeightmapFileBitsPerPixel", "Weightmap file has an invalid number of bits per pixel");
+				
+			}
+			else
+			{
+				Result.PossibleResolutions = { Resolution };
+			}
+		}
+		else
 		{
 			Result.ResultCode = ELandscapeImportResult::Error;
 			Result.ErrorMessage = LOCTEXT("Import_WeightmapFileInvalidSize", "The layer file has an invalid size");
@@ -179,7 +256,15 @@ FLandscapeImportData<uint8> FLandscapeWeightmapFileFormat_Raw::Import(const TCHA
 
 void FLandscapeWeightmapFileFormat_Raw::Export(const TCHAR* WeightmapFilename, FName LayerName, TArrayView<const uint8> Data, FLandscapeFileResolution DataResolution, FVector Scale) const
 {
-	FFileHelper::SaveArrayToFile(Data, WeightmapFilename);
+	if (FFileHelper::SaveArrayToFile(Data, WeightmapFilename))
+	{
+		FString Extension = FPaths::GetExtension(WeightmapFilename);
+
+		if (Extension == "raw")
+		{
+			WriteMetaData(WeightmapFilename, DataResolution, 8);
+		}
+	}
 }
 
 #undef LOCTEXT_NAMESPACE
diff --git a/Engine/Source/Editor/LandscapeEditor/Private/LandscapeImageFileCache.cpp b/Engine/Source/Editor/LandscapeEditor/Private/LandscapeImageFileCache.cpp
new file mode 100644
--- /dev/null
+++ b/Engine/Source/Editor/LandscapeEditor/Private/LandscapeImageFileCache.cpp
@@ -0,0 +1,155 @@
+// Copyright Epic Games, Inc. All Rights Reserved.
+
+#include "LandscapeImageFileCache.h"
+#include "LandscapeSettings.h"
+
+#include "Modules/ModuleManager.h"
+#include "Internationalization/Internationalization.h"
+#include "IDirectoryWatcher.h"
+
+
+FLandscapeImageFileCache::FLandscapeImageFileCache()
+{
+	ULandscapeSettings* Settings = GetMutableDefault<ULandscapeSettings>();
+	SettingsChangedHandle = Settings->OnSettingChanged().AddRaw(this, &FLandscapeImageFileCache::OnLandscapeSettingsChanged);
+	MaxCacheSize = Settings->MaxImageImportCacheSizeMegaBytes * 1024U * 1024U;
+}
+
+FLandscapeImageFileCache::~FLandscapeImageFileCache()
+{
+	if (UObjectInitialized() && !GExitPurge)
+	{
+		GetMutableDefault<ULandscapeSettings>()->OnSettingChanged().Remove(SettingsChangedHandle);
+	}
+}
+
+void FLandscapeImageFileCache::MonitorCallback(const TArray<struct FFileChangeData>& Changes)
+{
+	for (const FFileChangeData& Change : Changes)
+	{
+		if (Change.Action == FFileChangeData::FCA_Modified || Change.Action == FFileChangeData::FCA_Removed)
+		{
+			Remove(Change.Filename);
+		}
+	}
+}
+
+bool FLandscapeImageFileCache::MonitorFile(const FString& Filename)
+{
+	FString Directory = FPaths::GetPath(Filename);
+
+	FDirectoryWatcherModule& DirectoryWatcherModule = FModuleManager::LoadModuleChecked<FDirectoryWatcherModule>("DirectoryWatcher");
+
+	if (FDirectoryMonitor* Monitor = MonitoredDirs.Find(Directory))
+	{
+		Monitor->NumFiles++;
+		return true;
+	}
+	else
+	{
+		FDelegateHandle Handle;
+		bool bWatcherResult = DirectoryWatcherModule.Get()->RegisterDirectoryChangedCallback_Handle(Directory, IDirectoryWatcher::FDirectoryChanged::CreateRaw(this, &FLandscapeImageFileCache::MonitorCallback), Handle);
+		if (bWatcherResult)
+		{
+			MonitoredDirs.Add(Directory, FDirectoryMonitor(Handle));
+		}
+		return bWatcherResult;
+	}
+}
+
+void FLandscapeImageFileCache::UnmonitorFile(const FString& Filename)
+{
+	FString Directory = FPaths::GetPath(Filename);
+
+	FDirectoryWatcherModule& DirectoryWatcherModule = FModuleManager::LoadModuleChecked<FDirectoryWatcherModule>("DirectoryWatcher");
+
+	if (FDirectoryMonitor* Monitor = MonitoredDirs.Find(Directory))
+	{
+		check(Monitor->NumFiles > 0);
+
+		Monitor->NumFiles--;
+		if (Monitor->NumFiles == 0)
+		{
+			DirectoryWatcherModule.Get()->UnregisterDirectoryChangedCallback_Handle(Directory, Monitor->MonitorHandle);
+		}
+
+		MonitoredDirs.Remove(Directory);
+	}
+}
+
+void FLandscapeImageFileCache::Add(const FString& Filename, FLandscapeImageDataRef NewImageData)
+{
+	check(CachedImages.Find(Filename) == nullptr);
+
+	CachedImages.Add(FString(Filename), FCacheEntry(NewImageData));
+	MonitorFile(Filename);
+	CacheSize += NewImageData.Data->Num();
+}
+
+void FLandscapeImageFileCache::Remove(const FString& Filename)
+{
+	if (FCacheEntry* CacheEntry = CachedImages.Find(Filename))
+	{
+		CacheSize -= CacheEntry->ImageData.Data->Num();
+		CachedImages.Remove(Filename);
+		UnmonitorFile(Filename);
+	}
+}
+
+void FLandscapeImageFileCache::OnLandscapeSettingsChanged(UObject* InObject, struct FPropertyChangedEvent& InPropertyChangedEvent)
+{
+	if (InPropertyChangedEvent.GetPropertyName() == "MaxImageImportCacheSizeMegaBytes")
+	{
+		ULandscapeSettings* LandscapeSettings = Cast<ULandscapeSettings>(InObject);
+		SetMaxSize(LandscapeSettings->MaxImageImportCacheSizeMegaBytes);
+	}
+}
+
+void FLandscapeImageFileCache::SetMaxSize(uint64 InNewMaxSize)
+{
+	if (MaxCacheSize != InNewMaxSize)
+	{
+		MaxCacheSize = InNewMaxSize * 1024U *1024U;
+		Trim();
+	}
+}
+
+void FLandscapeImageFileCache::Clear()
+{
+	CachedImages.Empty();
+	FDirectoryWatcherModule& DirectoryWatcherModule = FModuleManager::LoadModuleChecked<FDirectoryWatcherModule>("DirectoryWatcher");
+	
+	for (const TPair<FString, FDirectoryMonitor>& MonitoredDir : MonitoredDirs)
+	{
+		DirectoryWatcherModule.Get()->UnregisterDirectoryChangedCallback_Handle(MonitoredDir.Key, MonitoredDir.Value.MonitorHandle);
+	}
+	MonitoredDirs.Empty();
+}
+
+
+void FLandscapeImageFileCache::Trim()
+{
+	if (CacheSize < MaxCacheSize)
+	{
+		return;
+	}
+
+	CachedImages.ValueSort([](const FCacheEntry& LHS, const FCacheEntry& RHS) { return  LHS.UsageCount < RHS.UsageCount; });
+	TArray<FString> ToRemove;
+	uint64 Size = CacheSize;
+	for (const TPair<FString, FCacheEntry>& It : CachedImages)
+	{
+		ToRemove.Add(It.Key);
+		uint64 ImageSize = It.Value.ImageData.Data->Num();
+		Size -= ImageSize;
+		if (Size <= MaxCacheSize)
+		{
+			break;
+		}
+	}
+
+	for (const FString& Filename : ToRemove)
+	{
+		Remove(Filename);
+	}
+}
\ No newline at end of file
diff --git a/Engine/Source/Editor/LandscapeEditor/Private/LandscapeImageFileCache.h b/Engine/Source/Editor/LandscapeEditor/Private/LandscapeImageFileCache.h
new file mode 100644
--- /dev/null
+++ b/Engine/Source/Editor/LandscapeEditor/Private/LandscapeImageFileCache.h
@@ -0,0 +1,120 @@
+// Copyright Epic Games, Inc. All Rights Reserved.
+
+#pragma once
+
+#include "IImageWrapper.h"
+#include "IImageWrapperModule.h"
+#include "Internationalization/Internationalization.h"
+#include "Internationalization/Text.h"
+#include "Math/Vector.h"
+#include "Misc/FileHelper.h"
+#include "Modules/ModuleManager.h"
+#include "Templates/SharedPointer.h"
+#include "DirectoryWatcherModule.h"
+#include "LandscapeFileFormatInterface.h"
+#include "LandscapeEditorModule.h"
+
+#include "Delegates/IDelegateInstance.h"
+
+#define LOCTEXT_NAMESPACE "LandscapeEditor"
+
+struct FLandscapeImageDataRef
+{
+	TSharedPtr<TArray<uint8>> Data;
+	FIntPoint Resolution;
+};
+
+class FLandscapeImageFileCache
+{
+public:
+	FLandscapeImageFileCache();
+	~FLandscapeImageFileCache();
+
+
+	template<typename T>
+	FLandscapeFileInfo FindImage(const TCHAR* InImageFilename, FLandscapeImageDataRef& OutImageData)
+	{
+		FCacheEntry* CacheEntry = CachedImages.Find(InImageFilename);
+		FLandscapeFileInfo Result;
+
+		if (CacheEntry)
+		{
+			CacheEntry->UsageCount++;
+			OutImageData = CacheEntry->ImageData;
+			Result.PossibleResolutions.Add(FLandscapeFileResolution(OutImageData.Resolution.X, OutImageData.Resolution.Y));
+			return Result;
+		}
+
+		FLandscapeImageDataRef NewImageData;
+		ILandscapeEditorModule& LandscapeEditorModule = FModuleManager::GetModuleChecked<ILandscapeEditorModule>("LandscapeEditor");
+		const ILandscapeFileFormat<T>* FileFormat = LandscapeEditorModule.GetFormatByExtension<T>(*FPaths::GetExtension(InImageFilename, true));
+
+		const FLandscapeFileInfo FileInfo = FileFormat->Validate(InImageFilename);
+
+		if (FileInfo.ResultCode == ELandscapeImportResult::Success && FileInfo.PossibleResolutions.Num() > 0)
+		{
+			FLandscapeFileResolution ExpectedResolution = FileInfo.PossibleResolutions[0];
+			FLandscapeImportData<T> ImportData = FileFormat->Import(InImageFilename, ExpectedResolution);
+
+			const int32 BufferSize = ImportData.Data.Num() * sizeof(T);
+			NewImageData.Data = TSharedPtr<TArray<uint8>>(new TArray<uint8>());
+			NewImageData.Data->SetNumUninitialized(BufferSize);
+			FMemory::Memcpy(NewImageData.Data->GetData(), ImportData.Data.GetData(), BufferSize);
+			NewImageData.Resolution = FIntPoint(ExpectedResolution.Width, ExpectedResolution.Height);
+		}
+		else
+		{
+			return FileInfo;
+		}
+		
+		Trim();
+		OutImageData = NewImageData;
+		Add(FString(InImageFilename), OutImageData);
+		
+		Result.PossibleResolutions.Add(FLandscapeFileResolution(OutImageData.Resolution.X, OutImageData.Resolution.Y));
+		return Result;
+	}
+
+	void SetMaxSize(uint64 InNewMaxSize);
+
+	void Clear();
+
+private:
+
+	struct FCacheEntry
+	{
+		FCacheEntry(FLandscapeImageDataRef ImageData) : ImageData(ImageData) {}
+
+		uint32 UsageCount = 1;
+		FLandscapeImageDataRef ImageData;
+	};
+
+	struct FDirectoryMonitor
+	{
+		FDirectoryMonitor(FDelegateHandle Handle) { MonitorHandle = Handle; }
+		int32 NumFiles = 1;
+		FDelegateHandle MonitorHandle;
+	};
+	
+	void OnLandscapeSettingsChanged(UObject* InObject, struct FPropertyChangedEvent& InPropertyChangedEvent);
+	 
+	bool MonitorFile(const FString& Filename);
+	void UnmonitorFile(const FString& Filename);
+
+	void MonitorCallback(const TArray<struct FFileChangeData>& Changes);
+
+	TMap<FString, FCacheEntry> CachedImages;
+	
+	TMap<FString, FDirectoryMonitor> MonitoredDirs;
+
+	void Add(const FString& Filename, FLandscapeImageDataRef NewImageData);
+	void Remove(const FString& Filename);
+
+	void Trim();
+
+	uint64 MaxCacheSize = 32 * 1024 * 1024;
+	uint64 CacheSize = 0;
+	FDelegateHandle SettingsChangedHandle;
+};
+
+#undef LOCTEXT_NAMESPACE 
diff --git a/Engine/Source/Editor/LandscapeEditor/Private/LandscapeImportHelper.cpp b/Engine/Source/Editor/LandscapeEditor/Private/LandscapeImportHelper.cpp
--- a/Engine/Source/Editor/LandscapeEditor/Private/LandscapeImportHelper.cpp
+++ b/Engine/Source/Editor/LandscapeEditor/Private/LandscapeImportHelper.cpp
@@ -7,6 +7,7 @@
 #include "Modules/ModuleManager.h"
 #include "HAL/FileManager.h"
 #include "Misc/Paths.h"
+#include "LandscapeEditorUtils.h"
 
 #define LOCTEXT_NAMESPACE "LandscapeImportHelper"
 
@@ -299,7 +300,7 @@ void TransformImportDataInternal(const TArray<T>& InData, TArray<T>& OutData, co
 		const FIntRect DestRegion(0, 0, RequiredResolution.Width - 1, RequiredResolution.Height - 1);
 		FLandscapeConfigHelper::ResampleData<T>(InData, OutData, SrcRegion, DestRegion);
 	}
-	else if(TransformType != ELandscapeImportTransformType::None)
+	else if(TransformType == ELandscapeImportTransformType::ExpandCentered || TransformType == ELandscapeImportTransformType::ExpandOffset)
 	{
 		int32 OffsetX = 0;
 		int32 OffsetY = 0;
@@ -360,11 +361,12 @@ void FLandscapeImportHelper::ChooseBestComponentSizeForImport(int32 Width, int32
 	bool bValidQuadsPerSectionParam = false;
 	check(Width > 0 && Height > 0);
 		
+	const int32 MaxComponents = LandscapeEditorUtils::GetMaxSizeInComponents();
 	bool bFoundMatch = false;
 	// Try to find a section size and number of sections that exactly matches the dimensions of the heightfield
 	for (int32 SectionSizesIdx = UE_ARRAY_COUNT(FLandscapeConfig::SubsectionSizeQuadsValues) - 1; SectionSizesIdx >= 0; SectionSizesIdx--)
 	{
-		for (int32 NumSectionsIdx = UE_ARRAY_COUNT(FLandscapeConfig::NumSectionValues) - 1; NumSectionsIdx >= 0; NumSectionsIdx--)
+		for (int32 NumSectionsIdx = 0; NumSectionsIdx < UE_ARRAY_COUNT(FLandscapeConfig::NumSectionValues); NumSectionsIdx++)
 		{
 			int32 ss = FLandscapeConfig::SubsectionSizeQuadsValues[SectionSizesIdx];
 			int32 ns = FLandscapeConfig::NumSectionValues[NumSectionsIdx];
@@ -373,8 +375,8 @@ void FLandscapeImportHelper::ChooseBestComponentSizeForImport(int32 Width, int32
 			bValidSubsectionSizeParam |= (InOutSectionsPerComponent == ns);
 			bValidQuadsPerSectionParam |= (InOutQuadsPerSection == ss);
 
-			if (((Width - 1) % (ss * ns)) == 0 && ((Width - 1) / (ss * ns)) <= 32 &&
-				((Height - 1) % (ss * ns)) == 0 && ((Height - 1) / (ss * ns)) <= 32)
+			if (((Width - 1) % (ss * ns)) == 0 && ((Width - 1) / (ss * ns)) <= MaxComponents &&
+				((Height - 1) % (ss * ns)) == 0 && ((Height - 1) / (ss * ns)) <= MaxComponents)
 			{
 				bFoundMatch = true;
 				InOutQuadsPerSection = ss;
diff --git a/Engine/Source/Editor/LandscapeEditor/Private/LandscapeRegionUtils.cpp b/Engine/Source/Editor/LandscapeEditor/Private/LandscapeRegionUtils.cpp
new file mode 100644
--- /dev/null
+++ b/Engine/Source/Editor/LandscapeEditor/Private/LandscapeRegionUtils.cpp
@@ -0,0 +1,148 @@
+// Copyright Epic Games, Inc. All Rights Reserved.
+
+#include "LandscapeRegionUtils.h"
+
+#include "Landscape.h"
+
+#include "WorldPartition/WorldPartition.h"
+
+#include "SourceControlHelpers.h"
+#include "UObject/SavePackage.h"
+#include "Builders/CubeBuilder.h"
+#include "ActorFactories/ActorFactory.h"
+#include "LandscapeProxy.h"
+#include "LocationVolume.h"
+#include "LandscapeStreamingProxy.h"
+#include "LandscapeEditorUtils.h"
+#include "LandscapeEditorPrivate.h"
+
+namespace LandscapeRegionUtils
+{
+
+ALocationVolume* CreateLandscapeRegionVolume(UWorld* InWorld, ALandscapeProxy* InParentLandscapeActor, const FIntPoint& InRegionCoordinate, double InRegionSize)
+{
+	const double Shrink = 0.95;
+
+	FVector ParentLocation = InParentLandscapeActor->GetActorLocation();
+	FVector Location = ParentLocation + FVector(InRegionCoordinate.X * InRegionSize, InRegionCoordinate.Y * InRegionSize, 0.0) + FVector(InRegionSize / 2.0, InRegionSize / 2.0, 0.0);
+	FRotator Rotation;
+	FActorSpawnParameters SpawnParameters;
+
+	const FString Label = FString::Printf(TEXT("LandscapeRegion_%i_%i"), InRegionCoordinate.X, InRegionCoordinate.Y);
+	SpawnParameters.Name = FName(*Label);
+
+	ALocationVolume* LocationVolume = InWorld->SpawnActor<ALocationVolume>(Location, Rotation, SpawnParameters);
+	LocationVolume->SetActorLabel(Label);
+
+	LocationVolume->AttachToActor(InParentLandscapeActor, FAttachmentTransformRules::KeepWorldTransform);
+	FVector Scale{ InRegionSize * Shrink,  InRegionSize * Shrink, InRegionSize * 0.5f };
+	LocationVolume->SetActorScale3D(Scale);
+
+	// Specify a cube shape for the LocationVolume
+	UCubeBuilder* Builder = NewObject<UCubeBuilder>();
+	Builder->X = 1.0f;
+	Builder->Y = 1.0f;
+	Builder->Z = 1.0f;
+	UActorFactory::CreateBrushForVolumeActor(LocationVolume, Builder);
+
+	LocationVolume->Tags.Add(FName("LandscapeRegion"));
+	return LocationVolume;
+}
+
+void ForEachComponentByRegion(int32 RegionSize, const TArray<FIntPoint>& ComponentCoordinates, TFunctionRef<bool(const FIntPoint&, const TArray<FIntPoint>&)> RegionFn)
+{
+	if (RegionSize <= 0)
+	{
+		UE_LOG(LogLandscapeTools, Error, TEXT("RegionSize must be greater than zero"));
+		return;
+	}
+
+	TMap<FIntPoint, TArray< FIntPoint>> Regions;
+	for (FIntPoint ComponentCoordinate : ComponentCoordinates)
+	{
+		FIntPoint RegionCoordinate;
+		RegionCoordinate.X = (ComponentCoordinate.X) / RegionSize;
+		RegionCoordinate.Y = (ComponentCoordinate.Y) / RegionSize;
+
+		if (TArray<FIntPoint>* ComponentIndices = Regions.Find(RegionCoordinate))
+		{
+			ComponentIndices->Add(ComponentCoordinate);
+		}
+		else
+		{
+			Regions.Add(RegionCoordinate, TArray<FIntPoint> {ComponentCoordinate});
+		}
+	}
+
+	auto Sorter = [](const FIntPoint& A, const FIntPoint& B)
+	{
+		if (A.Y == B.Y)
+			return A.X < B.X;
+
+		return A.Y < B.Y;
+	};
+
+	Regions.KeySort(Sorter);
+
+	for (const TPair< FIntPoint, TArray< FIntPoint>>& RegionIt : Regions)
+	{
+		if (!RegionFn(RegionIt.Key, RegionIt.Value))
+		{
+			break;
+		}
+	}
+}
+
+void ForEachRegion_LoadProcessUnload(ULandscapeInfo* InLandscapeInfo, const FIntRect& InDomain, UWorld* InWorld, TFunctionRef<bool(const FBox&, const TArray<ALandscapeProxy*>)> InRegionFn)
+{
+	const int32 RegionSizeInTexels = InLandscapeInfo->ComponentSizeQuads * InLandscapeInfo->RegionSizeInComponents + 1;
+	const FIntRect RegionCoordinates = InDomain / RegionSizeInTexels;
+
+	UWorldPartition* WorldPartition = InWorld->GetWorldPartition();
+	check(WorldPartition != nullptr);
+
+	TArray<AActor*> Children;
+	InLandscapeInfo->LandscapeActor->GetAttachedActors(Children);
+	TArray<ALocationVolume*> LandscapeRegions;
+
+	for (AActor* Child : Children)
+	{
+		if (Child->IsA<ALocationVolume>())
+		{
+			LandscapeRegions.Add(Cast<ALocationVolume>(Child));
+		}
+	}
+
+	for (ALocationVolume* Region : LandscapeRegions)
+	{
+		Region->Load();
+
+		FBox RegionBounds = Region->GetComponentsBoundingBox();
+
+		TArray<AActor*> AllActors = InWorld->GetLevel(0)->Actors;
+		TArray<ALandscapeProxy*> LandscapeProxies;
+		for (AActor* Actor : AllActors)
+		{
+			if (ALandscapeStreamingProxy* Proxy = Cast<ALandscapeStreamingProxy>(Actor))
+			{
+				LandscapeProxies.Add(Proxy);
+			}
+		}
+
+		// Save the actor
+		bool bShouldExit = !InRegionFn(RegionBounds, LandscapeProxies);
+
+		InLandscapeInfo->ForceLayersFullUpdate();
+
+		LandscapeEditorUtils::SaveLandscapeProxies(MakeArrayView(LandscapeProxies), WorldPartition);
+
+		Region->Unload();
+
+		if (bShouldExit)
+		{
+			break;
+		}
+	}
+}
+
+} // LandscapeEditor
\ No newline at end of file
diff --git a/Engine/Source/Editor/LandscapeEditor/Private/LandscapeRegionUtils.h b/Engine/Source/Editor/LandscapeEditor/Private/LandscapeRegionUtils.h
new file mode 100644
--- /dev/null
+++ b/Engine/Source/Editor/LandscapeEditor/Private/LandscapeRegionUtils.h
@@ -0,0 +1,27 @@
+// Copyright Epic Games, Inc. All Rights Reserved.
+
+#pragma once
+
+#include "Math/MathFwd.h"
+#include "Containers/Array.h"
+#include "Templates/Function.h"
+
+class AActor;
+class UWorld;
+class ALandscapeProxy;
+class ULandscapeInfo;
+class ALandscapeStreamingProxy;
+class ALandscapeProxy;
+class ALocationVolume;
+
+namespace LandscapeRegionUtils
+{
+	// Create a LocationVolume for a Region located at RegionCoordinate with a Box shape of dimension RegionSize
+	ALocationVolume* CreateLandscapeRegionVolume(UWorld* World, ALandscapeProxy* ParentLandscapeActor, const FIntPoint& RegionCoordinate, double RegionSize);
+
+	// Iterate over all Components grouping by Region
+	void ForEachComponentByRegion(int32 RegionSize, const TArray<FIntPoint>& ComponentCoordinates, TFunctionRef<bool(const FIntPoint&, const TArray<FIntPoint>&)> RegionFn);
+
+	// Load, Process (call RegionFn) and Unload the Region
+	void ForEachRegion_LoadProcessUnload(ULandscapeInfo* LandscapeInfo, const FIntRect& Domain, UWorld* World, TFunctionRef<bool(const FBox&, const TArray<ALandscapeProxy*>)> RegionFn);
+}
diff --git a/Engine/Source/Editor/LandscapeEditor/Private/LandscapeTiledImage.cpp b/Engine/Source/Editor/LandscapeEditor/Private/LandscapeTiledImage.cpp
new file mode 100644
--- /dev/null
+++ b/Engine/Source/Editor/LandscapeEditor/Private/LandscapeTiledImage.cpp
@@ -0,0 +1,213 @@
+// Copyright Epic Games, Inc. All Rights Reserved.
+
+#include "LandscapeTiledImage.h"
+
+#include "HAL/FileManager.h"
+#include "Internationalization/Regex.h"
+#include "LandscapeImageFileCache.h"
+
+#include "Internationalization/Internationalization.h"
+
+#define LOCTEXT_NAMESPACE "LandscapeEditor"
+
+
+TMap<FString, FString> FLandscapeTiledImage::Tokens = {
+{ TEXT("u"), TEXT("<u>") },
+{ TEXT("v"), TEXT("<v>") },
+{ TEXT("x"), TEXT("<x>") },
+{ TEXT("y"), TEXT("<y>") }
+};
+
+FString FLandscapeTiledImage::GetTokenRegex(const FString& Prefix)
+{
+	return Prefix + TEXT("(-?[0-9]+)");
+}
+
+
+FLandscapeTiledImage::FLandscapeTiledImage()
+{
+}
+
+FLandscapeFileInfo FLandscapeTiledImage::Load(const TCHAR* Filename)
+{
+	FLandscapeFileInfo Result;
+
+	TArray<FString> FoundFiles;
+	FindFiles(Filename, FoundFiles);
+
+	FString Path = FPaths::GetPath(Filename);
+	FString RegexFilename(Filename);
+	TMap<int32, FString> Groups;
+
+	for (const TPair<FString, FString>& It : Tokens)
+	{
+		if (int32 StartIndex = RegexFilename.Find(It.Value); StartIndex != INDEX_NONE)
+		{
+			Groups.Add(StartIndex,  It.Value);
+		}
+	}
+
+	Groups.KeySort([](int32 A, int32 B) { return A < B; });
+
+	for (const TPair<FString, FString>& It : Tokens)
+	{
+		if (int32 start = RegexFilename.Find(It.Value); start != INDEX_NONE)
+		{
+			RegexFilename.RemoveAt(start, It.Value.Len());
+			RegexFilename.InsertAt(start, GetTokenRegex(""));
+		}
+	}
+
+	const bool bExactFilename = RegexFilename == Filename;
+
+	SizeInTiles = FIntPoint::NoneValue;
+	TileResolution = FIntPoint::NoneValue;
+	
+	if (FoundFiles.Num() == 1 && bExactFilename)
+	{
+		FIntPoint TileCoord(0, 0);
+		TileFilenames.Add(TileCoord, FString(Filename));
+	}
+	else
+	{
+		for (const FString& FoundFilname : FoundFiles)
+		{
+			FRegexPattern Pattern(RegexFilename);
+			FString FullPath = FPaths::Combine(Path, FoundFilname);
+			FRegexMatcher Matcher(Pattern, FullPath);
+
+			if (Matcher.FindNext())
+			{
+				int32 X = -1;
+				int32 Y = -1;
+
+				int32 Group = 1;
+				for (const TPair<int32, FString> &it : Groups)
+				{
+					if (it.Value == "<u>")
+					{
+						X = FCString::Atoi(*Matcher.GetCaptureGroup(Group++)) - 1;
+					}
+					else if (it.Value == "<x>")
+					{
+						X = FCString::Atoi(*Matcher.GetCaptureGroup(Group++));
+					}
+					else if (it.Value == "<v>")
+					{
+						Y = FCString::Atoi(*Matcher.GetCaptureGroup(Group++)) - 1;
+					}
+					else if (it.Value == "<y>")
+					{
+						Y = FCString::Atoi(*Matcher.GetCaptureGroup(Group++));
+					}
+				}
+
+				if (X >= 0 && Y >= 0)
+				{
+					FIntPoint TileCoord(X, Y);
+					TileFilenames.Add(TileCoord, FullPath);
+				}			
+			}
+		}
+	}
+	
+	for (const TPair<FIntPoint, FString> & Tile : TileFilenames)
+	{
+		FLandscapeImageDataRef ImageData;
+		FLandscapeImageFileCache& LandscapeImageFileCache = FModuleManager::GetModuleChecked<ILandscapeEditorModule>("LandscapeEditor").GetImageFileCache();
+		if (FLandscapeFileInfo TileResult16 = LandscapeImageFileCache.FindImage<uint16>(*Tile.Value, ImageData); TileResult16.ResultCode != ELandscapeImportResult::Success)
+		{
+			if (FLandscapeFileInfo TileResult8 = LandscapeImageFileCache.FindImage<uint8>(*Tile.Value, ImageData);  TileResult8.ResultCode != ELandscapeImportResult::Success)
+			{
+				return TileResult16;
+			}
+		}
+
+
+		SizeInTiles.X = FMath::Max(SizeInTiles.X, Tile.Key.X + 1);
+		SizeInTiles.Y = FMath::Max(SizeInTiles.Y, Tile.Key.Y + 1);
+
+		uint32 Width = ImageData.Resolution.X;
+		uint32 Height = ImageData.Resolution.Y;
+
+		if (TileResolution == FIntPoint::NoneValue)
+		{
+			TileResolution.X = Width;
+			TileResolution.Y = Height;
+		}
+		else if (Width != TileResolution.X && Height != TileResolution.Y)
+		{
+			Result.ResultCode = ELandscapeImportResult::Error;
+			Result.ErrorMessage = LOCTEXT("LandscapeFileReadError", "Mimatched resolution found in tiled image");
+			return Result;
+		}
+	}
+
+	if (TileFilenames.Num() == 0)
+	{
+		Result.ResultCode = ELandscapeImportResult::Error;
+		Result.ErrorMessage = LOCTEXT("LandscapeFileReadError", "No files found");
+		return Result;
+	}
+
+	Result.PossibleResolutions.Add(FLandscapeFileResolution(GetResolution().X, GetResolution().Y));
+
+	return Result;
+}
+
+void FLandscapeTiledImage::FindFiles(const TCHAR* FilenamePattern, TArray<FString>& OutPaths)
+{
+	FString GlobFilename(FilenamePattern);
+
+	for (const TPair<FString, FString>& It : Tokens)
+	{
+		GlobFilename = GlobFilename.Replace(*It.Value, TEXT("*"));
+	}
+	
+	OutPaths.Empty();
+	IFileManager::Get().FindFiles(OutPaths, *GlobFilename, true, false);
+}
+
+bool FLandscapeTiledImage::CheckTiledNamePath(const FString& Filename, FString& OutTiledFilenamePattern)
+{
+	const FString Extension = FPaths::GetExtension(Filename);
+	const FString Root = FPaths::GetPath(Filename);
+	const FString BaseFilename = FPaths::GetBaseFilename(Filename);
+
+	bool HasMatch = true;
+	FString CurrentFilename = BaseFilename;
+	int32 TokenIndex = 0;
+
+	for (const TPair<FString, FString>& It : Tokens)
+	{
+		while (true)
+		{
+			FRegexMatcher Matcher(FRegexPattern(GetTokenRegex(It.Key)), CurrentFilename);
+			if (!Matcher.FindNext())
+			{
+				break;
+			}
+
+			int32 MatchBegin = Matcher.GetMatchBeginning();
+			int32 MatchEnd = Matcher.GetMatchEnding();
+
+			CurrentFilename.RemoveAt(MatchBegin, MatchEnd - MatchBegin);
+			CurrentFilename.InsertAt(MatchBegin, It.Key + It.Value);
+		}
+	}
+
+	const FString PatternFilename = CurrentFilename + FString(".") + Extension;
+	FString TiledFilenamePattern = FPaths::Combine(Root, PatternFilename);
+
+	TArray<FString> FoundFiles;
+	FLandscapeTiledImage::FindFiles(*TiledFilenamePattern, FoundFiles);
+	const bool bFoundTiledFiles = FoundFiles.Num() > 0 && TiledFilenamePattern != Filename;
+
+	if (bFoundTiledFiles )
+	{
+		OutTiledFilenamePattern = TiledFilenamePattern;
+	}
+	return bFoundTiledFiles;
+}
+
+#undef LOCTEXT_NAMESPACE
\ No newline at end of file
diff --git a/Engine/Source/Editor/LandscapeEditor/Private/LandscapeTiledImage.h b/Engine/Source/Editor/LandscapeEditor/Private/LandscapeTiledImage.h
new file mode 100644
--- /dev/null
+++ b/Engine/Source/Editor/LandscapeEditor/Private/LandscapeTiledImage.h
@@ -0,0 +1,114 @@
+// Copyright Epic Games, Inc. All Rights Reserved.
+
+#pragma once 
+
+#include "Internationalization/Internationalization.h"
+#include "Internationalization/Text.h"
+#include "LandscapeImageFileCache.h"
+#include "LandscapeFileFormatInterface.h"
+#include "LandscapeDataAccess.h"
+
+template<typename T>
+class LandscapeImageTraits
+{
+public:
+	static T DefaultValue() { return 0; }
+};
+
+template<>
+class LandscapeImageTraits<uint16>
+{
+public:
+	static uint16 DefaultValue() { return LandscapeDataAccess::MidValue; }
+};
+
+class FLandscapeTiledImage
+{
+public:
+	FLandscapeTiledImage();
+	FLandscapeFileInfo Load(const TCHAR* Filename);
+	
+	FIntPoint GetTileResolution() const { return TileResolution;  }
+	FIntPoint GetSizeInTiles() const { return SizeInTiles; }
+
+	FIntPoint GetResolution() const {
+		return FIntPoint(TileResolution.X * SizeInTiles.X, TileResolution.Y * SizeInTiles.Y); }
+
+	static void FindFiles(const TCHAR* FilenamePattern, TArray<FString>& OutPaths);
+
+	template<typename T>
+	void Read(TArray<T>& OutImageData, bool bFlipYAxis = false) const
+	{
+		const FIntPoint Resolution = GetResolution();
+		ReadRegion(FIntRect(0, 0, Resolution.X, Resolution.Y), OutImageData, bFlipYAxis);
+	}
+
+	static bool CheckTiledNamePath(const FString& Filename, FString& OutTiledFilenamePattern);
+
+	template<typename T>
+	FLandscapeFileInfo ReadRegion(const FIntRect& Region, TArray<T>& OutImageData, bool bFlipYAxis = false, T DefaultValue = LandscapeImageTraits<T>::DefaultValue()) const
+	{
+		OutImageData.Init(DefaultValue, Region.Height() * Region.Width());
+		
+		int32 MinTileX = Region.Min.X / TileResolution.X;
+		int32 MinTileY = Region.Min.Y / TileResolution.Y;
+
+		int32 MaxTileX = Region.Max.X / TileResolution.X;
+		int32 MaxTileY = Region.Max.Y / TileResolution.Y;
+		
+		for (int32 TileY = MinTileY; TileY <= MaxTileY; ++TileY)
+		{
+			for (int32 TileX = MinTileX; TileX <= MaxTileX; ++TileX)
+			{
+				const FIntPoint TileMin(TileX * TileResolution.X, TileY * TileResolution.Y);
+
+				const FString* TileFilename = TileFilenames.Find(FIntPoint(TileX, TileY));			
+				if (TileFilename)
+				{
+					FLandscapeImageDataRef ImageData;
+					FLandscapeImageFileCache& LandscapeImageFileCache = FModuleManager::GetModuleChecked<ILandscapeEditorModule>("LandscapeEditor").GetImageFileCache();
+					FLandscapeFileInfo TileResult = LandscapeImageFileCache.FindImage<T>(**TileFilename, ImageData);
+
+					if (TileResult.ResultCode != ELandscapeImportResult::Success)
+					{
+						return TileResult;
+					}
+
+					const T* TileData = (const T* ) ImageData.Data->GetData();
+
+					FIntRect TileRegion = Region;
+					TileRegion -= TileMin;
+					TileRegion.Clip(FIntRect(0, 0, TileResolution.X, TileResolution.Y));
+
+					for (int32 Y = 0; Y < TileRegion.Height(); ++Y)
+					{
+						for (int32 X = 0; X < TileRegion.Width(); ++X)
+						{
+							const FIntPoint TileCoord = FIntPoint(X, Y) + TileRegion.Min;
+							const FIntPoint OutCoord = TileCoord + TileMin - Region.Min;
+							
+							const int32 SrcY = bFlipYAxis ? TileResolution.Y - TileCoord.Y - 1 : TileCoord.Y;
+							//const int32 DstY = bFlipYAxis ? Region.Height() - OutCoord.Y - 1 : OutCoord.Y;
+
+							OutImageData[OutCoord.Y * Region.Width() + OutCoord.X] = TileData[SrcY * TileResolution.X + TileCoord.X];
+						}
+					}
+				}
+			}
+		}
+
+		return FLandscapeFileInfo();
+	}
+
+private:
+
+	static TMap<FString, FString> Tokens;
+	static FString GetTokenRegex(const FString& Prefix);
+
+	TMap<FIntPoint, FString> TileFilenames;
+
+	FIntPoint TileResolution = FIntPoint::NoneValue;
+	FIntPoint SizeInTiles = FIntPoint::NoneValue;
+};
+
+
diff --git a/Engine/Source/Editor/LandscapeEditor/Private/NewLandscapeUtils.cpp b/Engine/Source/Editor/LandscapeEditor/Private/NewLandscapeUtils.cpp
--- a/Engine/Source/Editor/LandscapeEditor/Private/NewLandscapeUtils.cpp
+++ b/Engine/Source/Editor/LandscapeEditor/Private/NewLandscapeUtils.cpp
@@ -9,6 +9,7 @@
 #include "LandscapeEditorUtils.h"
 #include "LandscapeEdMode.h"
 #include "LandscapeConfigHelper.h"
+#include "LandscapeSettings.h"
 
 #include "Misc/MessageDialog.h"
 #include "Modules/ModuleManager.h"
diff --git a/Engine/Source/Editor/LandscapeEditor/Public/LandscapeEditorModule.h b/Engine/Source/Editor/LandscapeEditor/Public/LandscapeEditorModule.h
--- a/Engine/Source/Editor/LandscapeEditor/Public/LandscapeEditorModule.h
+++ b/Engine/Source/Editor/LandscapeEditor/Public/LandscapeEditorModule.h
@@ -15,6 +15,7 @@ using ILandscapeHeightmapFileFormat = ILandscapeFileFormat<uint16>;
 using ILandscapeWeightmapFileFormat = ILandscapeFileFormat<uint8>;
 
 class FUICommandList;
+class FLandscapeImageFileCache;
 
 /**
  * LandscapeEditor module interface
@@ -51,4 +52,8 @@ public:
 	}
 
 	virtual TSharedPtr<FUICommandList> GetLandscapeLevelViewportCommandList() const = 0;
+
+	virtual FLandscapeImageFileCache& GetImageFileCache() const = 0;
+
+
 };
diff --git a/Engine/Source/Editor/LandscapeEditor/Public/LandscapeEditorObject.h b/Engine/Source/Editor/LandscapeEditor/Public/LandscapeEditorObject.h
--- a/Engine/Source/Editor/LandscapeEditor/Public/LandscapeEditorObject.h
+++ b/Engine/Source/Editor/LandscapeEditor/Public/LandscapeEditorObject.h
@@ -270,6 +270,18 @@ struct FLandscapePatternBrushWorldSpaceSettings
 	{}
 };
 
+UENUM()
+enum class ELandscapeImportExportMode
+{
+	// Import and export only loaded landscape proxies
+	LoadedOnly, 
+	// Import and export the whole world loading & unloading regions as required.
+	All,		
+	// Import and export regions selected in the WP editor loading & unloading as required.
+	// Selected, // TODO: don.boogert   
+};
+
+
 UCLASS(MinimalAPI)
 class ULandscapeEditorObject : public UObject
 {
@@ -573,6 +585,9 @@ class ULandscapeEditorObject : public UObject
 	UPROPERTY(Category = "Import / Export", EditAnywhere, NonTransactional, meta = (DisplayName = "Export Single File", ShowForTools = "ImportExport", ToolTip = "(World Partition only) When true, exports the landscape as a single file, if false exports each grid tile individually."))
 	bool bExportSingleFile = false;
 
+	UPROPERTY(Category = "Import / Export", EditAnywhere, NonTransactional, meta = (DisplayName = "Mode", ShowForTools = "ImportExport", ToolTip = "Import Loaded or All Landscape Regions"))
+	ELandscapeImportExportMode ImportExportMode = ELandscapeImportExportMode::LoadedOnly;
+
 	UPROPERTY(NonTransactional)
 	FLandscapeImportDescriptor HeightmapImportDescriptor;
 	
@@ -591,6 +606,9 @@ public:
 	UPROPERTY(Category = "New Landscape", EditAnywhere, NonTransactional, meta = (DisplayName= "World Partition Grid Size", ToolTip="Number of components per landscape streaming proxies per axis", ShowForTools="NewLandscape", ClampMin=1, ClampMax=16, UIMin=1, UIMax=16))
 	uint32 WorldPartitionGridSize = 2;
 
+	UPROPERTY(Category = "New Landscape", EditAnywhere, NonTransactional, meta = (DisplayName= "World Partition Region Size", ToolTip="Number of components per Landscape World Partition Region per axis.", ShowForTools="NewLandscape", ClampMin=4, ClampMax=64, UIMin=4, UIMax=64))
+	uint32 WorldPartitionRegionSize = 16;
+
 	// Whether the imported alpha maps are to be interpreted as "layered" or "additive" (UE uses additive internally)
 	UPROPERTY(Category="New Landscape", EditAnywhere, NonTransactional, meta=(DisplayName="Layer Alphamap Type", ShowForTools="NewLandscape,ImportExport"))
 	ELandscapeImportAlphamapType ImportLandscape_AlphamapType;
@@ -733,8 +751,8 @@ public:
 
 	int32 ClampLandscapeSize(int32 InComponentsCount) const
 	{
-		// Max size is either whole components below 8192 verts, or 32 components
-		return FMath::Clamp(InComponentsCount, 1, FMath::Min(32, FMath::FloorToInt(8191.0f / (float)(NewLandscape_SectionsPerComponent * NewLandscape_QuadsPerSection))));
+		// Max size is either whole components below 8192 verts, or 256 components
+		return FMath::Clamp(InComponentsCount, 1, FMath::Min(256, FMath::FloorToInt(8191.0f / (float)(NewLandscape_SectionsPerComponent * NewLandscape_QuadsPerSection))));
 	}
 	
 	int32 CalcComponentsCount(int32 InResolution) const
@@ -743,7 +761,7 @@ public:
 	}
 
 	void NewLandscape_ClampSize()
-	{
+	{	
 		NewLandscape_ComponentCount.X = ClampLandscapeSize(NewLandscape_ComponentCount.X);
 		NewLandscape_ComponentCount.Y = ClampLandscapeSize(NewLandscape_ComponentCount.Y);
 	}
diff --git a/Engine/Source/Editor/LandscapeEditor/Public/LandscapeEditorUtils.h b/Engine/Source/Editor/LandscapeEditor/Public/LandscapeEditorUtils.h
--- a/Engine/Source/Editor/LandscapeEditor/Public/LandscapeEditorUtils.h
+++ b/Engine/Source/Editor/LandscapeEditor/Public/LandscapeEditorUtils.h
@@ -3,12 +3,28 @@
 #pragma once
 
 #include "CoreMinimal.h"
+#include "FileHelpers.h"
 
 class ALandscapeProxy;
+class ALandscapeStreamingProxy;
 class ULandscapeLayerInfoObject;
+class UWorldPartition;
 
 namespace LandscapeEditorUtils
 {
 	bool LANDSCAPEEDITOR_API SetHeightmapData(ALandscapeProxy* Landscape, const TArray<uint16>& Data);
 	bool LANDSCAPEEDITOR_API SetWeightmapData(ALandscapeProxy* Landscape, ULandscapeLayerInfoObject* LayerObject, const TArray<uint8>& Data);
+
+	int32 GetMaxSizeInComponents();
+	TOptional<FString> GetImportExportFilename(const FString& InDialogTitle, const FString& InStartPath, const FString& InDialogTypeString);
+
+	template<typename T>
+	void SaveObjects(TArrayView<T*> InObjects)
+	{
+		TArray<UPackage*> Packages;
+		Algo::Transform(InObjects, Packages, [](UObject* InObject) { return InObject->GetPackage(); });
+		UEditorLoadingAndSavingUtils::SavePackages(Packages, /* bOnlyDirty = */ false);
+	}
+
+	void SaveLandscapeProxies(TArrayView<ALandscapeProxy*> Proxies, UWorldPartition* WorldPartition);
 }
diff --git a/Engine/Source/Editor/LandscapeEditor/Public/LandscapeImportHelper.h b/Engine/Source/Editor/LandscapeEditor/Public/LandscapeImportHelper.h
--- a/Engine/Source/Editor/LandscapeEditor/Public/LandscapeImportHelper.h
+++ b/Engine/Source/Editor/LandscapeEditor/Public/LandscapeImportHelper.h
@@ -50,7 +50,8 @@ enum class ELandscapeImportTransformType : int8
 	None UMETA(DisplayName="Original", ToolTip="Will Import the data at the gizmo location in the original size"),
 	ExpandOffset UMETA(DisplayName="Expand", ToolTip="Will Import the data at the gizmo location and expand the data to fill the landscape") ,
 	ExpandCentered UMETA(Hidden), 
-	Resample UMETA(ToolTip="Will resample Import data to fit landscape")
+	Resample UMETA(ToolTip="Will resample Import data to fit landscape"),
+	Subregion UMETA(ToolTop="Import Sub-region of the Image to Landscape")
 };
 
 FORCEINLINE bool operator==(const FLandscapeImportResolution& Lhs, const FLandscapeImportResolution& Rhs)
diff --git a/Engine/Source/Runtime/Landscape/Classes/LandscapeInfo.h b/Engine/Source/Runtime/Landscape/Classes/LandscapeInfo.h
--- a/Engine/Source/Runtime/Landscape/Classes/LandscapeInfo.h
+++ b/Engine/Source/Runtime/Landscape/Classes/LandscapeInfo.h
@@ -144,6 +144,9 @@ class ULandscapeInfo : public UObject
 #if WITH_EDITORONLY_DATA
 	UPROPERTY()
 	TArray<FLandscapeInfoLayerSettings> Layers;
+
+	UPROPERTY()
+	int32 RegionSizeInComponents;
 #endif // WITH_EDITORONLY_DATA
 
 public:
@@ -193,10 +196,10 @@ public:
 	virtual void Serialize(FArchive& Ar) override;
 	//~ End UObject Interface
 
-	FBox GetLoadedBounds() const;
+	LANDSCAPE_API FBox GetLoadedBounds() const;
 
 #if WITH_EDITOR
-	FBox GetCompleteBounds() const;
+	LANDSCAPE_API FBox GetCompleteBounds() const;
 #endif
 
 #if WITH_EDITOR
@@ -359,8 +362,19 @@ public:
 	 *     // Code
 	 * });
 	 */
+	UE_DEPRECATED(5.3, "This function has been deprecated, please use the ForEachLandscapeProxy property instead")
 	LANDSCAPE_API void ForAllLandscapeProxies(TFunctionRef<void(ALandscapeProxy*)> Fn) const;
 
+	/**
+	 * Runs the given function on the root landscape actor and all streaming proxies, with the posibility of early exit
+	 * Most easily used with a lambda as follows:
+	 * ForEachLandscapeProxy([](ALandscapeProxy* Proxy) -> bool
+	 * {
+	 *     return continueLoop ? true : false;
+	 * });
+	 */
+	LANDSCAPE_API void ForEachLandscapeProxy(TFunctionRef<bool(ALandscapeProxy*)> Fn) const;
+
 	/** Associates passed actor with this info object
  *  @param	Proxy		Landscape actor to register
  *  @param  bMapCheck	Whether to warn about landscape errors
diff --git a/Engine/Source/Runtime/Landscape/Private/Landscape.cpp b/Engine/Source/Runtime/Landscape/Private/Landscape.cpp
--- a/Engine/Source/Runtime/Landscape/Private/Landscape.cpp
+++ b/Engine/Source/Runtime/Landscape/Private/Landscape.cpp
@@ -3082,7 +3082,7 @@ FLandscapeEditorLayerSettings& ULandscapeInfo::GetLayerEditorSettings(ULandscape
 
 void ULandscapeInfo::CreateLayerEditorSettingsFor(ULandscapeLayerInfoObject* LayerInfo)
 {
-	ForAllLandscapeProxies([LayerInfo](ALandscapeProxy* Proxy)
+	ForEachLandscapeProxy([LayerInfo](ALandscapeProxy* Proxy)
 	{
 		FLandscapeEditorLayerSettings* EditorLayerSettings = Proxy->EditorLayerSettings.FindByKey(LayerInfo);
 		if (!EditorLayerSettings)
@@ -3090,6 +3090,7 @@ void ULandscapeInfo::CreateLayerEditorSettingsFor(ULandscapeLayerInfoObject* Lay
 			Proxy->Modify();
 			Proxy->EditorLayerSettings.Add(FLandscapeEditorLayerSettings(LayerInfo));
 		}
+		return true;
 	});
 }
 
@@ -3305,13 +3306,14 @@ bool ULandscapeInfo::UpdateLayerInfoMapInternal(ALandscapeProxy* Proxy, bool bIn
 
 			if (!bInvalidate)
 			{
-				ForAllLandscapeProxies([this](ALandscapeProxy* EachProxy)
+				ForEachLandscapeProxy([this](ALandscapeProxy* EachProxy)
 				{
 					if (!EachProxy->IsPendingKillPending())
 					{
 						checkSlow(EachProxy->GetLandscapeInfo() == this);
 						UpdateLayerInfoMapInternal(EachProxy, false);
 					}
+					return true;
 				});
 			}
 		}
@@ -3998,13 +4000,16 @@ void ULandscapeInfo::ModifyObject(UObject* InObject, bool bAlwaysMarkDirty)
 
 ALandscapeProxy* ULandscapeInfo::GetLandscapeProxyForLevel(ULevel* Level) const
 {
+	TRACE_CPUPROFILER_EVENT_SCOPE(ULandscapeInfo::GetLandscapeProxyForLevel);
 	ALandscapeProxy* LandscapeProxy = nullptr;
-	ForAllLandscapeProxies([&LandscapeProxy, Level](ALandscapeProxy* Proxy)
+	ForEachLandscapeProxy([&LandscapeProxy, Level](ALandscapeProxy* Proxy) -> bool 
 	{
 		if (Proxy->GetLevel() == Level)
 		{
 			LandscapeProxy = Proxy;
+			return false;
 		}
+		return true;
 	});
 	return LandscapeProxy;
 }
@@ -4013,8 +4018,10 @@ ALandscapeProxy* ULandscapeInfo::GetLandscapeProxyForLevel(ULevel* Level) const
 
 ALandscapeProxy* ULandscapeInfo::GetCurrentLevelLandscapeProxy(bool bRegistered) const
 {
+	TRACE_CPUPROFILER_EVENT_SCOPE(ULandscapeInfo::GetCurrentLevelLandscapeProxy);
+
 	ALandscapeProxy* LandscapeProxy = nullptr;
-	ForAllLandscapeProxies([&LandscapeProxy, bRegistered](ALandscapeProxy* Proxy)
+	ForEachLandscapeProxy([&LandscapeProxy, bRegistered](ALandscapeProxy* Proxy) -> bool
 	{
 		if (!bRegistered || Proxy->GetRootComponent()->IsRegistered())
 		{
@@ -4023,8 +4030,10 @@ ALandscapeProxy* ULandscapeInfo::GetCurrentLevelLandscapeProxy(bool bRegistered)
 				ProxyWorld->GetCurrentLevel() == Proxy->GetOuter())
 			{
 				LandscapeProxy = Proxy;
+				return false;
 			}
 		}
+		return true;
 	});
 	return LandscapeProxy;
 }
@@ -4146,12 +4155,13 @@ void ULandscapeInfo::FixupProxiesTransform(bool bDirty)
 
 void ULandscapeInfo::UpdateComponentLayerAllowList()
 {
-	ForAllLandscapeProxies([](ALandscapeProxy* Proxy)
+	ForEachLandscapeProxy([](ALandscapeProxy* Proxy)
 	{
 		for (ULandscapeComponent* Comp : Proxy->LandscapeComponents)
 		{
 			Comp->UpdateLayerAllowListFromPaintedLayers();
 		}
+		return true;
 	});
 }
 
@@ -4276,6 +4286,28 @@ void ULandscapeInfo::ForAllLandscapeProxies(TFunctionRef<void(ALandscapeProxy*)>
 	}
 }
 
+void ULandscapeInfo::ForEachLandscapeProxy(TFunctionRef<bool(ALandscapeProxy*)> Fn) const
+{
+	if (ALandscape* Landscape = LandscapeActor.Get())
+	{
+		if (!Fn(Landscape))
+		{
+			return;
+		}
+	}
+
+	for (TWeakObjectPtr<ALandscapeStreamingProxy> StreamingProxyPtr : StreamingProxies)
+	{
+		if (ALandscapeProxy* LandscapeProxy = StreamingProxyPtr.Get())
+		{
+			if (!Fn(LandscapeProxy))
+			{
+				return;
+			}
+		}
+	}
+}
+
 bool ULandscapeInfo::IsRegistered(const ALandscapeProxy* Proxy) const
 {
 	if (Proxy == nullptr)
diff --git a/Engine/Source/Runtime/Landscape/Private/LandscapeCollision.cpp b/Engine/Source/Runtime/Landscape/Private/LandscapeCollision.cpp
--- a/Engine/Source/Runtime/Landscape/Private/LandscapeCollision.cpp
+++ b/Engine/Source/Runtime/Landscape/Private/LandscapeCollision.cpp
@@ -1803,6 +1803,8 @@ bool ULandscapeHeightfieldCollisionComponent::RecreateCollision()
 #if WITH_EDITORONLY_DATA
 void ULandscapeHeightfieldCollisionComponent::SnapFoliageInstances()
 {
+	TRACE_CPUPROFILER_EVENT_SCOPE(ULandscapeHeightfieldCollisionComponent::SnapFoliageInstances);
+
 	SnapFoliageInstances(FBox(FVector(-WORLD_MAX), FVector(WORLD_MAX)));
 }
 
@@ -2241,6 +2243,8 @@ void ULandscapeHeightfieldCollisionComponent::PreSave(FObjectPreSaveContext Obje
 #if WITH_EDITOR
 void ULandscapeInfo::UpdateAllAddCollisions()
 {
+
+	TRACE_CPUPROFILER_EVENT_SCOPE(ULandscapeInfo::UpdateAllAddCollisions);
 	XYtoAddCollisionMap.Reset();
 
 	// Don't recreate add collisions if the landscape is not registered. This can happen during Undo.
diff --git a/Engine/Source/Runtime/Landscape/Private/LandscapeConfigHelper.cpp b/Engine/Source/Runtime/Landscape/Private/LandscapeConfigHelper.cpp
--- a/Engine/Source/Runtime/Landscape/Private/LandscapeConfigHelper.cpp
+++ b/Engine/Source/Runtime/Landscape/Private/LandscapeConfigHelper.cpp
@@ -354,14 +354,15 @@ void FLandscapeConfigHelper::MoveSplinesToLandscape(ULandscapeInfo* InLandscapeI
 		InLandscape->CreateSplineComponent();
 	}
 
-	InLandscapeInfo->ForAllLandscapeProxies([InLandscapeInfo, InLandscape](ALandscapeProxy* LandscapeProxy)
+	InLandscapeInfo->ForEachLandscapeProxy([InLandscapeInfo, InLandscape](ALandscapeProxy* LandscapeProxy)
 	{
 		if (InLandscape == LandscapeProxy || !LandscapeProxy->GetSplinesComponent())
 		{
-			return;
+			return true;
 		}
 
 		InLandscapeInfo->MoveSplinesToProxy(LandscapeProxy->GetSplinesComponent(), InLandscape);
+		return true;
 	});
 }
 
diff --git a/Engine/Source/Runtime/Landscape/Private/LandscapeEdit.cpp b/Engine/Source/Runtime/Landscape/Private/LandscapeEdit.cpp
--- a/Engine/Source/Runtime/Landscape/Private/LandscapeEdit.cpp
+++ b/Engine/Source/Runtime/Landscape/Private/LandscapeEdit.cpp
@@ -1128,6 +1128,8 @@ void ULandscapeComponent::ClearDirtyCollisionHeightData()
 
 void ULandscapeComponent::UpdateCollisionHeightData(const FColor* const HeightmapTextureMipData, const FColor* const SimpleCollisionHeightmapTextureData, int32 ComponentX1/*=0*/, int32 ComponentY1/*=0*/, int32 ComponentX2/*=MAX_int32*/, int32 ComponentY2/*=MAX_int32*/, bool bUpdateBounds/*=false*/, const FColor* XYOffsetTextureMipData/*=nullptr*/, bool bInUpdateHeightfieldRegion/*=true*/)
 {
+	TRACE_CPUPROFILER_EVENT_SCOPE(ULandscapeComponent::UpdateCollisionHeightData);
+	
 	ALandscapeProxy* Proxy = GetLandscapeProxy();
 	FIntPoint ComponentKey = GetSectionBase() / ComponentSizeQuads;
 	ULandscapeHeightfieldCollisionComponent* CollisionComp = GetCollisionComponent();
@@ -1917,6 +1919,8 @@ void ULandscapeComponent::UpdateCollisionPhysicalMaterialData(TArray<UPhysicalMa
 
 void ULandscapeComponent::GenerateHeightmapMips(TArray<FColor*>& HeightmapTextureMipData, int32 ComponentX1/*=0*/, int32 ComponentY1/*=0*/, int32 ComponentX2/*=MAX_int32*/, int32 ComponentY2/*=MAX_int32*/, FLandscapeTextureDataInfo* TextureDataInfo/*=nullptr*/)
 {
+	TRACE_CPUPROFILER_EVENT_SCOPE(ULandscapeComponent::GenerateHeightmapMips);
+	
 	bool EndX = false;
 	bool EndY = false;
 
@@ -2526,36 +2530,44 @@ FIntRect ULandscapeComponent::GetComponentExtent() const
 bool ULandscapeInfo::SupportsLandscapeEditing() const
 {
 	bool bSupportsEditing = true;
-	ForAllLandscapeProxies([&bSupportsEditing](ALandscapeProxy* Proxy)
+	ForEachLandscapeProxy([&bSupportsEditing](ALandscapeProxy* Proxy)
 	{
 		if(Proxy->GetOutermost()->bIsCookedForEditor)
 		{
 			bSupportsEditing = false;
+			return false;
 		}
+		return true;
+
 	});
 	return bSupportsEditing;
 }
 
 bool ULandscapeInfo::AreAllComponentsRegistered() const
 {
-	const TArray<ALandscapeProxy*>& LandscapeProxies = ALandscapeProxy::GetLandscapeProxies();
-	for(ALandscapeProxy* LandscapeProxy : LandscapeProxies)
-	{
+	bool bAllRegistered = true;
+	ForEachLandscapeProxy([&bAllRegistered, this](ALandscapeProxy* LandscapeProxy) {
 		if (!IsValid(LandscapeProxy))
 		{
-			continue;
+			return true;
 		}
 
 		if (LandscapeProxy->GetLandscapeGuid() == LandscapeGuid)
-		{
+		{	
 			for (ULandscapeComponent* LandscapeComponent : LandscapeProxy->LandscapeComponents)
 			{
 				if (LandscapeComponent && !LandscapeComponent->IsRegistered())
 				{
-					return false;
+					bAllRegistered = false;
 				}
 			}
 		}
+		return true;
+	});
+
+	if (!bAllRegistered)
+	{
+		return false;
 	}
 		
 	for (TScriptInterface<ILandscapeSplineInterface> SplineOwner : SplineActors)
@@ -2746,6 +2758,8 @@ ENGINE_API extern bool GDisableAutomaticTextureMaterialUpdateDependencies;
 LANDSCAPE_API void ALandscapeProxy::Import(const FGuid& InGuid, int32 InMinX, int32 InMinY, int32 InMaxX, int32 InMaxY, int32 InNumSubsections, int32 InSubsectionSizeQuads, const TMap<FGuid, TArray<uint16>>& InImportHeightData, 
 										   const TCHAR* const InHeightmapFileName, const TMap<FGuid, TArray<FLandscapeImportLayerInfo>>& InImportMaterialLayerInfos, ELandscapeImportAlphamapType InImportMaterialLayerType, const TArray<FLandscapeLayer>* InImportLayers)
 {
+	TRACE_CPUPROFILER_EVENT_SCOPE(ALandscapeProxy::Import);
+	
 	check(InGuid.IsValid());
 	check(InImportHeightData.Num() == InImportMaterialLayerInfos.Num());
 
@@ -3323,6 +3337,8 @@ LANDSCAPE_API void ALandscapeProxy::Import(const FGuid& InGuid, int32 InMinX, in
 	SlowTask.EnterProgressFrame(1.0f);
 
 	{
+		TRACE_CPUPROFILER_EVENT_SCOPE(ALandscapeProxy::Import - Update Materials);
+		
 		// We disable automatic material update context, to manage it manually
 		GDisableAutomaticTextureMaterialUpdateDependencies = true;
 	
@@ -4035,7 +4051,7 @@ LANDSCAPE_API bool ULandscapeInfo::GetLandscapeXYComponentBounds(FIntRect& OutXY
 
 LANDSCAPE_API void ULandscapeInfo::ForAllLandscapeComponents(TFunctionRef<void(ULandscapeComponent*)> Fn) const
 {
-	ForAllLandscapeProxies([&](ALandscapeProxy* Proxy)
+	ForEachLandscapeProxy([&](ALandscapeProxy* Proxy)
 	{
 		for (ULandscapeComponent* Component : Proxy->LandscapeComponents)
 		{
@@ -4044,6 +4060,7 @@ LANDSCAPE_API void ULandscapeInfo::ForAllLandscapeComponents(TFunctionRef<void(U
 				Fn(Component);
 			}
 		}
+		return true;
 	});
 }
 
@@ -4245,7 +4262,7 @@ void ULandscapeInfo::DeleteLayer(ULandscapeLayerInfoObject* LayerInfo, const FNa
 		}
 	}
 
-	ForAllLandscapeProxies([LayerInfo](ALandscapeProxy* Proxy)
+	ForEachLandscapeProxy([LayerInfo](ALandscapeProxy* Proxy)
 	{
 		Proxy->Modify();
 		int32 Index = Proxy->EditorLayerSettings.IndexOfByKey(LayerInfo);
@@ -4253,6 +4270,7 @@ void ULandscapeInfo::DeleteLayer(ULandscapeLayerInfoObject* LayerInfo, const FNa
 		{
 			Proxy->EditorLayerSettings.RemoveAt(Index);
 		}
+		return true;
 	});
 
 	//UpdateLayerInfoMap();
@@ -4279,7 +4297,7 @@ void ULandscapeInfo::ReplaceLayer(ULandscapeLayerInfoObject* FromLayerInfo, ULan
 			}
 		}
 
-		ForAllLandscapeProxies([FromLayerInfo, ToLayerInfo](ALandscapeProxy* Proxy)
+		ForEachLandscapeProxy([FromLayerInfo, ToLayerInfo](ALandscapeProxy* Proxy)
 		{
 			Proxy->Modify();
 			FLandscapeEditorLayerSettings* ToEditorLayerSettings = Proxy->EditorLayerSettings.FindByKey(ToLayerInfo);
@@ -4306,6 +4324,7 @@ void ULandscapeInfo::ReplaceLayer(ULandscapeLayerInfoObject* FromLayerInfo, ULan
 					Proxy->EditorLayerSettings.Add(FLandscapeEditorLayerSettings(ToLayerInfo));
 				}
 			}
+			return true;
 		});
 
 		//UpdateLayerInfoMap();
@@ -4317,7 +4336,7 @@ void ULandscapeInfo::ReplaceLayer(ULandscapeLayerInfoObject* FromLayerInfo, ULan
 void ULandscapeInfo::GetUsedPaintLayers(const FGuid& InLayerGuid, TArray<ULandscapeLayerInfoObject*>& OutUsedLayerInfos) const
 {
 	OutUsedLayerInfos.Empty();
-	ForAllLandscapeProxies([&](ALandscapeProxy* Proxy)
+	ForEachLandscapeProxy([&](ALandscapeProxy* Proxy)
 	{
 		for (ULandscapeComponent* Component : Proxy->LandscapeComponents)
 		{
@@ -4326,6 +4345,8 @@ void ULandscapeInfo::GetUsedPaintLayers(const FGuid& InLayerGuid, TArray<ULandsc
 				Component->GetUsedPaintLayers(InLayerGuid, OutUsedLayerInfos);
 			}
 		}
+
+		return true;
 	});
 }
 
@@ -4778,17 +4799,19 @@ void ALandscapeProxy::RecreateCollisionComponents()
 
 void ULandscapeInfo::RecreateCollisionComponents()
 {
-	ForAllLandscapeProxies([](ALandscapeProxy* Proxy)
+	ForEachLandscapeProxy([](ALandscapeProxy* Proxy)
 	{
 		Proxy->RecreateCollisionComponents();
+		return true;
 	});
 }
 
 void ULandscapeInfo::RemoveXYOffsets()
 {
-	ForAllLandscapeProxies([](ALandscapeProxy* Proxy)
+	ForEachLandscapeProxy([](ALandscapeProxy* Proxy)
 	{
 		Proxy->RemoveXYOffsets();
+		return true;
 	});
 }
 
@@ -4814,9 +4837,10 @@ void ULandscapeInfo::ClearDirtyData()
 
 void ULandscapeInfo::UpdateAllComponentMaterialInstances(bool bInInvalidateCombinationMaterials)
 {
-	ForAllLandscapeProxies([=](ALandscapeProxy* Proxy)
+	ForEachLandscapeProxy([=](ALandscapeProxy* Proxy)
 	{
 		Proxy->UpdateAllComponentMaterialInstances(bInInvalidateCombinationMaterials);
+		return true;
 	});
 }
 
@@ -4890,6 +4914,8 @@ ALandscapeProxy* ULandscapeInfo::MoveComponentsToLevel(const TArray<ULandscapeCo
 
 ALandscapeProxy* ULandscapeInfo::MoveComponentsToProxy(const TArray<ULandscapeComponent*>& InComponents, ALandscapeProxy* LandscapeProxy, bool bSetPositionAndOffset, ULevel* TargetLevel)
 {
+	TRACE_CPUPROFILER_EVENT_SCOPE(ULandscapeInfo::MoveComponentsToProxy);
+	
 	ALandscape* Landscape = LandscapeActor.Get();
 	check(Landscape != nullptr);
 	
@@ -5090,6 +5116,8 @@ ALandscapeProxy* ULandscapeInfo::MoveComponentsToProxy(const TArray<ULandscapeCo
 
 void ALandscape::SplitHeightmap(ULandscapeComponent* Comp, ALandscapeProxy* TargetProxy, FMaterialUpdateContext* InOutUpdateContext, TArray<FComponentRecreateRenderStateContext>* InOutRecreateRenderStateContext, bool InReregisterComponent)
 {
+	TRACE_CPUPROFILER_EVENT_SCOPE(ALandscape::SplitHeightmap);
+	
 	ULandscapeInfo* Info = Comp->GetLandscapeInfo();
 
 	// Make sure the heightmap UVs are powers of two.
diff --git a/Engine/Source/Runtime/Landscape/Private/LandscapeEditInterface.cpp b/Engine/Source/Runtime/Landscape/Private/LandscapeEditInterface.cpp
--- a/Engine/Source/Runtime/Landscape/Private/LandscapeEditInterface.cpp
+++ b/Engine/Source/Runtime/Landscape/Private/LandscapeEditInterface.cpp
@@ -261,7 +261,7 @@ bool FLandscapeEditDataInterface::GetComponentsInRegion(int32 X1, int32 Y1, int3
 void FLandscapeEditDataInterface::SetHeightData(int32 X1, int32 Y1, int32 X2, int32 Y2, const uint16* InData, int32 InStride, bool InCalcNormals, const uint16* InNormalData, const uint16* InHeightAlphaBlendData, const uint8* InHeightFlagsData, bool InCreateComponents, UTexture2D* InHeightmap, UTexture2D* InXYOffsetmapTexture,
 											   bool InUpdateBounds, bool InUpdateCollision, bool InGenerateMips)
 {
-	TRACE_CPUPROFILER_EVENT_SCOPE(LandscapeEditDataInterface_SetHeightData);
+	TRACE_CPUPROFILER_EVENT_SCOPE(FLandscapeEditDataInterface::SetHeightData);
 
 	const int32 NumVertsX = 1 + X2 - X1;
 	const int32 NumVertsY = 1 + Y2 - Y1;
diff --git a/Engine/Source/Runtime/Landscape/Private/LandscapeEditLayers.cpp b/Engine/Source/Runtime/Landscape/Private/LandscapeEditLayers.cpp
--- a/Engine/Source/Runtime/Landscape/Private/LandscapeEditLayers.cpp
+++ b/Engine/Source/Runtime/Landscape/Private/LandscapeEditLayers.cpp
@@ -33,6 +33,7 @@ LandscapeEditLayers.cpp: Landscape editing layers mode
 #include "Rendering/Texture2DResource.h"
 #include "SceneView.h"
 #include "MaterialCachedData.h"
+#include "ContentStreaming.h"
 
 #if WITH_EDITOR
 #include "LandscapeEditorModule.h"
@@ -191,7 +192,7 @@ void OnLandscapeEditLayersLocalMergeChanged(IConsoleVariable* CVar)
 	}
 }
 
-int32 LandscapeEditLayersLocalMerge = 0;
+int32 LandscapeEditLayersLocalMerge = 1;
 static FAutoConsoleVariableRef CVarLandscapeEditLayersLocalMerge(
 	TEXT("landscape.EditLayersLocalMerge.Enable"),
 	LandscapeEditLayersLocalMerge,
@@ -2054,7 +2055,7 @@ void ALandscape::ReleaseLayersRenderingResource()
 		return;
 	}
 
-	Info->ForAllLandscapeProxies([&](ALandscapeProxy* Proxy)
+	Info->ForEachLandscapeProxy([&](ALandscapeProxy* Proxy)
 	{
 		for (auto& ItPair : Proxy->HeightmapsCPUReadback)
 		{
@@ -2069,6 +2070,7 @@ void ALandscape::ReleaseLayersRenderingResource()
 			delete WeightmapCPUReadback;
 		}
 		Proxy->WeightmapsCPUReadback.Empty();
+		return true;
 	});
 
 	if (CombinedLayersWeightmapAllMaterialLayersResource != nullptr)
@@ -2120,7 +2122,7 @@ FIntPoint ALandscape::ComputeComponentCounts() const
 	FIntPoint MaxSectionBase(TNumericLimits<int32>::Min(), TNumericLimits<int32>::Min());
 	FIntPoint MinSectionBase(TNumericLimits<int32>::Max(), TNumericLimits<int32>::Max());
 
-	Info->ForAllLandscapeProxies([&MaxSectionBase, &MinSectionBase](ALandscapeProxy* Proxy)
+	Info->ForEachLandscapeProxy([&MaxSectionBase, &MinSectionBase](ALandscapeProxy* Proxy)
 	{
 		for (ULandscapeComponent* Component : Proxy->LandscapeComponents)
 		{
@@ -2130,6 +2132,7 @@ FIntPoint ALandscape::ComputeComponentCounts() const
 			MinSectionBase.X = FMath::Min(MinSectionBase.X, Component->SectionBaseX);
 			MinSectionBase.Y = FMath::Min(MinSectionBase.Y, Component->SectionBaseY);
 		}
+		return true;
 	});
 
 	if ((MaxSectionBase.X >= MinSectionBase.X) && (MaxSectionBase.Y >= MinSectionBase.Y))
@@ -2150,9 +2153,10 @@ void ALandscape::CopyOldDataToDefaultLayer()
 		return;
 	}
 
-	Info->ForAllLandscapeProxies([&](ALandscapeProxy* Proxy)
+	Info->ForEachLandscapeProxy([&](ALandscapeProxy* Proxy)
 	{
 		CopyOldDataToDefaultLayer(Proxy);
+		return true;
 	});
 }
 
@@ -2317,9 +2321,10 @@ void ALandscape::UpdateProxyLayersWeightmapUsage()
 		return;
 	}
 
-	Info->ForAllLandscapeProxies([&](ALandscapeProxy* Proxy)
+	Info->ForEachLandscapeProxy([&](ALandscapeProxy* Proxy)
 	{
 		Proxy->UpdateProxyLayersWeightmapUsage();
+		return true;
 	});
 }
 
@@ -2348,9 +2353,10 @@ void ALandscape::InitializeLandscapeLayersWeightmapUsage()
 		return;
 	}
 
-	Info->ForAllLandscapeProxies([&](ALandscapeProxy* Proxy)
+	Info->ForEachLandscapeProxy([&](ALandscapeProxy* Proxy)
 	{
 		Proxy->InitializeProxyLayersWeightmapUsage();
+		return true;
 	});
 }
 
@@ -2464,9 +2470,10 @@ void ALandscape::ValidateProxyLayersWeightmapUsage() const
 		return;
 	}
 
-	Info->ForAllLandscapeProxies([=](const ALandscapeProxy* Proxy)
+	Info->ForEachLandscapeProxy([=](const ALandscapeProxy* Proxy)
 	{
 		Proxy->ValidateProxyLayersWeightmapUsage();
+		return true;
 	});
 }
 
@@ -3481,7 +3488,7 @@ bool ALandscape::PrepareTextureResources(bool bInWaitForStreaming)
 	TSet<ULandscapeComponent*> StreamedInHeightmapComponents;
 
 	bool bIsReady = true;
-	Info->ForAllLandscapeProxies([&](ALandscapeProxy* Proxy)
+	Info->ForEachLandscapeProxy([&](ALandscapeProxy* Proxy)
 	{
 		for (ULandscapeComponent* Component : Proxy->LandscapeComponents)
 		{
@@ -3524,6 +3531,7 @@ bool ALandscape::PrepareTextureResources(bool bInWaitForStreaming)
 				bIsReady &= bIsTextureReady;
 			}
 		}
+		return true;
 	});
 
 	// The assets that were streaming in before and are not anymore can be considered streamed in: 
@@ -3623,7 +3631,7 @@ bool ALandscape::PrepareLayersTextureResources(const TArray<FLandscapeLayer>& In
 	}
 
 	bool bIsReady = true;
-	Info->ForAllLandscapeProxies([&](ALandscapeProxy* Proxy)
+	Info->ForEachLandscapeProxy([&](ALandscapeProxy* Proxy)
 	{
 		for (const FLandscapeLayer& Layer : InLayers)
 		{
@@ -3640,6 +3648,7 @@ bool ALandscape::PrepareLayersTextureResources(const TArray<FLandscapeLayer>& In
 				}
 			}
 		}
+		return true;
 	});
 
 	return bIsReady;
@@ -4401,8 +4410,8 @@ namespace EditLayersHeightmapLocalMerge_RenderThread
 			const FLandscapeRDGTrackedTexture* TrackedTexture = InTrackedTextures.Find(TextureResolveInfo.Texture);
 			check((TrackedTexture != nullptr) && (TrackedTexture->ScratchTextureRef != nullptr));
 
-			check(TextureResolveInfo.NumMips == TrackedTexture->ScratchTextureRef->Desc.NumMips);
-			check(TrackedTexture->ScratchTextureMipsSRVRefs.Num() == TextureResolveInfo.NumMips);
+			//check(TextureResolveInfo.NumMips == TrackedTexture->ScratchTextureRef->Desc.NumMips);
+			//check(TrackedTexture->ScratchTextureMipsSRVRefs.Num() == TextureResolveInfo.NumMips);
 
 			FIntPoint CurrentMipSubregionSize = InLocalMergeInfo.ComponentSizeVerts;
 			for (int32 MipLevel = 1; MipLevel < TextureResolveInfo.NumMips; ++MipLevel)
@@ -5870,7 +5879,7 @@ void ALandscape::ReallocateLayersWeightmaps(FUpdateLayersContentContext& InUpdat
 	//GDisableAutomaticTextureMaterialUpdateDependencies = false;
 
 	// Clean-up unused weightmap CPUReadback resources
-	Info->ForAllLandscapeProxies([](ALandscapeProxy* Proxy)
+	Info->ForEachLandscapeProxy([](ALandscapeProxy* Proxy)
 	{
 		TArray<UTexture2D*, TInlineAllocator<64>> EntriesToRemoveFromMap;
 		for (auto& Pair : Proxy->WeightmapsCPUReadback)
@@ -5906,6 +5915,8 @@ void ALandscape::ReallocateLayersWeightmaps(FUpdateLayersContentContext& InUpdat
 				}
 			}
 		}
+
+		return true;
 	});
 
 	ValidateProxyLayersWeightmapUsage();
@@ -7714,7 +7725,7 @@ void ALandscape::RequestLayersContentUpdateForceAll(ELandscapeLayerUpdateMode In
 	const bool bUpdateClientUdpateEditing = (InModeMask & ELandscapeLayerUpdateMode::Update_Client_Editing) != 0;
 	if (ULandscapeInfo* LandscapeInfo = GetLandscapeInfo())
 	{
-		LandscapeInfo->ForAllLandscapeProxies([bUpdateHeightmap, bUpdateWeightmap, bUpdateAllHeightmap, bUpdateAllWeightmap, bUpdateHeightCollision, bUpdateWeightCollision, bUpdateClientUdpateEditing, bInUserTriggered](ALandscapeProxy* Proxy)
+		LandscapeInfo->ForEachLandscapeProxy([bUpdateHeightmap, bUpdateWeightmap, bUpdateAllHeightmap, bUpdateAllWeightmap, bUpdateHeightCollision, bUpdateWeightCollision, bUpdateClientUdpateEditing, bInUserTriggered](ALandscapeProxy* Proxy)
 		{
 			if (Proxy)
 			{
@@ -7736,6 +7747,7 @@ void ALandscape::RequestLayersContentUpdateForceAll(ELandscapeLayerUpdateMode In
 					}
 				}
 			}
+			return true;
 		});
 	}
 
@@ -8427,6 +8439,14 @@ void ALandscape::ForceUpdateLayersContent(bool bIntermediateRender)
 
 void ALandscape::ForceLayersFullUpdate()
 {
+	TRACE_CPUPROFILER_EVENT_SCOPE(ALandscape::ForceLayersFullUpdate);
+
+	FAssetCompilingManager::Get().FinishAllCompilation();
+
+	FStreamingManagerCollection& StreamingManagers = IStreamingManager::Get();
+	StreamingManagers.UpdateResourceStreaming(GetWorld()->GetDeltaSeconds(), /* bProcessEverything */ true);
+	StreamingManagers.BlockTillAllRequestsFinished();
+
 	RequestSplineLayerUpdate();
 	RequestLayersContentUpdateForceAll();
 	ForceUpdateLayersContent(/* bIntermediateRender */ false);
@@ -8912,9 +8932,10 @@ void ALandscape::DeleteLayer(int32 InLayerIndex)
 	FGuid LayerGuid = Layer->Guid;
 
 	// Clean up Weightmap usage in LandscapeProxies
-	LandscapeInfo->ForAllLandscapeProxies([&LayerGuid](ALandscapeProxy* Proxy)
+	LandscapeInfo->ForEachLandscapeProxy([&LayerGuid](ALandscapeProxy* Proxy)
 	{
 		Proxy->DeleteLayer(LayerGuid);
+		return true;
 	});
 
 	const FLandscapeLayer* SplinesReservedLayer = GetLandscapeSplinesReservedLayer();
@@ -9037,13 +9058,14 @@ void ALandscape::ClearPaintLayer(const FGuid& InLayerGuid, ULandscapeLayerInfoOb
 	FScopedSetLandscapeEditingLayer Scope(this, InLayerGuid, [=] { RequestLayersContentUpdate(ELandscapeLayerUpdateMode::Update_Weightmap_All); });
 
 	FLandscapeEditDataInterface LandscapeEdit(LandscapeInfo);
-	LandscapeInfo->ForAllLandscapeProxies([&](ALandscapeProxy* Proxy)
+	LandscapeInfo->ForEachLandscapeProxy([&](ALandscapeProxy* Proxy)
 	{
 		Proxy->Modify();
 		for (ULandscapeComponent* Component : Proxy->LandscapeComponents)
 		{
 			Component->DeleteLayer(InLayerInfo, LandscapeEdit);
 		}
+		return true;
 	});
 }
 
@@ -9111,10 +9133,11 @@ void ALandscape::ClearLayer(const FGuid& InLayerGuid, TSet<TObjectPtr<ULandscape
 	}
 	else
 	{
-		LandscapeInfo->ForAllLandscapeProxies([&](ALandscapeProxy* Proxy)
+		LandscapeInfo->ForEachLandscapeProxy([&](ALandscapeProxy* Proxy)
 		{
 			Proxy->Modify(bMarkPackageDirty);
 			Components.Append(Proxy->LandscapeComponents);
+			return true;
 		});
 	}
 
@@ -9443,9 +9466,10 @@ FLandscapeLayer* ALandscape::DuplicateLayerAndMoveBrushes(const FLandscapeLayer&
 	int32 AddedIndex = LandscapeLayers.Add(NewLayer);
 
 	// Create associated layer data in each landscape proxy
-	LandscapeInfo->ForAllLandscapeProxies([&NewLayer](ALandscapeProxy* Proxy)
+	LandscapeInfo->ForEachLandscapeProxy([&NewLayer](ALandscapeProxy* Proxy)
 	{
 		Proxy->AddLayer(NewLayer.Guid);
+		return true;
 	});
 
 	return &LandscapeLayers[AddedIndex];
@@ -9465,9 +9489,10 @@ int32 ALandscape::CreateLayer(FName InName)
 	int32 LayerIndex = LandscapeLayers.Add(NewLayer);
 
 	// Create associated layer data in each landscape proxy
-	LandscapeInfo->ForAllLandscapeProxies([&NewLayer](ALandscapeProxy* Proxy)
+	LandscapeInfo->ForEachLandscapeProxy([&NewLayer](ALandscapeProxy* Proxy)
 	{
 		Proxy->AddLayer(NewLayer.Guid);
+		return true;
 	});
 
 	return LayerIndex;
diff --git a/Engine/Source/Runtime/Landscape/Private/LandscapeSubsystem.cpp b/Engine/Source/Runtime/Landscape/Private/LandscapeSubsystem.cpp
--- a/Engine/Source/Runtime/Landscape/Private/LandscapeSubsystem.cpp
+++ b/Engine/Source/Runtime/Landscape/Private/LandscapeSubsystem.cpp
@@ -493,6 +493,8 @@ bool ULandscapeSubsystem::IsGridBased() const
 
 void ULandscapeSubsystem::ChangeGridSize(ULandscapeInfo* LandscapeInfo, uint32 GridSizeInComponents)
 {
+	TRACE_CPUPROFILER_EVENT_SCOPE(ULandscapeSubsystem::ChangeGridSize);
+	
 	if (!IsGridBased())
 	{
 		return;
diff --git a/Engine/Source/Runtime/Landscape/Public/LandscapeConfigHelper.h b/Engine/Source/Runtime/Landscape/Public/LandscapeConfigHelper.h
--- a/Engine/Source/Runtime/Landscape/Public/LandscapeConfigHelper.h
+++ b/Engine/Source/Runtime/Landscape/Public/LandscapeConfigHelper.h
@@ -32,7 +32,6 @@ struct LANDSCAPE_API FLandscapeConfig
 		, SubsectionSizeQuads(InSubsectionSizeQuads)
 		, GridSizeInComponents(InGridSizeInComponents)
 	{
-
 	}
 
 	FLandscapeConfig(ULandscapeInfo* InLandscapeInfo);
@@ -131,6 +130,32 @@ public:
 			}
 		}
 	}
+
+	template<typename T>
+	static void CopySubregion(const TArrayView<T>& InData, TArray<T>& OutData, const FIntRect& InSrcRegion, uint32 InSrcDataPitch)
+	{
+		check(InData.Size() >= InSrcRegion.Area());
+
+		const int32 SrcWidth = InSrcRegion.Width();
+		const int32 SrcHeight = InSrcRegion.Height();
+		
+		OutData.Empty(SrcWidth * SrcHeight);
+		OutData.AddUninitialized(SrcWidth * SrcHeight);
+
+		for (int32 Y = 0; Y < SrcHeight; ++Y)
+		{
+			for (int32 X = 0; X < SrcWidth; ++X)
+			{
+				const int32 SrcX = X + InSrcRegion.Min.X;
+				const int32 SrcY = Y + InSrcRegion.Min.Y;
+
+				const int32 SrcIndex = SrcX + SrcY * InSrcDataPitch;
+				const int32 DstIndex = X + Y * InSrcRegion.Width();
+
+				OutData[DstIndex] = InData[SrcIndex];
+			}
+		}
+	}
 		
 	template<typename T>
 	static void ResampleData(const TArray<T>& InData, TArray<T>& OutData, const FIntRect& InSrcRegion, const FIntRect& InDestRegion)
diff --git a/Engine/Source/Runtime/Landscape/Public/LandscapeSettings.h b/Engine/Source/Runtime/Landscape/Public/LandscapeSettings.h
--- a/Engine/Source/Runtime/Landscape/Public/LandscapeSettings.h
+++ b/Engine/Source/Runtime/Landscape/Public/LandscapeSettings.h
@@ -4,11 +4,14 @@
 
 #include "CoreMinimal.h"
 #include "Engine/DeveloperSettings.h"
+#include "Delegates/Delegate.h"
 #include "LandscapeSettings.generated.h"
 
 class UMaterialInterface;
 class ULandscapeLayerInfoObject;
 
+DECLARE_DELEGATE_OneParam(FnOnMaxComponentsChanged, uint64);
+
 UENUM()
 enum class ELandscapeDirtyingMode : uint8
 {
@@ -42,11 +45,16 @@ public:
 
 	/** Returns the default landscape layer info object that will be assigned to unset layers when creating a new landscape. */
 	TSoftObjectPtr<ULandscapeLayerInfoObject> GetDefaultLayerInfoObject() const { return DefaultLayerInfoObject; }
-
 public:
 	UPROPERTY(config, EditAnywhere, Category = "Layers", meta=(UIMin = "1", UIMax = "32", ClampMin = "1", ClampMax = "32", ToolTip = "This option controls the maximum editing layers that can be added to a Landscape"))
 	int32 MaxNumberOfLayers = 8;
 
+	UPROPERTY(config, EditAnywhere, Category = "Configuration", meta=(ToolTip = "Maximum Dimension of Landscape in Components"))
+	int32 MaxComponents = 256;
+
+	UPROPERTY(config, EditAnywhere, Category = "Configuration", meta = (UIMin = "1", UIMax = "1024", ClampMin = "1", ClampMax = "1024", ToolTip = "Maximum Size of Import Image Cache in MB"))
+	uint32 MaxImageImportCacheSizeMegaBytes = 256;
+	
 	UPROPERTY(Config, Category = "Configuration", EditAnywhere)
 	ELandscapeDirtyingMode LandscapeDirtyingMode;
 
@@ -62,4 +70,5 @@ protected:
 
 	UPROPERTY(transient)
 	bool bRestrictiveMode = false;
+
 };
