[Backout] - CL24913693
Original CL Desc
-----------------------------------------------------------------
Landscape: Create Large landscapes using World Partition

diff --git a/Engine/Plugins/Experimental/PCG/Source/PCG/Private/Grid/PCGLandscapeCache.cpp b/Engine/Plugins/Experimental/PCG/Source/PCG/Private/Grid/PCGLandscapeCache.cpp
--- a/Engine/Plugins/Experimental/PCG/Source/PCG/Private/Grid/PCGLandscapeCache.cpp
+++ b/Engine/Plugins/Experimental/PCG/Source/PCG/Private/Grid/PCGLandscapeCache.cpp
@@ -566,7 +566,7 @@ void UPCGLandscapeCache::PrimeCache()
 		if (IsValid(LandscapeInfo))
 		{
 			// Build per-component information
-			LandscapeInfo->ForEachLandscapeProxy([this, LandscapeInfo, &CacheEntriesToBuild](const ALandscapeProxy* LandscapeProxy)
+			LandscapeInfo->ForAllLandscapeProxies([this, LandscapeInfo, &CacheEntriesToBuild](const ALandscapeProxy* LandscapeProxy)
 			{
 				check(LandscapeProxy);
 				const FGuid LandscapeGuid = LandscapeProxy->GetLandscapeGuid();
@@ -586,7 +586,6 @@ void UPCGLandscapeCache::PrimeCache()
 						CacheEntriesToBuild.Emplace(LandscapeInfo, LandscapeComponent, ComponentKey);
 					}
 				}
-				return true;
 			});
 		}
 	}
diff --git a/Engine/Source/Editor/LandscapeEditor/LandscapeEditor.Build.cs b/Engine/Source/Editor/LandscapeEditor/LandscapeEditor.Build.cs
--- a/Engine/Source/Editor/LandscapeEditor/LandscapeEditor.Build.cs
+++ b/Engine/Source/Editor/LandscapeEditor/LandscapeEditor.Build.cs
@@ -10,7 +10,6 @@ public class LandscapeEditor : ModuleRules
 			new string[] {
 				"Core",
 				"CoreUObject",
-				"Json",
 				"ApplicationCore",
 				"Slate",
 				"SlateCore",
@@ -29,9 +28,6 @@ public class LandscapeEditor : ModuleRules
                 "Foliage",
 				"ToolMenus",
 				"ToolWidgets",
-				"SourceControl",
-				"DirectoryWatcher",
-				"DeveloperSettings"
 			}
 			);
 
diff --git a/Engine/Source/Editor/LandscapeEditor/Private/LandscapeEdMode.cpp b/Engine/Source/Editor/LandscapeEditor/Private/LandscapeEdMode.cpp
--- a/Engine/Source/Editor/LandscapeEditor/Private/LandscapeEdMode.cpp
+++ b/Engine/Source/Editor/LandscapeEditor/Private/LandscapeEdMode.cpp
@@ -20,7 +20,6 @@
 #include "LandscapeStreamingProxy.h"
 #include "LandscapeSubsystem.h"
 #include "LandscapeSettings.h"
-#include "LandscapeTiledImage.h"
 
 #include "EditorSupportDelegates.h"
 #include "ScopedTransaction.h"
@@ -3162,77 +3161,102 @@ void FEdModeLandscape::ReimportData(const FLandscapeTargetListInfo& TargetInfo)
 template<class T>
 void ImportDataInternal(ULandscapeInfo* LandscapeInfo, const FString& Filename, FName LayerName, bool bSingleFile, bool bFlipYAxis, const FIntRect& ImportRegionVerts, ELandscapeImportTransformType TransformType, FIntPoint Offset, TFunctionRef<void(int32, int32, int32, int32, const TArray<T>&)> SetDataFunc)
 {
-	if (!LandscapeInfo)
+	if (LandscapeInfo)
 	{
-		return;
-	}
-
-	FLandscapeTiledImage TiledImage;
-	TiledImage.Load(*Filename);
-	FIntPoint ImportResolution = TiledImage.GetResolution();
+		const FLandscapeImportResolution LandscapeResolution = { (uint32)(ImportRegionVerts.Width()), (uint32)(ImportRegionVerts.Height()) };
+		FLandscapeImportDescriptor OutImportDescriptor;
+		FText OutMessage;
 
-	bool bResolutionMismatch = false;
+		ELandscapeImportResult ImportResult = FLandscapeImportHelper::GetImportDescriptor<T>(Filename, bSingleFile, bFlipYAxis, LayerName, OutImportDescriptor, OutMessage);
+		if (ImportResult == ELandscapeImportResult::Error)
+		{
+			FMessageDialog::Open(EAppMsgType::Ok, OutMessage);
+			return;
+		}
 
-	if ((ImportResolution.X != ImportRegionVerts.Width() || ImportResolution.Y != ImportRegionVerts.Height()) && TransformType != ELandscapeImportTransformType::Subregion)
-	{
-		bResolutionMismatch = true;
+		check(OutImportDescriptor.ImportResolutions.Num() > 0);
 
-		FFormatNamedArguments Args;
-		Args.Add(TEXT("ImportSizeX"), ImportResolution.X);
-		Args.Add(TEXT("ImportSizeY"), ImportResolution.Y);
-		Args.Add(TEXT("LandscapeSizeX"), ImportRegionVerts.Width());
-		Args.Add(TEXT("LandscapeSizeY"), ImportRegionVerts.Height());
+		int32 DescriptorIndex = 0;
+		// If there is more than one possible res it needs to match perfectly (this constraint could probably removed)
+		if (OutImportDescriptor.ImportResolutions.Num() > 1)
+		{
+			DescriptorIndex = OutImportDescriptor.FindDescriptorIndex(LandscapeResolution.Width, LandscapeResolution.Height);
+			// if the file is a raw format with multiple possibly resolutions, only attempt import if one matches the current landscape
+			if (DescriptorIndex == INDEX_NONE)
+			{
+				FFormatNamedArguments Args;
+				Args.Add(TEXT("LandscapeSizeX"), LandscapeResolution.Width);
+				Args.Add(TEXT("LandscapeSizeY"), LandscapeResolution.Height);
 
-		auto Result = FMessageDialog::Open(EAppMsgType::OkCancel,
-			FText::Format(NSLOCTEXT("LandscapeEditor.Import", "Import_SizeMismatch", "The import size ({ImportSizeX}\u00D7{ImportSizeY}) does not match the current Landscape extent ({LandscapeSizeX}\u00D7{LandscapeSizeY}), if you continue it will be padded/clipped to fit"), Args));
+				FMessageDialog::Open(EAppMsgType::Ok,
+					FText::Format(NSLOCTEXT("LandscapeEditor.Import", "Import_SizeMismatchRaw", "The file resolution does not match the current Landscape extent ({LandscapeSizeX}\u00D7{LandscapeSizeY}), and its exact resolution could not be determined"), Args));
 
-		if (Result != EAppReturnType::Ok)
-		{
-			return;
+				return;
+			}
 		}
-	}
+				
+		// display warning message if there is one and allow user to cancel
+		if (ImportResult == ELandscapeImportResult::Warning)
+		{
+			auto Result = FMessageDialog::Open(EAppMsgType::OkCancel, OutMessage);
 
-	TArray<T> ImportData;
-	if (TransformType == ELandscapeImportTransformType::Subregion)
-	{
-		TiledImage.ReadRegion<T>(ImportRegionVerts, ImportData, bFlipYAxis);
-	}
-	else if (bResolutionMismatch)
-	{
-		TiledImage.Read<T>(ImportData, bFlipYAxis);
-	}
-	else
-	{
-		const FIntRect RegionToLoad(0, 0, ImportRegionVerts.Width(), ImportRegionVerts.Height());
-		TiledImage.ReadRegion<T>(RegionToLoad, ImportData, bFlipYAxis);
-	}
+			if (Result != EAppReturnType::Ok)
+			{
+				return;
+			}
+		}
 
-	// Expand if necessary...
-	{
-		TArray<T> FinalData;
-		if (bResolutionMismatch)
+		const FLandscapeImportResolution ImportResolution = OutImportDescriptor.ImportResolutions[DescriptorIndex];
+		bool bResolutionMismatch = false;
+		// if the file is a format with resolution information, warn the user if the resolution doesn't match the current landscape
+		// unlike for raw this is only a warning as we can pad/clip the data if we know what resolution it is
+		if (ImportResolution != LandscapeResolution)
 		{
-			FLandscapeImportHelper::TransformImportData<T>(ImportData, FinalData, FLandscapeImportResolution(ImportResolution.X, ImportResolution.Y), FLandscapeImportResolution(ImportRegionVerts.Width(), ImportRegionVerts.Height()), TransformType, Offset);
+			bResolutionMismatch = true;
+
+			FFormatNamedArguments Args;
+			Args.Add(TEXT("ImportSizeX"), ImportResolution.Width);
+			Args.Add(TEXT("ImportSizeY"), ImportResolution.Height);
+			Args.Add(TEXT("LandscapeSizeX"), LandscapeResolution.Width);
+			Args.Add(TEXT("LandscapeSizeY"), LandscapeResolution.Height);
+
+			auto Result = FMessageDialog::Open(EAppMsgType::OkCancel,
+				FText::Format(NSLOCTEXT("LandscapeEditor.Import", "Import_SizeMismatch", "The import size ({ImportSizeX}\u00D7{ImportSizeY}) does not match the current Landscape extent ({LandscapeSizeX}\u00D7{LandscapeSizeY}), if you continue it will be padded/clipped to fit"), Args));
+
+			if (Result != EAppReturnType::Ok)
+			{
+				return;
+			}
 		}
-		else if (TransformType == ELandscapeImportTransformType::Subregion)
+
+		TArray<T> ImportData;
+		ImportResult = FLandscapeImportHelper::GetImportData<T>(OutImportDescriptor, DescriptorIndex, LayerName, ImportData, OutMessage);
+		if (ImportResult == ELandscapeImportResult::Error)
 		{
-			FinalData = MoveTemp(ImportData);
+			FMessageDialog::Open(EAppMsgType::Ok, OutMessage);
+			return;
 		}
-		else
+
+		// Expand if necessary...
 		{
-			FinalData = MoveTemp(ImportData);
+			TArray<T> FinalData;
+			if (bResolutionMismatch)
+			{
+				FLandscapeImportHelper::TransformImportData<T>(ImportData, FinalData, ImportResolution, LandscapeResolution, TransformType, Offset);
+			}
+			else
+			{
+				FinalData = MoveTemp(ImportData);
+			}
+									
+			// Set Data is in Quads (so remove 1)
+			SetDataFunc(ImportRegionVerts.Min.X, ImportRegionVerts.Min.Y, ImportRegionVerts.Max.X-1, ImportRegionVerts.Max.Y-1, FinalData);
 		}
-
-		// Set Data is in Quads (so remove 1)
-		SetDataFunc(ImportRegionVerts.Min.X, ImportRegionVerts.Min.Y, ImportRegionVerts.Max.X - 1, ImportRegionVerts.Max.Y - 1, FinalData);
-	}
-	
+	}			
 }
 
 void FEdModeLandscape::ImportHeightData(ULandscapeInfo* LandscapeInfo, const FGuid& LayerGuid, const FString& Filename, const FIntRect& ImportRegionVerts, ELandscapeImportTransformType TransformType, FIntPoint Offset, const ELandscapeLayerPaintingRestriction& PaintRestriction, bool bFlipYAxis)
 {
-	TRACE_CPUPROFILER_EVENT_SCOPE(FEdModeLandscape::ImportHeightData);
-
 	ImportDataInternal<uint16>(LandscapeInfo, Filename, NAME_None, UseSingleFileImport(), bFlipYAxis, ImportRegionVerts, TransformType, Offset, [LandscapeInfo, LayerGuid, PaintRestriction](int32 MinX, int32 MinY, int32 MaxX, int32 MaxY, const TArray<uint16>& Data)
 	{
 		ALandscape* Landscape = LandscapeInfo->LandscapeActor.Get();
@@ -4279,8 +4303,13 @@ bool FEdModeLandscape::NeedToFillEmptyMaterialLayers() const
 
 	bool bCanFill = true;
 
-	CurrentToolTarget.LandscapeInfo->ForEachLandscapeProxy([&](ALandscapeProxy* Proxy)
+	CurrentToolTarget.LandscapeInfo->ForAllLandscapeProxies([&](ALandscapeProxy* Proxy)
 	{
+		if (!bCanFill)
+		{
+			return;
+		}
+
 		ALandscape* Landscape = Proxy->GetLandscapeActor();
 
 		if (Landscape != nullptr)
@@ -4298,14 +4327,13 @@ bool FEdModeLandscape::NeedToFillEmptyMaterialLayers() const
 							if (Alloc.LayerInfo != nullptr)
 							{
 								bCanFill = false;
-								return false;
+								return;
 							}
 						}
 					}
 				}
 			}
 		}
-		return true;
 	});	
 
 	return bCanFill;
diff --git a/Engine/Source/Editor/LandscapeEditor/Private/LandscapeEdModeComponentTools.cpp b/Engine/Source/Editor/LandscapeEditor/Private/LandscapeEdModeComponentTools.cpp
--- a/Engine/Source/Editor/LandscapeEditor/Private/LandscapeEdModeComponentTools.cpp
+++ b/Engine/Source/Editor/LandscapeEditor/Private/LandscapeEdModeComponentTools.cpp
@@ -536,105 +536,6 @@ public:
 //
 // FLandscapeToolAddComponent
 //
-
-namespace
-{
-	void AddComponents(ULandscapeInfo* LandscapeInfo, ULandscapeSubsystem* LandscapeSubsystem, const TArray<FIntPoint>& ComponentCoordinates)
-	{
-		TArray<ULandscapeComponent*> NewComponents;
-		LandscapeInfo->Modify();
-		for (const auto& ComponentCoordinate : ComponentCoordinates)
-		{
-			ULandscapeComponent* LandscapeComponent = LandscapeInfo->XYtoComponentMap.FindRef(ComponentCoordinate);
-			if (LandscapeComponent)
-			{
-				continue;
-			}
-
-			// Add New component...
-			FIntPoint ComponentBase = ComponentCoordinate * LandscapeInfo->ComponentSizeQuads;
-
-			ALandscapeProxy* LandscapeProxy = LandscapeSubsystem->FindOrAddLandscapeProxy(LandscapeInfo, ComponentBase);
-			if (!LandscapeProxy)
-			{
-				continue;
-			}
-
-			LandscapeComponent = NewObject<ULandscapeComponent>(LandscapeProxy, NAME_None, RF_Transactional);
-			NewComponents.Add(LandscapeComponent);
-			LandscapeComponent->Init(
-				ComponentBase.X, ComponentBase.Y,
-				LandscapeProxy->ComponentSizeQuads,
-				LandscapeProxy->NumSubsections,
-				LandscapeProxy->SubsectionSizeQuads
-			);
-
-			TArray<FColor> HeightData;
-			const int32 ComponentVerts = (LandscapeComponent->SubsectionSizeQuads + 1) * LandscapeComponent->NumSubsections;
-			HeightData.Empty(FMath::Square(ComponentVerts));
-			HeightData.AddZeroed(FMath::Square(ComponentVerts));
-			LandscapeComponent->InitHeightmapData(HeightData, true);
-			LandscapeComponent->UpdateMaterialInstances();
-
-			LandscapeInfo->XYtoComponentMap.Add(ComponentCoordinate, LandscapeComponent);
-			LandscapeInfo->XYtoAddCollisionMap.Remove(ComponentCoordinate);
-		}
-
-		// Need to register to use general height/xyoffset data update
-		for (int32 Idx = 0; Idx < NewComponents.Num(); Idx++)
-		{
-			NewComponents[Idx]->RegisterComponent();
-		}
-
-		const bool bHasXYOffset = false;
-		ALandscape* Landscape = LandscapeInfo->LandscapeActor.Get();
-
-		bool bHasLandscapeLayersContent = Landscape && Landscape->HasLayersContent();
-
-		if (bHasLandscapeLayersContent)
-		{
-			Landscape->RequestLayersInitialization();
-		}
-
-		for (ULandscapeComponent* NewComponent : NewComponents)
-		{
-			if (bHasLandscapeLayersContent)
-			{
-				TArray<ULandscapeComponent*> ComponentsUsingHeightmap;
-				ComponentsUsingHeightmap.Add(NewComponent);
-
-				for (const FLandscapeLayer& Layer : Landscape->LandscapeLayers)
-				{
-					// Since we do not share heightmap when adding new component, we will provided the required array, but they will only be used for 1 component
-					TMap<UTexture2D*, UTexture2D*> CreatedHeightmapTextures;
-					NewComponent->AddDefaultLayerData(Layer.Guid, ComponentsUsingHeightmap, CreatedHeightmapTextures);
-				}
-			}
-
-			// Update Collision
-			NewComponent->UpdateCachedBounds();
-			NewComponent->UpdateBounds();
-			NewComponent->MarkRenderStateDirty();
-
-			if (!bHasLandscapeLayersContent)
-			{
-				ULandscapeHeightfieldCollisionComponent* CollisionComp = NewComponent->GetCollisionComponent();
-				if (CollisionComp && !bHasXYOffset)
-				{
-					CollisionComp->MarkRenderStateDirty();
-					CollisionComp->RecreateCollision();
-				}
-			}
-		}
-
-		if (Landscape)
-		{
-			GEngine->BroadcastOnActorMoved(Landscape);
-		}
-
-	}
-}
-
 class FLandscapeToolStrokeAddComponent : public FLandscapeToolStrokeBase
 {
 	using Super = FLandscapeToolStrokeBase;
diff --git a/Engine/Source/Editor/LandscapeEditor/Private/LandscapeEditorDetailCustomization_ImportExport.cpp b/Engine/Source/Editor/LandscapeEditor/Private/LandscapeEditorDetailCustomization_ImportExport.cpp
--- a/Engine/Source/Editor/LandscapeEditor/Private/LandscapeEditorDetailCustomization_ImportExport.cpp
+++ b/Engine/Source/Editor/LandscapeEditor/Private/LandscapeEditorDetailCustomization_ImportExport.cpp
@@ -15,16 +15,6 @@
 #include "Widgets/Text/STextBlock.h"
 #include "Widgets/Notifications/SErrorText.h"
 #include "Widgets/SToolTip.h"
-
-#include "Misc/ScopedSlowTask.h"
-#include "Misc/MessageDialog.h"
-
-#include "LocationVolume.h"
-#include "SourceControlHelpers.h"
-#include "UObject/SavePackage.h"
-#include "WorldPartition/WorldPartition.h"
-#include "Internationalization/Regex.h"
-
 #include "Framework/Application/SlateApplication.h"
 #include "Framework/MultiBox/MultiBoxBuilder.h"
 
@@ -34,12 +24,6 @@
 #include "LandscapeEditorObject.h"
 #include "LandscapeEditorModule.h"
 #include "LandscapeImportHelper.h"
-#include "LandscapeSubsystem.h"
-#include "LandscapeStreamingProxy.h"
-#include "LandscapeImageFileCache.h"
-#include "LandscapeTiledImage.h"
-#include "LandscapeRegionUtils.h"
-#include "LandscapeEditorUtils.h"
 
 #define LOCTEXT_NAMESPACE "LandscapeEditor.ImportExport"
 
@@ -455,18 +439,42 @@ FReply FLandscapeEditorDetailCustomization_ImportExport::OnBrowseFilenameButtonC
 {
 	FEdModeLandscape* LandscapeEdMode = GetEditorMode();
 	check(LandscapeEdMode != nullptr);
-	ILandscapeEditorModule& LandscapeEditorModule = FModuleManager::GetModuleChecked<ILandscapeEditorModule>("LandscapeEditor");
-	const bool bIsImporting = IsImporting();
-	const FString DialogType = bIsImporting ? LandscapeEditorModule.GetHeightmapImportDialogTypeString() : LandscapeEditorModule.GetHeightmapExportDialogTypeString();
-	const FString DialogTitle = bIsImporting ? LOCTEXT("ImportHeightmap", "Import Heightmap").ToString() : LOCTEXT("ExportHeightmap", "Export Heightmap").ToString();
 
-	TOptional<FString> OptionalFilename = LandscapeEditorUtils::GetImportExportFilename(DialogTitle, LandscapeEdMode->UISettings->LastImportPath, DialogType);
-	if (OptionalFilename.IsSet())
+	// Prompt the user for the Filenames
+	IDesktopPlatform* DesktopPlatform = FDesktopPlatformModule::Get();
+	if (DesktopPlatform != nullptr)
 	{
-		const FString& Filename = OptionalFilename.GetValue();
+		bool bSuccess = false;
+		TArray<FString> Filenames;
+		ILandscapeEditorModule& LandscapeEditorModule = FModuleManager::GetModuleChecked<ILandscapeEditorModule>("LandscapeEditor");
+		if (IsImporting())
+		{
+			bSuccess = DesktopPlatform->OpenFileDialog(
+				FSlateApplication::Get().FindBestParentWindowHandleForDialogs(nullptr),
+				LOCTEXT("ImportHeightmap", "Import Heightmap").ToString(),
+				LandscapeEdMode->UISettings->LastImportPath,
+				TEXT(""),
+				LandscapeEditorModule.GetHeightmapImportDialogTypeString(),
+				EFileDialogFlags::None,
+				Filenames);
+		}
+		else
+		{
+			bSuccess = DesktopPlatform->SaveFileDialog(
+				FSlateApplication::Get().FindBestParentWindowHandleForDialogs(nullptr),
+				LOCTEXT("ExportHeightmap", "Export Heightmap").ToString(),
+				LandscapeEdMode->UISettings->LastImportPath,
+				TEXT(""),
+				LandscapeEditorModule.GetHeightmapExportDialogTypeString(),
+				EFileDialogFlags::None,
+				Filenames);
+		}
 
-		ensure(PropertyHandle_Filename->SetValue(Filename) == FPropertyAccess::Success);
-		LandscapeEdMode->UISettings->LastImportPath = FPaths::GetPath(Filename);
+		if (bSuccess)
+		{
+			ensure(PropertyHandle_Filename->SetValue(Filenames[0]) == FPropertyAccess::Success);
+			LandscapeEdMode->UISettings->LastImportPath = FPaths::GetPath(Filenames[0]);
+		}
 	}
 
 	return FReply::Handled();
@@ -544,13 +552,13 @@ FReply FLandscapeEditorDetailCustomization_ImportExport::OnImportExportButtonCli
 {
 	FEdModeLandscape* LandscapeEdMode = GetEditorMode();
 	check(LandscapeEdMode);
-
-	UWorld* World = LandscapeEdMode->GetWorld();
-
 	if (ULandscapeInfo* LandscapeInfo = LandscapeEdMode->CurrentToolTarget.LandscapeInfo.Get())
 	{
 		FIntRect LandscapeExtent;
-		LandscapeInfo->GetLandscapeExtent(LandscapeExtent);
+		if (!LandscapeInfo->GetLandscapeExtent(LandscapeExtent))
+		{
+			return FReply::Handled();
+		}
 	
 		if (IsImporting())
 		{
@@ -577,87 +585,20 @@ FReply FLandscapeEditorDetailCustomization_ImportExport::OnImportExportButtonCli
 				ImportRegion = FIntRect(LocalGizmoPoint.X, LocalGizmoPoint.Y, LocalGizmoPoint.X + LandscapeEdMode->UISettings->ImportLandscape_Width, LocalGizmoPoint.Y + LandscapeEdMode->UISettings->ImportLandscape_Height);
 			}
 			
-			check(LandscapeEdMode->UISettings->ImportLandscape_HeightmapImportResult != ELandscapeImportResult::Error);
-			const bool bIsWorldPartition = LandscapeEdMode->GetWorld()->GetSubsystem<ULandscapeSubsystem>()->IsGridBased();
-
-			if (bIsWorldPartition && LandscapeEdMode->UISettings->ImportExportMode != ELandscapeImportExportMode::LoadedOnly) // If in world partition import region by region
+			// Import Heightmap
+			if (LandscapeEdMode->UISettings->bHeightmapSelected)
 			{
-				TArray<ALocationVolume*> LandscapeRegions;
-				TArray<AActor*> Children;
-				LandscapeInfo->LandscapeActor->GetAttachedActors(Children);
-				for (AActor* Child : Children)
-				{
-					if (Child->IsA<ALocationVolume>())
-					{
-						LandscapeRegions.Add(Cast<ALocationVolume>(Child));
-					}
-				}
-
-				int32 NumRegions = LandscapeRegions.Num();
-					
-				FScopedSlowTask Progress(NumRegions, LOCTEXT("ImportingLandscapeRegions", "Importing Landscape Regions..."));
-				Progress.MakeDialog(/*bShowCancelButton = */ true);
-
-				auto RegionImporter = [&Progress, LandscapeInfo, CurrentLayerGuid, LandscapeEdMode, ImportRegion, TransformType, ImportOffset, PaintRestriction](const FBox& RegionBounds, const TArray<ALandscapeProxy*>& Proxies)
-				{
-					// todo need a better way of working out the landscape bounds of the StreamingProxies in this region.
-					// What if some of the landscape is loaded?
-					FIntRect LandscapeLoadedExtent;
-					LandscapeInfo->GetLandscapeExtent(LandscapeLoadedExtent);
-					LandscapeLoadedExtent.Max.X += 1;
-					LandscapeLoadedExtent.Max.Y += 1;
-
-					Progress.EnterProgressFrame(1.0f, LOCTEXT("ImportingLandscapeRegions", "Importing Landscape Regions"));
-
-					if (LandscapeEdMode->UISettings->bHeightmapSelected)
-					{
-						LandscapeEdMode->ImportHeightData(LandscapeInfo, CurrentLayerGuid, LandscapeEdMode->UISettings->ImportLandscape_HeightmapFilename, LandscapeLoadedExtent, TransformType, ImportOffset, PaintRestriction, LandscapeEdMode->UISettings->bFlipYAxis);
-					}
-
-					for (const FLandscapeImportLayer& ImportLayer : LandscapeEdMode->UISettings->ImportLandscape_Layers)
-					{
-						if (ImportLayer.bSelected)
-						{
-							check(ImportLayer.ImportResult != ELandscapeImportResult::Error);
-							LandscapeEdMode->ImportWeightData(LandscapeInfo, CurrentLayerGuid, ImportLayer.LayerInfo, ImportLayer.SourceFilePath, LandscapeLoadedExtent, TransformType, ImportOffset, PaintRestriction, LandscapeEdMode->UISettings->bFlipYAxis);
-						}
-					}
-
-					return !Progress.ShouldCancel();
-				};
-					
-				LandscapeRegionUtils::ForEachRegion_LoadProcessUnload(LandscapeInfo, ImportRegion, World, RegionImporter);
-
-				FLandscapeImageFileCache& LandscapeImageFileCache = FModuleManager::GetModuleChecked<ILandscapeEditorModule>("LandscapeEditor").GetImageFileCache();
-				LandscapeImageFileCache.Clear();
+				check(LandscapeEdMode->UISettings->ImportLandscape_HeightmapImportResult != ELandscapeImportResult::Error);
+				LandscapeEdMode->ImportHeightData(LandscapeInfo, CurrentLayerGuid, LandscapeEdMode->UISettings->ImportLandscape_HeightmapFilename, ImportRegion, TransformType, ImportOffset, PaintRestriction, LandscapeEdMode->UISettings->bFlipYAxis);
 			}
-			else
-			{
-				FScopedSlowTask Progress(1 + LandscapeEdMode->UISettings->ImportLandscape_Layers.Num(), LOCTEXT("ImportingLandscape", "Importing Landscape"));
-				Progress.MakeDialog(/*bShowCancelButton = */ true);
-
-				if (LandscapeEdMode->UISettings->bHeightmapSelected)
-				{
-					Progress.EnterProgressFrame(1.0f, LOCTEXT("ImportingLandscapeHeight", "Importing Landscape Height"));
-					LandscapeEdMode->ImportHeightData(LandscapeInfo, CurrentLayerGuid, LandscapeEdMode->UISettings->ImportLandscape_HeightmapFilename, ImportRegion, TransformType, ImportOffset, PaintRestriction, LandscapeEdMode->UISettings->bFlipYAxis);
-				}
 
-				for (const FLandscapeImportLayer& ImportLayer : LandscapeEdMode->UISettings->ImportLandscape_Layers)
+			for (const FLandscapeImportLayer& ImportLayer : LandscapeEdMode->UISettings->ImportLandscape_Layers)
+			{
+				if (ImportLayer.bSelected)
 				{
-					if (Progress.ShouldCancel())
-					{
-						continue;
-					}
-
-					if (ImportLayer.bSelected)
-					{
-						Progress.EnterProgressFrame(1.0f, LOCTEXT("ImportingLandscapeWeight", "Importing Landscape Weight"));
-						check(ImportLayer.ImportResult != ELandscapeImportResult::Error);
-						LandscapeEdMode->ImportWeightData(LandscapeInfo, CurrentLayerGuid, ImportLayer.LayerInfo, ImportLayer.SourceFilePath, ImportRegion, TransformType, ImportOffset, PaintRestriction, LandscapeEdMode->UISettings->bFlipYAxis);
-					}
+					check(ImportLayer.ImportResult != ELandscapeImportResult::Error);
+					LandscapeEdMode->ImportWeightData(LandscapeInfo, CurrentLayerGuid, ImportLayer.LayerInfo, ImportLayer.SourceFilePath, ImportRegion, TransformType, ImportOffset, PaintRestriction, LandscapeEdMode->UISettings->bFlipYAxis);
 				}
-
-				LandscapeInfo->ForceLayersFullUpdate();
 			}
 		}
 		else
@@ -707,39 +648,8 @@ FReply FLandscapeEditorDetailCustomization_ImportExport::OnImportExportButtonCli
 			}
 			else
 			{
-				FScopedSlowTask Progress(0, LOCTEXT("ExportingLandscapeRegions", "Exporting Landscape Regions..."));
-				Progress.MakeDialog(/*bShowCancelButton = */ true);
-
-				auto Exporter = [&Progress, LandscapeInfo, Landscape, LandscapeEdMode, LandscapeExtent, PerformExport, World](const FBox& RegionBounds, const TArray<ALandscapeProxy*>& Proxies)
-				{
-					Progress.EnterProgressFrame(0.0f, LOCTEXT("ExportingLandscapeRegions", "Exporting Landscape Regions"));
-					for (ALandscapeProxy* LandscapeProxy : Proxies)
-					{
-						FIntRect ExportRegion;
-
-						if (LandscapeInfo->GetLandscapeExtent(LandscapeProxy, ExportRegion))
-						{
-							const int32 YCoord = LandscapeEdMode->UISettings->bFlipYAxis ? LandscapeExtent.Max.Y - ExportRegion.Max.Y : ExportRegion.Min.Y - LandscapeExtent.Min.Y;
-							FIntPoint FileOffset = FIntPoint((ExportRegion.Min.X - LandscapeExtent.Min.X) / Landscape->GridSize,
-								YCoord / Landscape->GridSize);
-
-							// Remove the shared line/column that this proxy has with its neighbors because it 
-							// will be included by the neighbor or lost if there is none (that could become an option to avoid that loss)
-							ExportRegion.Max.X -= 1;
-							ExportRegion.Max.Y -= 1;
-
-							PerformExport(LandscapeEdMode->UISettings, LandscapeInfo, ExportRegion, FileOffset, /*bInUseOffset = */true);
-						}
-					}
-					return !Progress.ShouldCancel();
-					
-				};
-
-				LandscapeRegionUtils::ForEachRegion_LoadProcessUnload(LandscapeInfo, LandscapeExtent, World, Exporter);
-			
-
 				// For multiple file export, export each landscape proxy individually :
-				LandscapeInfo->ForEachLandscapeProxy([LandscapeInfo, Landscape, LandscapeEdMode, LandscapeExtent, BuildExportFileName, PerformExport](ALandscapeProxy* LandscapeProxy)
+				LandscapeInfo->ForAllLandscapeProxies([LandscapeInfo, Landscape, LandscapeEdMode, LandscapeExtent, BuildExportFileName, PerformExport](ALandscapeProxy* LandscapeProxy)
 				{
 					FIntRect ExportRegion;
 					if (LandscapeInfo->GetLandscapeExtent(LandscapeProxy, ExportRegion))
@@ -755,7 +665,6 @@ FReply FLandscapeEditorDetailCustomization_ImportExport::OnImportExportButtonCli
 
 						PerformExport(LandscapeEdMode->UISettings, LandscapeInfo, ExportRegion, FileOffset, /*bInUseOffset = */true);
 					}
-					return true;
 				});
 			}
 
@@ -895,6 +804,4 @@ FText FLandscapeEditorDetailCustomization_ImportExport::GetImportLandscapeResolu
 	return FText::GetEmpty();
 }
 
-
-
 #undef LOCTEXT_NAMESPACE
diff --git a/Engine/Source/Editor/LandscapeEditor/Private/LandscapeEditorDetailCustomization_ImportLayers.cpp b/Engine/Source/Editor/LandscapeEditor/Private/LandscapeEditorDetailCustomization_ImportLayers.cpp
--- a/Engine/Source/Editor/LandscapeEditor/Private/LandscapeEditorDetailCustomization_ImportLayers.cpp
+++ b/Engine/Source/Editor/LandscapeEditor/Private/LandscapeEditorDetailCustomization_ImportLayers.cpp
@@ -8,8 +8,6 @@
 #include "LandscapeEditorModule.h"
 #include "LandscapeEditorObject.h"
 #include "LandscapeUtils.h"
-#include "LandscapeTiledImage.h"
-#include "LandscapeEditorUtils.h"
 #include "SLandscapeEditor.h"
 
 #include "IDetailChildrenBuilder.h"
@@ -26,7 +24,6 @@
 #include "Widgets/SToolTip.h"
 #include "Framework/Application/SlateApplication.h"
 #include "Framework/MultiBox/MultiBoxBuilder.h"
-#include "Misc/MessageDialog.h"
 
 #define LOCTEXT_NAMESPACE "LandscapeEditor.ImportLayers"
 
@@ -231,18 +228,41 @@ FReply FLandscapeEditorStructCustomization_FLandscapeImportLayer::OnLayerFilenam
 	FEdModeLandscape* LandscapeEdMode = GetEditorMode();
 	check(LandscapeEdMode != nullptr);
 
-	ILandscapeEditorModule& LandscapeEditorModule = FModuleManager::GetModuleChecked<ILandscapeEditorModule>("LandscapeEditor");
-	const bool bIsImporting = IsImporting();
-	const FString DialogTypeString = bIsImporting  ? LandscapeEditorModule.GetWeightmapImportDialogTypeString() : LandscapeEditorModule.GetWeightmapExportDialogTypeString();
-	const FString DialogTitle = bIsImporting ? LOCTEXT("ImportLayer", "Import Layer").ToString() : LOCTEXT("ExportLayer", "Export Layer").ToString();
-
-	TOptional<FString> OptionalExportImportPath = LandscapeEditorUtils::GetImportExportFilename(DialogTitle, LandscapeEdMode->UISettings->LastImportPath, DialogTypeString);
-	if (OptionalExportImportPath.IsSet())
+	// Prompt the user for the Filenames
+	IDesktopPlatform* DesktopPlatform = FDesktopPlatformModule::Get();
+	if (DesktopPlatform != nullptr)
 	{
-		const FString& Filename = OptionalExportImportPath.GetValue();
+		bool bSuccess = false;
+		TArray<FString> Filenames;
+		ILandscapeEditorModule& LandscapeEditorModule = FModuleManager::GetModuleChecked<ILandscapeEditorModule>("LandscapeEditor");
+		if (IsImporting())
+		{
+			bSuccess = DesktopPlatform->OpenFileDialog(
+				FSlateApplication::Get().FindBestParentWindowHandleForDialogs(nullptr),
+				LOCTEXT("ImportLayer", "Import Layer").ToString(),
+				LandscapeEdMode->UISettings->LastImportPath,
+				TEXT(""),
+				LandscapeEditorModule.GetWeightmapImportDialogTypeString(),
+				EFileDialogFlags::None,
+				Filenames);
+		}
+		else
+		{
+			bSuccess = DesktopPlatform->SaveFileDialog(
+				FSlateApplication::Get().FindBestParentWindowHandleForDialogs(nullptr),
+				LOCTEXT("ExportLayer", "Export Layer").ToString(),
+				LandscapeEdMode->UISettings->LastImportPath,
+				TEXT(""),
+				LandscapeEditorModule.GetWeightmapExportDialogTypeString(),
+				EFileDialogFlags::None,
+				Filenames);
+		}
 
-		ensure(PropertyHandle_LayerFilename->SetValue(Filename) == FPropertyAccess::Success);
-		LandscapeEdMode->UISettings->LastImportPath = FPaths::GetPath(Filename);
+		if (bSuccess)
+		{
+			ensure(PropertyHandle_LayerFilename->SetValue(Filenames[0]) == FPropertyAccess::Success);
+			LandscapeEdMode->UISettings->LastImportPath = FPaths::GetPath(Filenames[0]);
+		}
 	}
 
 	return FReply::Handled();
diff --git a/Engine/Source/Editor/LandscapeEditor/Private/LandscapeEditorDetailCustomization_NewLandscape.cpp b/Engine/Source/Editor/LandscapeEditor/Private/LandscapeEditorDetailCustomization_NewLandscape.cpp
--- a/Engine/Source/Editor/LandscapeEditor/Private/LandscapeEditorDetailCustomization_NewLandscape.cpp
+++ b/Engine/Source/Editor/LandscapeEditor/Private/LandscapeEditorDetailCustomization_NewLandscape.cpp
@@ -21,16 +21,9 @@
 #include "LandscapeEditorModule.h"
 #include "LandscapeEditorObject.h"
 #include "Landscape.h"
-#include "LandscapeStreamingProxy.h"
-#include "LandscapeTiledImage.h"
-#include "LandscapeRegionUtils.h"
-#include "LandscapeEditorPrivate.h"
-#include "LandscapeEditorUtils.h"
-
 #include "LandscapeConfigHelper.h"
 #include "LandscapeImportHelper.h"
 #include "LandscapeSettings.h"
-#include "FileHelpers.h"
 
 #include "DetailLayoutBuilder.h"
 #include "IDetailChildrenBuilder.h"
@@ -43,6 +36,7 @@
 #include "Widgets/Input/SVectorInputBox.h"
 #include "Widgets/Input/SRotatorInputBox.h"
 #include "ScopedTransaction.h"
+#include "DesktopPlatformModule.h"
 #include "AssetRegistry/AssetRegistryModule.h"
 
 #include "TutorialMetaData.h"
@@ -53,22 +47,9 @@
 #include "Editor.h"
 #include "Editor/EditorEngine.h"
 #include "LandscapeSubsystem.h"
-#include "LocationVolume.h"
 #include "SPrimaryButton.h"
 #include "Widgets/Input/SSegmentedControl.h"
 
-#include "UObject/SavePackage.h"
-#include "Builders/CubeBuilder.h"
-#include "ActorFactories/ActorFactory.h"
-
-#include "SourceControlHelpers.h"
-#include "WorldPartition/WorldPartition.h"
-
-#include "Misc/ScopedSlowTask.h"
-
-#include "Math/Box.h"
-#include "LandscapeEditorUtils.h"
-
 #define LOCTEXT_NAMESPACE "LandscapeEditor.NewLandscape"
 
 
@@ -141,20 +122,16 @@ void FLandscapeEditorDetailCustomization_NewLandscape::CustomizeDetails(IDetailL
 	PropertyHandle_FlipYAxis->SetOnPropertyValueChanged(FSimpleDelegate::CreateLambda([this]() { OnImportHeightmapFilenameChanged(); }));
 
 	TSharedRef<IPropertyHandle> PropertyHandle_GridSize = DetailBuilder.GetProperty(GET_MEMBER_NAME_CHECKED(ULandscapeEditorObject, WorldPartitionGridSize));
-	TSharedRef<IPropertyHandle> PropertyHandle_RegionSize = DetailBuilder.GetProperty(GET_MEMBER_NAME_CHECKED(ULandscapeEditorObject, WorldPartitionRegionSize));
 	
 	FEdModeLandscape* LandscapeEdMode = GetEditorMode();
 	if (LandscapeEdMode->GetWorld()->GetSubsystem<ULandscapeSubsystem>()->IsGridBased())
 	{
 		NewLandscapeCategory.AddProperty(PropertyHandle_GridSize);
-		NewLandscapeCategory.AddProperty(PropertyHandle_RegionSize);
 	}
 	else
 	{
 		DetailBuilder.HideProperty(PropertyHandle_GridSize);
-		DetailBuilder.HideProperty(PropertyHandle_RegionSize);
 	}
-	
 	TSharedRef<IPropertyHandle> PropertyHandle_HeightmapFilename = DetailBuilder.GetProperty(GET_MEMBER_NAME_CHECKED(ULandscapeEditorObject, ImportLandscape_HeightmapFilename));
 	TSharedRef<IPropertyHandle> PropertyHandle_HeightmapImportResult = DetailBuilder.GetProperty(GET_MEMBER_NAME_CHECKED(ULandscapeEditorObject, ImportLandscape_HeightmapImportResult));
 	TSharedRef<IPropertyHandle> PropertyHandle_HeightmapErrorMessage = DetailBuilder.GetProperty(GET_MEMBER_NAME_CHECKED(ULandscapeEditorObject, ImportLandscape_HeightmapErrorMessage));
@@ -415,9 +392,9 @@ void FLandscapeEditorDetailCustomization_NewLandscape::CustomizeDetails(IDetailL
 			SNew(SNumericEntryBox<int32>)
 			.Font(DetailBuilder.GetDetailFont())
 			.MinValue(1)
-			.MaxValue_Lambda(LandscapeEditorUtils::GetMaxSizeInComponents)
+			.MaxValue(32)
 			.MinSliderValue(1)
-			.MaxSliderValue_Lambda(LandscapeEditorUtils::GetMaxSizeInComponents)
+			.MaxSliderValue(32)
 			.AllowSpin(true)
 			.UndeterminedString(NSLOCTEXT("PropertyEditor", "MultipleValues", "Multiple Values"))
 			.Value_Static(&FLandscapeEditorDetailCustomization_Base::OnGetValue<int32>, PropertyHandle_ComponentCount_X)
@@ -439,9 +416,9 @@ void FLandscapeEditorDetailCustomization_NewLandscape::CustomizeDetails(IDetailL
 			SNew(SNumericEntryBox<int32>)
 			.Font(DetailBuilder.GetDetailFont())
 			.MinValue(1)
-			.MaxValue_Lambda(LandscapeEditorUtils::GetMaxSizeInComponents)
+			.MaxValue(32)
 			.MinSliderValue(1)
-			.MaxSliderValue_Lambda(LandscapeEditorUtils::GetMaxSizeInComponents)
+			.MaxSliderValue(32)
 			.AllowSpin(true)
 			.UndeterminedString(NSLOCTEXT("PropertyEditor", "MultipleValues", "Multiple Values"))
 			.Value_Static(&FLandscapeEditorDetailCustomization_Base::OnGetValue<int32>, PropertyHandle_ComponentCount_Y)
@@ -898,360 +875,133 @@ FText FLandscapeEditorDetailCustomization_NewLandscape::GetNewLandscapeErrorText
 	return FText::GetEmpty();
 }
 
-void FLandscapeEditorDetailCustomization_NewLandscape::AddComponents(ULandscapeInfo* InLandscapeInfo, ULandscapeSubsystem* InLandscapeSubsystem, const TArray<FIntPoint>& InComponentCoordinates, TArray<ALandscapeProxy*>& OutCreatedStreamingProxies)
+FReply FLandscapeEditorDetailCustomization_NewLandscape::OnCreateButtonClicked()
 {
-	TRACE_CPUPROFILER_EVENT_SCOPE(AddComponents);
-	TArray<ULandscapeComponent*> NewComponents;
-	InLandscapeInfo->Modify();
-	for (const FIntPoint& ComponentCoordinate : InComponentCoordinates)
+	FEdModeLandscape* LandscapeEdMode = GetEditorMode();
+	if (LandscapeEdMode != nullptr && 
+		LandscapeEdMode->GetWorld() != nullptr && 
+		LandscapeEdMode->GetWorld()->GetCurrentLevel()->bIsVisible)
 	{
-		ULandscapeComponent* LandscapeComponent = InLandscapeInfo->XYtoComponentMap.FindRef(ComponentCoordinate);
-		if (LandscapeComponent)
-		{
-			continue;
-		}
+		ULandscapeEditorObject* UISettings = LandscapeEdMode->UISettings;
+		const int32 ComponentCountX = UISettings->NewLandscape_ComponentCount.X;
+		const int32 ComponentCountY = UISettings->NewLandscape_ComponentCount.Y;
+		const int32 QuadsPerComponent = UISettings->NewLandscape_SectionsPerComponent * UISettings->NewLandscape_QuadsPerSection;
+		const int32 SizeX = ComponentCountX * QuadsPerComponent + 1;
+		const int32 SizeY = ComponentCountY * QuadsPerComponent + 1;
 
-		// Add New component...
-		FIntPoint ComponentBase = ComponentCoordinate * InLandscapeInfo->ComponentSizeQuads;
+		TArray<FLandscapeImportLayerInfo> MaterialImportLayers;
+		ELandscapeImportResult LayerImportResult = LandscapeEdMode->NewLandscapePreviewMode == ENewLandscapePreviewMode::NewLandscape ? UISettings->CreateNewLayersInfo(MaterialImportLayers) : UISettings->CreateImportLayersInfo(MaterialImportLayers);
 
-		ALandscapeProxy* LandscapeProxy = InLandscapeSubsystem->FindOrAddLandscapeProxy(InLandscapeInfo, ComponentBase);
-		if (!LandscapeProxy)
+		if (LayerImportResult == ELandscapeImportResult::Error)
 		{
-			continue;
+			return FReply::Handled();
 		}
 
-		OutCreatedStreamingProxies.Add(LandscapeProxy);
-
-		LandscapeComponent = NewObject<ULandscapeComponent>(LandscapeProxy, NAME_None, RF_Transactional);
-		NewComponents.Add(LandscapeComponent);
-		LandscapeComponent->Init(
-			ComponentBase.X, ComponentBase.Y,
-			LandscapeProxy->ComponentSizeQuads,
-			LandscapeProxy->NumSubsections,
-			LandscapeProxy->SubsectionSizeQuads
-		);
-
-		TArray<FColor> HeightData;
-		const int32 ComponentVerts = (LandscapeComponent->SubsectionSizeQuads + 1) * LandscapeComponent->NumSubsections;
-		const FColor PackedMidpoint = LandscapeDataAccess::PackHeight(LandscapeDataAccess::GetTexHeight(0.0f));
-		HeightData.Init(PackedMidpoint, FMath::Square(ComponentVerts));
-
-		LandscapeComponent->InitHeightmapData(HeightData, true);
-		LandscapeComponent->UpdateMaterialInstances();
-
-		InLandscapeInfo->XYtoComponentMap.Add(ComponentCoordinate, LandscapeComponent);
-		InLandscapeInfo->XYtoAddCollisionMap.Remove(ComponentCoordinate);
-	}
+		TMap<FGuid, TArray<uint16>> HeightDataPerLayers;
+		TMap<FGuid, TArray<FLandscapeImportLayerInfo>> MaterialLayerDataPerLayers;
 
-	// Need to register to use general height/xyoffset data update
-	for (int32 Idx = 0; Idx < NewComponents.Num(); Idx++)
-	{
-		NewComponents[Idx]->RegisterComponent();
-	}
-
-	const bool bHasXYOffset = false;
-	ALandscape* Landscape = InLandscapeInfo->LandscapeActor.Get();
-
-	bool bHasLandscapeLayersContent = Landscape && Landscape->HasLayersContent();
-
-	for (ULandscapeComponent* NewComponent : NewComponents)
-	{
-		if (bHasLandscapeLayersContent)
+		TArray<uint16> OutHeightData;
+		if (LandscapeEdMode->NewLandscapePreviewMode == ENewLandscapePreviewMode::NewLandscape)
 		{
-			TArray<ULandscapeComponent*> ComponentsUsingHeightmap;
-			ComponentsUsingHeightmap.Add(NewComponent);
-
-			for (const FLandscapeLayer& Layer : Landscape->LandscapeLayers)
-			{
-				// Since we do not share heightmap when adding new component, we will provided the required array, but they will only be used for 1 component
-				TMap<UTexture2D*, UTexture2D*> CreatedHeightmapTextures;
-				NewComponent->AddDefaultLayerData(Layer.Guid, ComponentsUsingHeightmap, CreatedHeightmapTextures);
-			}
+			UISettings->InitializeDefaultHeightData(OutHeightData);
 		}
-
-		// Update Collision
-		NewComponent->UpdateCachedBounds();
-		NewComponent->UpdateBounds();
-		NewComponent->MarkRenderStateDirty();
-
-		if (!bHasLandscapeLayersContent)
+		else
 		{
-			ULandscapeHeightfieldCollisionComponent* CollisionComp = NewComponent->GetCollisionComponent();
-			if (CollisionComp && !bHasXYOffset)
-			{
-				CollisionComp->MarkRenderStateDirty();
-				CollisionComp->RecreateCollision();
-			}
+			UISettings->ExpandImportData(OutHeightData, MaterialImportLayers);
 		}
-	}
 
+		HeightDataPerLayers.Add(FGuid(), OutHeightData);
+		// ComputeHeightData will also modify/expand material layers data, which is why we create MaterialLayerDataPerLayers after calling ComputeHeightData
+		MaterialLayerDataPerLayers.Add(FGuid(), MoveTemp(MaterialImportLayers));
 
-	if (Landscape)
-	{
-		GEngine->BroadcastOnActorMoved(Landscape);
-	}
-}
-
-FReply FLandscapeEditorDetailCustomization_NewLandscape::OnCreateButtonClicked()
-{
-	TRACE_CPUPROFILER_EVENT_SCOPE(FLandscapeEditorDetailCustomization_NewLandscape::OnCreateButtonClicked);
-
-	FEdModeLandscape* LandscapeEdMode = GetEditorMode();
-	const bool bIsNewLandscape = LandscapeEdMode->NewLandscapePreviewMode == ENewLandscapePreviewMode::NewLandscape;
-	
-	UWorld* World = LandscapeEdMode->GetWorld();
-	const bool bIsTempPackage = FPackageName::IsTempPackage(World->GetPackage()->GetName());
-	
-	const bool bCreateLandscape = LandscapeEdMode != nullptr && 
-		World != nullptr &&
-		World->GetCurrentLevel()->bIsVisible;
-	
-	if (!bCreateLandscape)
-	{
-		return FReply::Handled();
-	}
-	
-	ULandscapeEditorObject* UISettings = LandscapeEdMode->UISettings;
-	
-	const bool bIsWorldPartition = World->GetSubsystem<ULandscapeSubsystem>()->IsGridBased();
-	const bool bLandscapeLargerThanRegion = static_cast<int32>(UISettings->WorldPartitionRegionSize) < UISettings->NewLandscape_ComponentCount.X || static_cast<int32>(UISettings->WorldPartitionRegionSize) < UISettings->NewLandscape_ComponentCount.Y;
-	const bool bNeedsLandscapeRegions =  bIsWorldPartition && bLandscapeLargerThanRegion;
+		FScopedTransaction Transaction(LOCTEXT("Undo", "Creating New Landscape"));
 
-	// If we need to ensure the map is saved before proceeding to create a landscape with regions 
-	if (bIsTempPackage && bNeedsLandscapeRegions)
-	{
-		FString NewMapPackageName;
-		if (!FEditorFileUtils::SaveCurrentLevel())
+		const FVector Offset = FTransform(UISettings->NewLandscape_Rotation, FVector::ZeroVector, UISettings->NewLandscape_Scale).TransformVector(FVector(-ComponentCountX * QuadsPerComponent / 2, -ComponentCountY * QuadsPerComponent / 2, 0));
+		
+		ALandscape* Landscape = LandscapeEdMode->GetWorld()->SpawnActor<ALandscape>(UISettings->NewLandscape_Location + Offset, UISettings->NewLandscape_Rotation);
+		Landscape->bCanHaveLayersContent = UISettings->bCanHaveLayersContent;
+		Landscape->LandscapeMaterial = UISettings->NewLandscape_Material.Get();
+		Landscape->SetActorRelativeScale3D(UISettings->NewLandscape_Scale);
+
+		// automatically calculate a lighting LOD that won't crash lightmass (hopefully)
+		// < 2048x2048 -> LOD0
+		// >=2048x2048 -> LOD1
+		// >= 4096x4096 -> LOD2
+		// >= 8192x8192 -> LOD3
+		Landscape->StaticLightingLOD = FMath::DivideAndRoundUp(FMath::CeilLogTwo((SizeX * SizeY) / (2048 * 2048) + 1), (uint32)2);
+
+		FString ReimportHeightmapFilePath;
+		if (LandscapeEdMode->NewLandscapePreviewMode == ENewLandscapePreviewMode::ImportLandscape)
 		{
-			UE_LOG(LogLandscapeTools, Error, TEXT("Unable to save current level"));
-			return FReply::Handled();
+			ReimportHeightmapFilePath = UISettings->ImportLandscape_HeightmapFilename;
 		}
-	}
-
-	const int32 QuadsPerSection = UISettings->NewLandscape_QuadsPerSection;
-	
-	const FIntPoint TotalLandscapeComponentSize { UISettings->NewLandscape_ComponentCount.X, UISettings->NewLandscape_ComponentCount.Y };
-
-	const int32 ComponentCountX = bNeedsLandscapeRegions ? UISettings->WorldPartitionRegionSize : TotalLandscapeComponentSize.X;
-	const int32 ComponentCountY = bNeedsLandscapeRegions ? UISettings->WorldPartitionRegionSize : TotalLandscapeComponentSize.Y;
-	const int32 QuadsPerComponent = UISettings->NewLandscape_SectionsPerComponent * QuadsPerSection;
-	const int32 SizeX = ComponentCountX * QuadsPerComponent + 1;
-	const int32 SizeY = ComponentCountY * QuadsPerComponent + 1;
 
-	TArray<FLandscapeImportLayerInfo> MaterialImportLayers;
-	ELandscapeImportResult LayerImportResult = bIsNewLandscape  ? UISettings->CreateNewLayersInfo(MaterialImportLayers) : UISettings->CreateImportLayersInfo(MaterialImportLayers);
+		Landscape->Import(FGuid::NewGuid(), 0, 0, SizeX - 1, SizeY - 1, UISettings->NewLandscape_SectionsPerComponent, UISettings->NewLandscape_QuadsPerSection, HeightDataPerLayers, *ReimportHeightmapFilePath, MaterialLayerDataPerLayers, UISettings->ImportLandscape_AlphamapType);
 
-	if (LayerImportResult == ELandscapeImportResult::Error)
-	{
-		UE_LOG(LogLandscapeTools, Error, TEXT("Unable to import weight maps"));
-		return FReply::Handled();
-	}
+		ULandscapeInfo* LandscapeInfo = Landscape->GetLandscapeInfo();
+		check(LandscapeInfo);
 
-	TMap<FGuid, TArray<uint16>> HeightDataPerLayers;
-	TMap<FGuid, TArray<FLandscapeImportLayerInfo>> MaterialLayerDataPerLayers;
-
-	TArray<uint16> OutHeightData;
-	
-	if (bIsNewLandscape || bNeedsLandscapeRegions)
-	{
-		UISettings->InitializeDefaultHeightData(OutHeightData);
-	}
-	else
-	{
-		UISettings->ExpandImportData(OutHeightData, MaterialImportLayers);
-	}
+		FActorLabelUtilities::SetActorLabelUnique(Landscape, ALandscape::StaticClass()->GetName());
 
-	HeightDataPerLayers.Add(FGuid(), OutHeightData);
-	// ComputeHeightData will also modify/expand material layers data, which is why we create MaterialLayerDataPerLayers after calling ComputeHeightData
-	MaterialLayerDataPerLayers.Add(FGuid(), MoveTemp(MaterialImportLayers));
+		LandscapeInfo->UpdateLayerInfoMap(Landscape);
 
-	FScopedTransaction Transaction(bIsNewLandscape ? LOCTEXT("Undo_CreateLandscape", "Creating New Landscape") : LOCTEXT("Undo_CreateAndImportLandscape", "Creating and Importing New Landscape"));
-
-	const FVector Offset = FTransform(UISettings->NewLandscape_Rotation, FVector::ZeroVector, UISettings->NewLandscape_Scale).TransformVector(FVector(-UISettings->NewLandscape_ComponentCount.X * QuadsPerComponent / 2, -UISettings->NewLandscape_ComponentCount.Y * QuadsPerComponent / 2, 0));
-	
-	ALandscape* Landscape = LandscapeEdMode->GetWorld()->SpawnActor<ALandscape>(UISettings->NewLandscape_Location + Offset, UISettings->NewLandscape_Rotation);
-	Landscape->bCanHaveLayersContent = UISettings->bCanHaveLayersContent;
-	Landscape->LandscapeMaterial = UISettings->NewLandscape_Material.Get();
-	Landscape->SetActorRelativeScale3D(UISettings->NewLandscape_Scale);
-
-	// automatically calculate a lighting LOD that won't crash lightmass (hopefully)
-	// < 2048x2048 -> LOD0
-	// >=2048x2048 -> LOD1
-	// >= 4096x4096 -> LOD2
-	// >= 8192x8192 -> LOD3
-	Landscape->StaticLightingLOD = FMath::DivideAndRoundUp(FMath::CeilLogTwo((SizeX * SizeY) / (2048 * 2048) + 1), (uint32)2);
-
-	FString ReimportHeightmapFilePath;
-	if (LandscapeEdMode->NewLandscapePreviewMode == ENewLandscapePreviewMode::ImportLandscape)
-	{
-		ReimportHeightmapFilePath = UISettings->ImportLandscape_HeightmapFilename;
-	}
-
-	Landscape->Import(FGuid::NewGuid(), 0, 0, SizeX - 1, SizeY - 1, UISettings->NewLandscape_SectionsPerComponent, QuadsPerSection, HeightDataPerLayers, *ReimportHeightmapFilePath, MaterialLayerDataPerLayers, UISettings->ImportLandscape_AlphamapType);
-
-	ULandscapeInfo* LandscapeInfo = Landscape->GetLandscapeInfo();
-	check(LandscapeInfo);
-
-	FActorLabelUtilities::SetActorLabelUnique(Landscape, ALandscape::StaticClass()->GetName());
-
-	LandscapeInfo->UpdateLayerInfoMap(Landscape);
-
-	// Import doesn't fill in the LayerInfo for layers with no data, do that now
-	const TArray<FLandscapeImportLayer>& ImportLandscapeLayersList = UISettings->ImportLandscape_Layers;
-	const ULandscapeSettings* Settings = GetDefault<ULandscapeSettings>();
-	TSoftObjectPtr<ULandscapeLayerInfoObject> DefaultLayerInfoObject = Settings->GetDefaultLayerInfoObject().LoadSynchronous();
-
-	for (int32 i = 0; i < ImportLandscapeLayersList.Num(); i++)
-	{
-		ULandscapeLayerInfoObject* LayerInfo = ImportLandscapeLayersList[i].LayerInfo;
-		FName LayerName = ImportLandscapeLayersList[i].LayerName;
-
-		// If DefaultLayerInfoObject is set and LayerInfo does not exist, we will try to create the new LayerInfo by cloning DefaultLayerInfoObject. Except for VisibilityLayer which doesn't require an asset.
-		if (DefaultLayerInfoObject.IsValid() && (LayerInfo == nullptr) && (LayerName != ALandscapeProxy::VisibilityLayer->LayerName))
-		{
-			LayerInfo = Landscape->CreateLayerInfo(*LayerName.ToString(), DefaultLayerInfoObject.Get());
-
-			if (LayerInfo != nullptr)
-			{
-				LayerInfo->LayerUsageDebugColor = LayerInfo->GenerateLayerUsageDebugColor();
-				LayerInfo->MarkPackageDirty();
-			}
-		}
+		// Import doesn't fill in the LayerInfo for layers with no data, do that now
+		const TArray<FLandscapeImportLayer>& ImportLandscapeLayersList = UISettings->ImportLandscape_Layers;
+		const ULandscapeSettings* Settings = GetDefault<ULandscapeSettings>();
+		TSoftObjectPtr<ULandscapeLayerInfoObject> DefaultLayerInfoObject = Settings->GetDefaultLayerInfoObject().LoadSynchronous();
 
-		if (LayerInfo != nullptr)
+		for (int32 i = 0; i < ImportLandscapeLayersList.Num(); i++)
 		{
-			if (LandscapeEdMode->NewLandscapePreviewMode == ENewLandscapePreviewMode::ImportLandscape)
-			{
-				Landscape->EditorLayerSettings.Add(FLandscapeEditorLayerSettings(LayerInfo, ImportLandscapeLayersList[i].SourceFilePath));
-			}
-			else
-			{
-				Landscape->EditorLayerSettings.Add(FLandscapeEditorLayerSettings(LayerInfo));
-			}
+			ULandscapeLayerInfoObject* LayerInfo = ImportLandscapeLayersList[i].LayerInfo;
+			FName LayerName = ImportLandscapeLayersList[i].LayerName;
 
-			int32 LayerInfoIndex = LandscapeInfo->GetLayerInfoIndex(ImportLandscapeLayersList[i].LayerName);
-			if (ensure(LayerInfoIndex != INDEX_NONE))
+			// If DefaultLayerInfoObject is set and LayerInfo does not exist, we will try to create the new LayerInfo by cloning DefaultLayerInfoObject. Except for VisibilityLayer which doesn't require an asset.
+			if (DefaultLayerInfoObject.IsValid() && (LayerInfo == nullptr) && (LayerName != ALandscapeProxy::VisibilityLayer->LayerName))
 			{
-				FLandscapeInfoLayerSettings& LayerSettings = LandscapeInfo->Layers[LayerInfoIndex];
-				LayerSettings.LayerInfoObj = LayerInfo;
-			}
-		}
-	}
-
-	LandscapeEdMode->UpdateLandscapeList();
-	LandscapeEdMode->SetLandscapeInfo(LandscapeInfo);
-	LandscapeEdMode->CurrentToolTarget.TargetType = ELandscapeToolTargetType::Heightmap;
-	LandscapeEdMode->SetCurrentTargetLayer(NAME_None, nullptr);
-	LandscapeEdMode->SetCurrentTool("Select"); // change tool so switching back to the manage mode doesn't give "New Landscape" again
-	LandscapeEdMode->SetCurrentTool("Sculpt"); // change to sculpting mode and tool
-	LandscapeEdMode->SetCurrentLayer(0);
+				LayerInfo = Landscape->CreateLayerInfo(*LayerName.ToString(), DefaultLayerInfoObject.Get());
 
-	World->GetSubsystem<ULandscapeSubsystem>()->ChangeGridSize(LandscapeInfo, UISettings->WorldPartitionGridSize);
-
-	if (LandscapeEdMode->CurrentToolTarget.LandscapeInfo.IsValid())
-	{
-		ALandscapeProxy* LandscapeProxy = LandscapeEdMode->CurrentToolTarget.LandscapeInfo->GetLandscapeProxy();
-		LandscapeProxy->OnMaterialChangedDelegate().AddRaw(LandscapeEdMode, &FEdModeLandscape::OnLandscapeMaterialChangedDelegate);
-	}
-
-	UWorldPartition* WorldPartition = World->GetWorldPartition();
-	ULandscapeSubsystem* LandscapeSubsystem = World->GetSubsystem<ULandscapeSubsystem>();
-	ALandscapeProxy* LandscapeProxy = LandscapeEdMode->CurrentToolTarget.LandscapeInfo->GetLandscapeProxy();
-	
-	if (bNeedsLandscapeRegions)
-	{
-		ULevel* Level = LandscapeProxy->GetLevel();
-		UPackage* LevelPackage = Level->GetPackage();
-
-		TArray<FIntPoint> NewComponents;
-		NewComponents.Empty(TotalLandscapeComponentSize.X * TotalLandscapeComponentSize.Y);
-		for (int32 Y = 0; Y < TotalLandscapeComponentSize.Y; Y++)
-		{
-			for (int32 X = 0; X < TotalLandscapeComponentSize.X; X++)
-			{
-				NewComponents.Add(FIntPoint(X, Y));
+				if (LayerInfo != nullptr)
+				{
+					LayerInfo->LayerUsageDebugColor = LayerInfo->GenerateLayerUsageDebugColor();
+					LayerInfo->MarkPackageDirty();
+				}
 			}
-		}
-
-		int32 NumRegions = FMath::DivideAndRoundUp(TotalLandscapeComponentSize.X, static_cast<int32>(UISettings->WorldPartitionRegionSize)) * FMath::DivideAndRoundUp(TotalLandscapeComponentSize.Y, static_cast<int32>(UISettings->WorldPartitionRegionSize));
-
-
-		FScopedSlowTask Progress(NumRegions, LOCTEXT("CreateLandscapeRegions", "Creating Landscape Editor Regions..."));
-		Progress.MakeDialog();
 
-		TArray<ALocationVolume*> RegionVolumes;
-		FBox LandscapeBounds;
-		auto AddComponentsToRegion = [&Progress, NumRegions, bIsNewLandscape, LandscapeEdMode, WorldPartition, World, LandscapeProxy, &UISettings, QuadsPerSection, LandscapeInfo, LandscapeSubsystem, &RegionVolumes, &LandscapeBounds,&MaterialLayerDataPerLayers](const FIntPoint& RegionCoordinate, const TArray<FIntPoint>& NewComponents)
-		{
-			TRACE_CPUPROFILER_EVENT_SCOPE(AddComponentsToRegion);
-						
-			// Create a LocationVolume around the region 
-			int32 RegionSizeXTexels = QuadsPerSection * UISettings->WorldPartitionRegionSize;
-			int32 RegionSizeYTexels = QuadsPerSection * UISettings->WorldPartitionRegionSize;
-
-			double RegionSizeX = QuadsPerSection * LandscapeProxy->GetActorScale3D().X * UISettings->WorldPartitionRegionSize;
-			double RegionSizeY = QuadsPerSection * LandscapeProxy->GetActorScale3D().Y * UISettings->WorldPartitionRegionSize;
-			ALocationVolume* RegionVolume = LandscapeRegionUtils::CreateLandscapeRegionVolume(World, LandscapeProxy, RegionCoordinate, RegionSizeX);
-			RegionVolumes.Add(RegionVolume);
-			
-			TArray<ALandscapeProxy*> CreatedStreamingProxies;
-			AddComponents(LandscapeInfo, LandscapeSubsystem, NewComponents, CreatedStreamingProxies);
-			
-			FIntRect ImportRegion (RegionSizeXTexels * RegionCoordinate.X, RegionSizeYTexels* RegionCoordinate.Y, RegionSizeXTexels * (RegionCoordinate.X +1) + 1  , RegionSizeYTexels * (RegionCoordinate.Y +1) +1 );
-			FIntPoint ImportOffset(0, 0);
-
-			if (!bIsNewLandscape && !UISettings->ImportLandscape_HeightmapFilename.IsEmpty())
+			if (LayerInfo != nullptr)
 			{
-				LandscapeEdMode->ImportHeightData(LandscapeInfo, LandscapeEdMode->GetCurrentLayerGuid(), UISettings->ImportLandscape_HeightmapFilename, ImportRegion, ELandscapeImportTransformType::Subregion, ImportOffset, ELandscapeLayerPaintingRestriction::None, LandscapeEdMode->UISettings->bFlipYAxis);
-
-				TArray<FLandscapeImportLayerInfo>& Weights = *MaterialLayerDataPerLayers.Find(FGuid());
+				if (LandscapeEdMode->NewLandscapePreviewMode == ENewLandscapePreviewMode::ImportLandscape)
+				{
+					Landscape->EditorLayerSettings.Add(FLandscapeEditorLayerSettings(LayerInfo, ImportLandscapeLayersList[i].SourceFilePath));
+				}
+				else
+				{
+					Landscape->EditorLayerSettings.Add(FLandscapeEditorLayerSettings(LayerInfo));
+				}
 
-				for (const TPair<FGuid, TArray<FLandscapeImportLayerInfo>>& Layer : MaterialLayerDataPerLayers)
+				int32 LayerInfoIndex = LandscapeInfo->GetLayerInfoIndex(ImportLandscapeLayersList[i].LayerName);
+				if (ensure(LayerInfoIndex != INDEX_NONE))
 				{
-					for (const FLandscapeImportLayerInfo& WeightMap : Layer.Value)
-					{						
-						LandscapeEdMode->ImportWeightData(LandscapeInfo, LandscapeEdMode->GetCurrentLayerGuid(), WeightMap.LayerInfo, WeightMap.SourceFilePath, ImportRegion, ELandscapeImportTransformType::Subregion, ImportOffset, ELandscapeLayerPaintingRestriction::None, LandscapeEdMode->UISettings->bFlipYAxis);
-					}
+					FLandscapeInfoLayerSettings& LayerSettings = LandscapeInfo->Layers[LayerInfoIndex];
+					LayerSettings.LayerInfoObj = LayerInfo;
 				}
 			}
+		}
 
-			// ensures all the final height textures have been updated.
-			LandscapeInfo->ForceLayersFullUpdate();
-			LandscapeEditorUtils::SaveLandscapeProxies(MakeArrayView(CreatedStreamingProxies), WorldPartition);
-			LandscapeBounds += LandscapeInfo->GetCompleteBounds();
-
-			Progress.EnterProgressFrame(1.0f , FText::Format(LOCTEXT("LandscapeCreateRegion", "Creating Landscape Editor Regions ({0}, {1})"), RegionCoordinate.X, RegionCoordinate.Y));
-			return true;
-		};
-
-		LandscapeEditorUtils::SaveObjects(TArrayView<ALandscape*>(TArray<ALandscape*> { LandscapeInfo->LandscapeActor.Get() }));
+		LandscapeEdMode->UpdateLandscapeList();
+		LandscapeEdMode->SetLandscapeInfo(LandscapeInfo);
+		LandscapeEdMode->CurrentToolTarget.TargetType = ELandscapeToolTargetType::Heightmap;
+		LandscapeEdMode->SetCurrentTargetLayer(NAME_None, nullptr);
+		LandscapeEdMode->SetCurrentTool("Select"); // change tool so switching back to the manage mode doesn't give "New Landscape" again
+		LandscapeEdMode->SetCurrentTool("Sculpt"); // change to sculpting mode and tool
+		LandscapeEdMode->SetCurrentLayer(0);
 
-		LandscapeRegionUtils::ForEachComponentByRegion(UISettings->WorldPartitionRegionSize, NewComponents, AddComponentsToRegion);
+		LandscapeEdMode->GetWorld()->GetSubsystem<ULandscapeSubsystem>()->ChangeGridSize(LandscapeInfo, UISettings->WorldPartitionGridSize);
 
-		// update the zcomponent of the volumes 
-		const float ZScale =  LandscapeBounds.Max.Z - LandscapeBounds.Min.Z;
-		for (ALocationVolume* RegionVolume : RegionVolumes)
+		if (LandscapeEdMode->CurrentToolTarget.LandscapeInfo.IsValid())
 		{
-			const FVector Scale = RegionVolume->GetActorScale();
-			const FVector NewScale{ Scale.X, Scale.Y, 1000000.0f}; //ZScale * 10.0 
-			RegionVolume->SetActorScale3D(NewScale);
+			ALandscapeProxy* LandscapeProxy = LandscapeEdMode->CurrentToolTarget.LandscapeInfo->GetLandscapeProxy();
+			LandscapeProxy->OnMaterialChangedDelegate().AddRaw(LandscapeEdMode, &FEdModeLandscape::OnLandscapeMaterialChangedDelegate);
 		}
-
-		LandscapeEditorUtils::SaveObjects(MakeArrayView(RegionVolumes));
-
-		TArray<ALandscapeProxy*> AllProxies;
-		
-		// save the initial region & unload it
-		LandscapeInfo->ForEachLandscapeProxy([&WorldPartition, &AllProxies](ALandscapeProxy* Proxy) {
-			if (Proxy->IsA<ALandscapeStreamingProxy>())
-			{		
-				AllProxies.Add(Proxy);
-				
-			}
-			return true;
-		});
-
-		LandscapeEditorUtils::SaveLandscapeProxies(MakeArrayView(AllProxies), WorldPartition);
 	}
 
 	return FReply::Handled();
@@ -1309,7 +1059,8 @@ bool FLandscapeEditorDetailCustomization_NewLandscape::GetImportButtonIsEnabled(
 	if (LandscapeEdMode != nullptr)
 	{
 		if (LandscapeEdMode->UISettings->ImportLandscape_HeightmapImportResult == ELandscapeImportResult::Error ||
-			LandscapeEdMode->UISettings->ImportLandscape_HeightmapFilename.IsEmpty()) 
+			LandscapeEdMode->UISettings->ImportLandscape_HeightmapFilename.IsEmpty() ||
+			LandscapeEdMode->UISettings->GetImportLandscapeData().IsEmpty())
 		{
 			return false;
 		}
@@ -1396,15 +1147,28 @@ FReply FLandscapeEditorDetailCustomization_NewLandscape::OnImportHeightmapFilena
 	FEdModeLandscape* LandscapeEdMode = GetEditorMode();
 	check(LandscapeEdMode != nullptr);
 
-	ILandscapeEditorModule& LandscapeEditorModule = FModuleManager::GetModuleChecked<ILandscapeEditorModule>("LandscapeEditor");
-	const FString FileTypes = LandscapeEditorModule.GetHeightmapImportDialogTypeString();
-
-	TOptional<FString> OptionalFilename = LandscapeEditorUtils::GetImportExportFilename(NSLOCTEXT("UnrealEd", "Import", "Import").ToString(), LandscapeEdMode->UISettings->LastImportPath, FileTypes);
-	if (OptionalFilename.IsSet())
+	// Prompt the user for the Filenames
+	IDesktopPlatform* DesktopPlatform = FDesktopPlatformModule::Get();
+	if (DesktopPlatform != nullptr)
 	{
-		const FString& Filename = OptionalFilename.GetValue();
-		ensure(PropertyHandle_HeightmapFilename->SetValue(Filename) == FPropertyAccess::Success);
-		LandscapeEdMode->UISettings->LastImportPath = FPaths::GetPath(Filename);
+		ILandscapeEditorModule& LandscapeEditorModule = FModuleManager::GetModuleChecked<ILandscapeEditorModule>("LandscapeEditor");
+		const TCHAR* FileTypes = LandscapeEditorModule.GetHeightmapImportDialogTypeString();
+
+		TArray<FString> OpenFilenames;
+		bool bOpened = DesktopPlatform->OpenFileDialog(
+			FSlateApplication::Get().FindBestParentWindowHandleForDialogs(nullptr),
+			NSLOCTEXT("UnrealEd", "Import", "Import").ToString(),
+			LandscapeEdMode->UISettings->LastImportPath,
+			TEXT(""),
+			FileTypes,
+			EFileDialogFlags::None,
+			OpenFilenames);
+
+		if (bOpened)
+		{
+			ensure(PropertyHandle_HeightmapFilename->SetValue(OpenFilenames[0]) == FPropertyAccess::Success);
+			LandscapeEdMode->UISettings->LastImportPath = FPaths::GetPath(OpenFilenames[0]);
+		}
 	}
 
 	return FReply::Handled();
diff --git a/Engine/Source/Editor/LandscapeEditor/Private/LandscapeEditorDetailCustomization_NewLandscape.h b/Engine/Source/Editor/LandscapeEditor/Private/LandscapeEditorDetailCustomization_NewLandscape.h
--- a/Engine/Source/Editor/LandscapeEditor/Private/LandscapeEditorDetailCustomization_NewLandscape.h
+++ b/Engine/Source/Editor/LandscapeEditor/Private/LandscapeEditorDetailCustomization_NewLandscape.h
@@ -13,7 +13,6 @@
 #include "IPropertyTypeCustomization.h"
 #include "PropertyHandle.h"
 #include "LandscapeEditorDetailCustomization_Base.h"
-#include "LandscapeSubsystem.h"
 
 class IDetailLayoutBuilder;
 
@@ -94,9 +93,5 @@ public:
 	bool ShouldShowResetMaterialToDefault(TSharedPtr<IPropertyHandle> InPropertyHandle);
 
 protected:
-
 	bool bUsingSlider;
-
-private:
-	static void AddComponents(ULandscapeInfo* InLandscapeInfo, ULandscapeSubsystem* InLandscapeSubsystem, const TArray<FIntPoint>& InComponentCoordinates, TArray<ALandscapeProxy*>& OutCreatedStreamingProxies);
 };
diff --git a/Engine/Source/Editor/LandscapeEditor/Private/LandscapeEditorModule.cpp b/Engine/Source/Editor/LandscapeEditor/Private/LandscapeEditorModule.cpp
--- a/Engine/Source/Editor/LandscapeEditor/Private/LandscapeEditorModule.cpp
+++ b/Engine/Source/Editor/LandscapeEditor/Private/LandscapeEditorModule.cpp
@@ -20,7 +20,6 @@
 #include "LandscapeFileFormatRaw.h"
 #include "Settings/EditorExperimentalSettings.h"
 #include "LandscapeEditorServices.h"
-#include "LandscapeImageFileCache.h"
 
 #include "Framework/MultiBox/MultiBoxBuilder.h"
 #include "PropertyEditorModule.h"
@@ -162,8 +161,6 @@ public:
 		ILandscapeModule& LandscapeModule = FModuleManager::GetModuleChecked<ILandscapeModule>("Landscape");
 		LandscapeEditorServices.Reset(new FLandscapeEditorServices);
 		LandscapeModule.SetLandscapeEditorServices(LandscapeEditorServices.Get());
-
-		LandscapeImageFileCache.Reset(new FLandscapeImageFileCache());
 	}
 
 	/**
@@ -201,7 +198,6 @@ public:
 			LandscapeModule.SetLandscapeEditorServices(nullptr);
 		}
 		LandscapeEditorServices.Reset();
-		LandscapeImageFileCache.Reset();
 	}
 
 	static void ConstructLandscapeViewportMenu(FMenuBuilder& MenuBuilder)
@@ -375,8 +371,6 @@ public:
 
 	virtual TSharedPtr<FUICommandList> GetLandscapeLevelViewportCommandList() const override;
 
-	FLandscapeImageFileCache& GetImageFileCache() const override;
-
 protected:
 	TSharedPtr<FExtender> ViewportMenuExtender;
 	TSharedPtr<FUICommandList> GlobalUICommandList;
@@ -387,7 +381,6 @@ protected:
 	mutable FString HeightmapExportDialogTypeString;
 	mutable FString WeightmapExportDialogTypeString;
 	TUniquePtr<ILandscapeEditorServices> LandscapeEditorServices;
-	TUniquePtr<FLandscapeImageFileCache> LandscapeImageFileCache;
 };
 
 IMPLEMENT_MODULE(FLandscapeEditorModule, LandscapeEditor);
@@ -556,11 +549,4 @@ TSharedPtr<FUICommandList> FLandscapeEditorModule::GetLandscapeLevelViewportComm
 	return GlobalUICommandList;
 }
 
-FLandscapeImageFileCache& FLandscapeEditorModule::GetImageFileCache() const
-{
-	check(LandscapeImageFileCache != nullptr);
-	return *LandscapeImageFileCache;
-}
-
-
 #undef LOCTEXT_NAMESPACE
diff --git a/Engine/Source/Editor/LandscapeEditor/Private/LandscapeEditorObject.cpp b/Engine/Source/Editor/LandscapeEditor/Private/LandscapeEditorObject.cpp
--- a/Engine/Source/Editor/LandscapeEditor/Private/LandscapeEditorObject.cpp
+++ b/Engine/Source/Editor/LandscapeEditor/Private/LandscapeEditorObject.cpp
@@ -11,7 +11,6 @@
 #include "LandscapeRender.h"
 #include "LandscapeSettings.h"
 #include "LandscapeImportHelper.h"
-#include "LandscapeTiledImage.h"
 #include "LandscapeMaterialInstanceConstant.h"
 #include "Misc/ConfigCacheIni.h"
 #include "EngineUtils.h"
@@ -156,11 +155,7 @@ void ULandscapeEditorObject::PostEditChangeProperty(FPropertyChangedEvent& Prope
 		PropertyChangedEvent.MemberProperty->GetFName() == GET_MEMBER_NAME_CHECKED(ULandscapeEditorObject, NewLandscape_SectionsPerComponent) ||
 		PropertyChangedEvent.MemberProperty->GetFName() == GET_MEMBER_NAME_CHECKED(ULandscapeEditorObject, NewLandscape_ComponentCount))
 	{
-		// Only clamp the landscape size if we're not in World Partition
-		if (!ParentMode->IsGridBased())
-		{
-			NewLandscape_ClampSize();
-		}
+		NewLandscape_ClampSize();
 	}
 
 	if (PropertyChangedEvent.MemberProperty == nullptr ||
@@ -642,8 +637,6 @@ bool ULandscapeEditorObject::UseSingleFileImport() const
 
 void ULandscapeEditorObject::RefreshImports()
 {
-	TRACE_CPUPROFILER_EVENT_SCOPE(ULandscapeEditorObject::RefreshImports);
-
 	ClearImportLandscapeData();
 	HeightmapImportDescriptorIndex = 0;
 	HeightmapImportDescriptor.Reset();
@@ -655,12 +648,12 @@ void ULandscapeEditorObject::RefreshImports()
 
 	if (!ImportLandscape_HeightmapFilename.IsEmpty())
 	{
-		FLandscapeTiledImage TiledImage;
-		
-		if (FLandscapeFileInfo FileInfo = TiledImage.Load(*ImportLandscape_HeightmapFilename); FileInfo.ResultCode == ELandscapeImportResult::Success)
+		ImportLandscape_HeightmapImportResult =
+			FLandscapeImportHelper::GetHeightmapImportDescriptor(ImportLandscape_HeightmapFilename, UseSingleFileImport(), bFlipYAxis, HeightmapImportDescriptor, ImportLandscape_HeightmapErrorMessage);
+		if (ImportLandscape_HeightmapImportResult != ELandscapeImportResult::Error)
 		{
-			ImportLandscape_Width = TiledImage.GetResolution().X;
-			ImportLandscape_Height = TiledImage.GetResolution().Y;
+			ImportLandscape_Width = HeightmapImportDescriptor.ImportResolutions[HeightmapImportDescriptorIndex].Width;
+			ImportLandscape_Height = HeightmapImportDescriptor.ImportResolutions[HeightmapImportDescriptorIndex].Height;
 			ChooseBestComponentSizeForImport();
 			ImportLandscapeData();
 		}
@@ -704,16 +697,27 @@ void ULandscapeEditorObject::RefreshLayerImport(FLandscapeImportLayer& ImportLay
 		}
 		else
 		{
-			FLandscapeTiledImage TiledImage;
-			FLandscapeFileInfo FileInfo = TiledImage.Load(*ImportLayer.SourceFilePath);
-			ImportLayer.ImportResult = FileInfo.ResultCode;
-			ImportLayer.ErrorMessage = FileInfo.ErrorMessage;
-			if (FileInfo.ResultCode == ELandscapeImportResult::Success)
+			ImportLayer.ImportResult = FLandscapeImportHelper::GetWeightmapImportDescriptor(ImportLayer.SourceFilePath, UseSingleFileImport(), bFlipYAxis, ImportLayer.LayerName, ImportLayer.ImportDescriptor, ImportLayer.ErrorMessage);
+			if (ImportLayer.ImportResult != ELandscapeImportResult::Error)
 			{
-				if (FileInfo.PossibleResolutions[0] != FLandscapeFileResolution(ImportLandscape_Height, ImportLandscape_Width))
+				if (ImportLandscape_Height != 0 || ImportLandscape_Width != 0)
 				{
-					ImportLayer.ImportResult = ELandscapeImportResult::Error;
-					ImportLayer.ErrorMessage = NSLOCTEXT("LandscapeEditor.ImportLandscape", "Import_WeightHeightResolutionMismatch", "Weightmap import resolution isn't same as Heightmap resolution.");
+					// Use same import index as Heightmap
+					int32 FoundIndex = INDEX_NONE;
+					for (int32 Index = 0; Index < ImportLayer.ImportDescriptor.FileResolutions.Num(); ++Index)
+					{
+						if (ImportLayer.ImportDescriptor.ImportResolutions[Index] == HeightmapImportDescriptor.ImportResolutions[HeightmapImportDescriptorIndex])
+						{
+							FoundIndex = Index;
+							break;
+						}
+					}
+
+					if (FoundIndex == INDEX_NONE)
+					{
+						ImportLayer.ImportResult = ELandscapeImportResult::Error;
+						ImportLayer.ErrorMessage = NSLOCTEXT("LandscapeEditor.ImportLandscape", "Import_WeightHeightResolutionMismatch", "Weightmap import resolution isn't same as Heightmap resolution.");
+					}
 				}
 			}
 		}
@@ -733,15 +737,11 @@ void ULandscapeEditorObject::OnChangeImportLandscapeResolution(int32 DescriptorI
 
 void ULandscapeEditorObject::ImportLandscapeData()
 {
-	FLandscapeTiledImage TiledImage;
-	FLandscapeFileInfo FileInfo = TiledImage.Load(*ImportLandscape_HeightmapFilename);
-
-	if (FileInfo.ResultCode != ELandscapeImportResult::Success)
+	ImportLandscape_HeightmapImportResult = FLandscapeImportHelper::GetHeightmapImportData(HeightmapImportDescriptor, HeightmapImportDescriptorIndex, ImportLandscape_Data, ImportLandscape_HeightmapErrorMessage);
+	if (ImportLandscape_HeightmapImportResult == ELandscapeImportResult::Error)
 	{
 		ImportLandscape_Data.Empty();
 	}
-
-	TiledImage.Read(ImportLandscape_Data, bFlipYAxis);
 }
 
 ELandscapeImportResult ULandscapeEditorObject::CreateImportLayersInfo(TArray<FLandscapeImportLayerInfo>& OutImportLayerInfos)
@@ -765,18 +765,38 @@ ELandscapeImportResult ULandscapeEditorObject::CreateImportLayersInfo(TArray<FLa
 
 		if (ImportLayer.LayerInfo != nullptr && !ImportLayer.SourceFilePath.IsEmpty())
 		{
-			FLandscapeTiledImage LayerImage;
-			FLandscapeFileInfo LayerFileInfo = LayerImage.Load(*ImportLayer.SourceFilePath);
+			UIImportLayer.ImportResult = FLandscapeImportHelper::GetWeightmapImportDescriptor(ImportLayer.SourceFilePath, UseSingleFileImport(), bFlipYAxis, ImportLayer.LayerName, UIImportLayer.ImportDescriptor, UIImportLayer.ErrorMessage);
+			if (UIImportLayer.ImportResult == ELandscapeImportResult::Error)
+			{
+				FMessageDialog::Open(EAppMsgType::Ok, UIImportLayer.ErrorMessage);
+				return ELandscapeImportResult::Error;
+			}
 
-			UIImportLayer.ImportResult = LayerFileInfo.ResultCode;
+			// Use same import index as Heightmap
+			int32 FoundIndex = INDEX_NONE;
+			for (int32 Index = 0; Index < UIImportLayer.ImportDescriptor.FileResolutions.Num(); ++Index)
+			{
+				if (UIImportLayer.ImportDescriptor.ImportResolutions[Index] == HeightmapImportDescriptor.ImportResolutions[HeightmapImportDescriptorIndex])
+				{
+					FoundIndex = Index;
+					break;
+				}
+			}
 
-			if (UIImportLayer.ImportResult == ELandscapeImportResult::Error)
+			if (FoundIndex == INDEX_NONE)
 			{
+				UIImportLayer.ImportResult = ELandscapeImportResult::Error;
+				UIImportLayer.ErrorMessage = NSLOCTEXT("LandscapeEditor.NewLandscape", "Import_WeightHeightResolutionMismatch", "Weightmap import resolution isn't same as Heightmap resolution.");
 				FMessageDialog::Open(EAppMsgType::Ok, UIImportLayer.ErrorMessage);
 				return ELandscapeImportResult::Error;
 			}
 
-			LayerImage.Read(ImportLayer.LayerData);
+			UIImportLayer.ImportResult = FLandscapeImportHelper::GetWeightmapImportData(UIImportLayer.ImportDescriptor, FoundIndex, ImportLayer.LayerName, ImportLayer.LayerData, UIImportLayer.ErrorMessage);
+			if (UIImportLayer.ImportResult == ELandscapeImportResult::Error)
+			{
+				FMessageDialog::Open(EAppMsgType::Ok, UIImportLayer.ErrorMessage);
+				return ELandscapeImportResult::Error;
+			}
 		}
 	}
 
diff --git a/Engine/Source/Editor/LandscapeEditor/Private/LandscapeEditorUtils.cpp b/Engine/Source/Editor/LandscapeEditor/Private/LandscapeEditorUtils.cpp
deleted file mode 100644
diff --git a/Engine/Source/Editor/LandscapeEditor/Private/LandscapeFileFormatRaw.cpp b/Engine/Source/Editor/LandscapeEditor/Private/LandscapeFileFormatRaw.cpp
--- a/Engine/Source/Editor/LandscapeEditor/Private/LandscapeFileFormatRaw.cpp
+++ b/Engine/Source/Editor/LandscapeEditor/Private/LandscapeFileFormatRaw.cpp
@@ -13,80 +13,37 @@
 #include "Misc/FileHelper.h"
 #include "Templates/UnrealTemplate.h"
 
-#include "Serialization/JsonReader.h"
-#include "Serialization/JsonWriter.h"
-#include "Serialization/JsonSerializer.h"
-
 #define LOCTEXT_NAMESPACE "LandscapeEditor.NewLandscape"
 
-namespace
+TArray<FLandscapeFileResolution> CalculatePossibleRawResolutions(int64 FileSize)
 {
-	void WriteMetaData(const TCHAR* InMapFilename, FLandscapeFileResolution InDataResolution, int32 InBitsPerPixel)
-	{
-		FString RawMetadataFilename = FPaths::SetExtension(InMapFilename, ".json");
-		FString JsonStr;
-		
-		TSharedRef<TJsonWriter<TCHAR>> JsonWriter = TJsonWriterFactory<TCHAR>::Create(&JsonStr);
-		TSharedRef<FJsonObject> JsonObject = MakeShareable(new FJsonObject);;
-		JsonObject->SetField("width", MakeShared<FJsonValueNumber>(InDataResolution.Width));
-		JsonObject->SetField("height", MakeShared<FJsonValueNumber>(InDataResolution.Height));
-		JsonObject->SetField("bpp", MakeShared<FJsonValueNumber>(InBitsPerPixel));
+	TArray<FLandscapeFileResolution> PossibleResolutions;
 
-		if (FJsonSerializer::Serialize<TCHAR>(JsonObject, JsonWriter))
+	// Find all possible heightmap sizes, between 8 and 8192 width/height
+	const int32 MinWidth = FMath::Max(8, (int32)FMath::DivideAndRoundUp(FileSize, (int64)8192));
+	const int32 MaxWidth = FMath::TruncToInt(FMath::Sqrt(static_cast<double>(FileSize)));
+	for (int32 Width = MinWidth; Width <= MaxWidth; Width++)
+	{
+		if (FileSize % Width == 0)
 		{
-			FFileHelper::SaveStringToFile(JsonStr, *RawMetadataFilename);
+			FLandscapeFileResolution ImportResolution;
+			ImportResolution.Width = Width;
+			ImportResolution.Height = FileSize / Width;
+			PossibleResolutions.Add(ImportResolution);
 		}
 	}
-}
 
-bool GetRawResolution(const TCHAR* InFilename, FLandscapeFileResolution& OutResolution, int32& OutBitsPerPixel)
-{
-	FString Extension = FPaths::GetExtension(InFilename);
-	FString RawMetadataFilename = FPaths::SetExtension(InFilename, ".json");
-	
-	if (FPaths::FileExists(RawMetadataFilename))
+	for (int32 i = PossibleResolutions.Num() - 1; i >= 0; --i)
 	{
-		FString JsonStr;
-		FFileHelper::LoadFileToString(JsonStr, *RawMetadataFilename);
-		TSharedRef<TJsonReader<TCHAR>> JsonReader = TJsonReaderFactory<TCHAR>::Create(JsonStr);
-		TSharedPtr<FJsonObject> JsonObject;
-		if (FJsonSerializer::Deserialize(JsonReader, JsonObject))
+		FLandscapeFileResolution ImportResolution = PossibleResolutions[i];
+		if (ImportResolution.Width != ImportResolution.Height)
 		{
-			OutResolution.Width = JsonObject->GetIntegerField("width");
-			OutResolution.Height = JsonObject->GetIntegerField("height");
-
-			OutBitsPerPixel = JsonObject->GetIntegerField("bpp");
-			
-			return true;
+			Swap(ImportResolution.Width, ImportResolution.Height);
+			PossibleResolutions.Add(ImportResolution);
 		}
-		
-		return false;
-	}
-
-	const bool bIs16Bit = Extension == FString("r16");
-	const bool bIs8Bit = Extension == FString("r8");
-
-	if (!(bIs16Bit || bIs8Bit))
-	{
-		return false;
-	}
-
-	const int32 BytesPerPixel = bIs16Bit ? 2 : 1;
-	int64 FileSize = IFileManager::Get().FileSize(InFilename);
-	
-	const int32 NumPixels = FileSize / BytesPerPixel;
-	int32 Dimension = FMath::TruncToInt(FMath::Sqrt(static_cast<double>(NumPixels)));
-	if (Dimension * Dimension == NumPixels)
-	{
-		OutResolution.Width = Dimension;
-		OutResolution.Height = Dimension;
-
-		OutBitsPerPixel = BytesPerPixel * 8;
-
-		return true;
 	}
 
-	return false;
+	return PossibleResolutions;
 }
 
 FLandscapeHeightmapFileFormat_Raw::FLandscapeHeightmapFileFormat_Raw()
@@ -115,22 +72,9 @@ FLandscapeFileInfo FLandscapeHeightmapFileFormat_Raw::Validate(const TCHAR* Heig
 	}
 	else
 	{
-		FLandscapeFileResolution Resolution;
-		int32 BitsPerPixel;
-		if (GetRawResolution(HeightmapFilename, Resolution, BitsPerPixel))
-		{
-			if (BitsPerPixel != 16)
-			{
-				Result.ResultCode = ELandscapeImportResult::Error;
-				Result.ErrorMessage = LOCTEXT("Import_WeightmapFileBitsPerPixel", "Height file has an invalid number of bits per pixel");
-			}
-			else
-			{
-				Result.PossibleResolutions = { Resolution };
-			}
-			
-		}
-		else
+		Result.PossibleResolutions = CalculatePossibleRawResolutions(ImportFileSize / 2);
+
+		if (Result.PossibleResolutions.Num() == 0)
 		{
 			Result.ResultCode = ELandscapeImportResult::Error;
 			Result.ErrorMessage = LOCTEXT("Import_HeightmapFileInvalidSize", "The heightmap file has an invalid size (possibly not 16-bit?)");
@@ -172,15 +116,7 @@ void FLandscapeHeightmapFileFormat_Raw::Export(const TCHAR* HeightmapFilename, F
 	TempData.AddUninitialized(DataResolution.Width * DataResolution.Height * 2);
 	FMemory::Memcpy(TempData.GetData(), Data.GetData(), DataResolution.Width * DataResolution.Height * 2);
 
-	if (FFileHelper::SaveArrayToFile(TempData, HeightmapFilename))
-	{
-		FString Extension = FPaths::GetExtension(HeightmapFilename);
-
-		if (Extension == "raw")
-		{
-			WriteMetaData(HeightmapFilename, DataResolution, 16);
-		}
-	}
+	FFileHelper::SaveArrayToFile(TempData, HeightmapFilename);
 }
 
 //////////////////////////////////////////////////////////////////////////
@@ -206,22 +142,9 @@ FLandscapeFileInfo FLandscapeWeightmapFileFormat_Raw::Validate(const TCHAR* Weig
 	}
 	else
 	{
-		FLandscapeFileResolution Resolution;
-		int32 BitsPerPixel;
-		if (GetRawResolution(WeightmapFilename, Resolution, BitsPerPixel))
-		{
-			if (BitsPerPixel != 8)
-			{
-				Result.ResultCode = ELandscapeImportResult::Error;
-				Result.ErrorMessage = LOCTEXT("Import_WeightmapFileBitsPerPixel", "Weightmap file has an invalid number of bits per pixel");
-				
-			}
-			else
-			{
-				Result.PossibleResolutions = { Resolution };
-			}
-		}
-		else
+		Result.PossibleResolutions = CalculatePossibleRawResolutions(ImportFileSize);
+
+		if (Result.PossibleResolutions.Num() == 0)
 		{
 			Result.ResultCode = ELandscapeImportResult::Error;
 			Result.ErrorMessage = LOCTEXT("Import_WeightmapFileInvalidSize", "The layer file has an invalid size");
@@ -256,15 +179,7 @@ FLandscapeImportData<uint8> FLandscapeWeightmapFileFormat_Raw::Import(const TCHA
 
 void FLandscapeWeightmapFileFormat_Raw::Export(const TCHAR* WeightmapFilename, FName LayerName, TArrayView<const uint8> Data, FLandscapeFileResolution DataResolution, FVector Scale) const
 {
-	if (FFileHelper::SaveArrayToFile(Data, WeightmapFilename))
-	{
-		FString Extension = FPaths::GetExtension(WeightmapFilename);
-
-		if (Extension == "raw")
-		{
-			WriteMetaData(WeightmapFilename, DataResolution, 8);
-		}
-	}
+	FFileHelper::SaveArrayToFile(Data, WeightmapFilename);
 }
 
 #undef LOCTEXT_NAMESPACE
diff --git a/Engine/Source/Editor/LandscapeEditor/Private/LandscapeImageFileCache.cpp b/Engine/Source/Editor/LandscapeEditor/Private/LandscapeImageFileCache.cpp
deleted file mode 100644
diff --git a/Engine/Source/Editor/LandscapeEditor/Private/LandscapeImageFileCache.h b/Engine/Source/Editor/LandscapeEditor/Private/LandscapeImageFileCache.h
deleted file mode 100644
diff --git a/Engine/Source/Editor/LandscapeEditor/Private/LandscapeImportHelper.cpp b/Engine/Source/Editor/LandscapeEditor/Private/LandscapeImportHelper.cpp
--- a/Engine/Source/Editor/LandscapeEditor/Private/LandscapeImportHelper.cpp
+++ b/Engine/Source/Editor/LandscapeEditor/Private/LandscapeImportHelper.cpp
@@ -7,7 +7,6 @@
 #include "Modules/ModuleManager.h"
 #include "HAL/FileManager.h"
 #include "Misc/Paths.h"
-#include "LandscapeEditorUtils.h"
 
 #define LOCTEXT_NAMESPACE "LandscapeImportHelper"
 
@@ -300,7 +299,7 @@ void TransformImportDataInternal(const TArray<T>& InData, TArray<T>& OutData, co
 		const FIntRect DestRegion(0, 0, RequiredResolution.Width - 1, RequiredResolution.Height - 1);
 		FLandscapeConfigHelper::ResampleData<T>(InData, OutData, SrcRegion, DestRegion);
 	}
-	else if(TransformType == ELandscapeImportTransformType::ExpandCentered || TransformType == ELandscapeImportTransformType::ExpandOffset)
+	else if(TransformType != ELandscapeImportTransformType::None)
 	{
 		int32 OffsetX = 0;
 		int32 OffsetY = 0;
@@ -361,12 +360,11 @@ void FLandscapeImportHelper::ChooseBestComponentSizeForImport(int32 Width, int32
 	bool bValidQuadsPerSectionParam = false;
 	check(Width > 0 && Height > 0);
 		
-	const int32 MaxComponents = LandscapeEditorUtils::GetMaxSizeInComponents();
 	bool bFoundMatch = false;
 	// Try to find a section size and number of sections that exactly matches the dimensions of the heightfield
 	for (int32 SectionSizesIdx = UE_ARRAY_COUNT(FLandscapeConfig::SubsectionSizeQuadsValues) - 1; SectionSizesIdx >= 0; SectionSizesIdx--)
 	{
-		for (int32 NumSectionsIdx = 0; NumSectionsIdx < UE_ARRAY_COUNT(FLandscapeConfig::NumSectionValues); NumSectionsIdx++)
+		for (int32 NumSectionsIdx = UE_ARRAY_COUNT(FLandscapeConfig::NumSectionValues) - 1; NumSectionsIdx >= 0; NumSectionsIdx--)
 		{
 			int32 ss = FLandscapeConfig::SubsectionSizeQuadsValues[SectionSizesIdx];
 			int32 ns = FLandscapeConfig::NumSectionValues[NumSectionsIdx];
@@ -375,8 +373,8 @@ void FLandscapeImportHelper::ChooseBestComponentSizeForImport(int32 Width, int32
 			bValidSubsectionSizeParam |= (InOutSectionsPerComponent == ns);
 			bValidQuadsPerSectionParam |= (InOutQuadsPerSection == ss);
 
-			if (((Width - 1) % (ss * ns)) == 0 && ((Width - 1) / (ss * ns)) <= MaxComponents &&
-				((Height - 1) % (ss * ns)) == 0 && ((Height - 1) / (ss * ns)) <= MaxComponents)
+			if (((Width - 1) % (ss * ns)) == 0 && ((Width - 1) / (ss * ns)) <= 32 &&
+				((Height - 1) % (ss * ns)) == 0 && ((Height - 1) / (ss * ns)) <= 32)
 			{
 				bFoundMatch = true;
 				InOutQuadsPerSection = ss;
diff --git a/Engine/Source/Editor/LandscapeEditor/Private/LandscapeRegionUtils.cpp b/Engine/Source/Editor/LandscapeEditor/Private/LandscapeRegionUtils.cpp
deleted file mode 100644
diff --git a/Engine/Source/Editor/LandscapeEditor/Private/LandscapeRegionUtils.h b/Engine/Source/Editor/LandscapeEditor/Private/LandscapeRegionUtils.h
deleted file mode 100644
diff --git a/Engine/Source/Editor/LandscapeEditor/Private/LandscapeTiledImage.cpp b/Engine/Source/Editor/LandscapeEditor/Private/LandscapeTiledImage.cpp
deleted file mode 100644
diff --git a/Engine/Source/Editor/LandscapeEditor/Private/LandscapeTiledImage.h b/Engine/Source/Editor/LandscapeEditor/Private/LandscapeTiledImage.h
deleted file mode 100644
diff --git a/Engine/Source/Editor/LandscapeEditor/Private/NewLandscapeUtils.cpp b/Engine/Source/Editor/LandscapeEditor/Private/NewLandscapeUtils.cpp
--- a/Engine/Source/Editor/LandscapeEditor/Private/NewLandscapeUtils.cpp
+++ b/Engine/Source/Editor/LandscapeEditor/Private/NewLandscapeUtils.cpp
@@ -9,7 +9,6 @@
 #include "LandscapeEditorUtils.h"
 #include "LandscapeEdMode.h"
 #include "LandscapeConfigHelper.h"
-#include "LandscapeSettings.h"
 
 #include "Misc/MessageDialog.h"
 #include "Modules/ModuleManager.h"
diff --git a/Engine/Source/Editor/LandscapeEditor/Public/LandscapeEditorModule.h b/Engine/Source/Editor/LandscapeEditor/Public/LandscapeEditorModule.h
--- a/Engine/Source/Editor/LandscapeEditor/Public/LandscapeEditorModule.h
+++ b/Engine/Source/Editor/LandscapeEditor/Public/LandscapeEditorModule.h
@@ -15,7 +15,6 @@ using ILandscapeHeightmapFileFormat = ILandscapeFileFormat<uint16>;
 using ILandscapeWeightmapFileFormat = ILandscapeFileFormat<uint8>;
 
 class FUICommandList;
-class FLandscapeImageFileCache;
 
 /**
  * LandscapeEditor module interface
@@ -52,8 +51,4 @@ public:
 	}
 
 	virtual TSharedPtr<FUICommandList> GetLandscapeLevelViewportCommandList() const = 0;
-
-	virtual FLandscapeImageFileCache& GetImageFileCache() const = 0;
-
-
 };
diff --git a/Engine/Source/Editor/LandscapeEditor/Public/LandscapeEditorObject.h b/Engine/Source/Editor/LandscapeEditor/Public/LandscapeEditorObject.h
--- a/Engine/Source/Editor/LandscapeEditor/Public/LandscapeEditorObject.h
+++ b/Engine/Source/Editor/LandscapeEditor/Public/LandscapeEditorObject.h
@@ -270,18 +270,6 @@ struct FLandscapePatternBrushWorldSpaceSettings
 	{}
 };
 
-UENUM()
-enum class ELandscapeImportExportMode
-{
-	// Import and export only loaded landscape proxies
-	LoadedOnly, 
-	// Import and export the whole world loading & unloading regions as required.
-	All,		
-	// Import and export regions selected in the WP editor loading & unloading as required.
-	// Selected, // TODO: don.boogert   
-};
-
-
 UCLASS(MinimalAPI)
 class ULandscapeEditorObject : public UObject
 {
@@ -585,9 +573,6 @@ class ULandscapeEditorObject : public UObject
 	UPROPERTY(Category = "Import / Export", EditAnywhere, NonTransactional, meta = (DisplayName = "Export Single File", ShowForTools = "ImportExport", ToolTip = "(World Partition only) When true, exports the landscape as a single file, if false exports each grid tile individually."))
 	bool bExportSingleFile = false;
 
-	UPROPERTY(Category = "Import / Export", EditAnywhere, NonTransactional, meta = (DisplayName = "Mode", ShowForTools = "ImportExport", ToolTip = "Import Loaded or All Landscape Regions"))
-	ELandscapeImportExportMode ImportExportMode = ELandscapeImportExportMode::LoadedOnly;
-
 	UPROPERTY(NonTransactional)
 	FLandscapeImportDescriptor HeightmapImportDescriptor;
 	
@@ -606,9 +591,6 @@ public:
 	UPROPERTY(Category = "New Landscape", EditAnywhere, NonTransactional, meta = (DisplayName= "World Partition Grid Size", ToolTip="Number of components per landscape streaming proxies per axis", ShowForTools="NewLandscape", ClampMin=1, ClampMax=16, UIMin=1, UIMax=16))
 	uint32 WorldPartitionGridSize = 2;
 
-	UPROPERTY(Category = "New Landscape", EditAnywhere, NonTransactional, meta = (DisplayName= "World Partition Region Size", ToolTip="Number of components per Landscape World Partition Region per axis.", ShowForTools="NewLandscape", ClampMin=4, ClampMax=64, UIMin=4, UIMax=64))
-	uint32 WorldPartitionRegionSize = 16;
-
 	// Whether the imported alpha maps are to be interpreted as "layered" or "additive" (UE uses additive internally)
 	UPROPERTY(Category="New Landscape", EditAnywhere, NonTransactional, meta=(DisplayName="Layer Alphamap Type", ShowForTools="NewLandscape,ImportExport"))
 	ELandscapeImportAlphamapType ImportLandscape_AlphamapType;
@@ -751,8 +733,8 @@ public:
 
 	int32 ClampLandscapeSize(int32 InComponentsCount) const
 	{
-		// Max size is either whole components below 8192 verts, or 256 components
-		return FMath::Clamp(InComponentsCount, 1, FMath::Min(256, FMath::FloorToInt(8191.0f / (float)(NewLandscape_SectionsPerComponent * NewLandscape_QuadsPerSection))));
+		// Max size is either whole components below 8192 verts, or 32 components
+		return FMath::Clamp(InComponentsCount, 1, FMath::Min(32, FMath::FloorToInt(8191.0f / (float)(NewLandscape_SectionsPerComponent * NewLandscape_QuadsPerSection))));
 	}
 	
 	int32 CalcComponentsCount(int32 InResolution) const
@@ -761,7 +743,7 @@ public:
 	}
 
 	void NewLandscape_ClampSize()
-	{	
+	{
 		NewLandscape_ComponentCount.X = ClampLandscapeSize(NewLandscape_ComponentCount.X);
 		NewLandscape_ComponentCount.Y = ClampLandscapeSize(NewLandscape_ComponentCount.Y);
 	}
diff --git a/Engine/Source/Editor/LandscapeEditor/Public/LandscapeEditorUtils.h b/Engine/Source/Editor/LandscapeEditor/Public/LandscapeEditorUtils.h
--- a/Engine/Source/Editor/LandscapeEditor/Public/LandscapeEditorUtils.h
+++ b/Engine/Source/Editor/LandscapeEditor/Public/LandscapeEditorUtils.h
@@ -3,28 +3,12 @@
 #pragma once
 
 #include "CoreMinimal.h"
-#include "FileHelpers.h"
 
 class ALandscapeProxy;
-class ALandscapeStreamingProxy;
 class ULandscapeLayerInfoObject;
-class UWorldPartition;
 
 namespace LandscapeEditorUtils
 {
 	bool LANDSCAPEEDITOR_API SetHeightmapData(ALandscapeProxy* Landscape, const TArray<uint16>& Data);
 	bool LANDSCAPEEDITOR_API SetWeightmapData(ALandscapeProxy* Landscape, ULandscapeLayerInfoObject* LayerObject, const TArray<uint8>& Data);
-
-	int32 GetMaxSizeInComponents();
-	TOptional<FString> GetImportExportFilename(const FString& InDialogTitle, const FString& InStartPath, const FString& InDialogTypeString);
-
-	template<typename T>
-	void SaveObjects(TArrayView<T*> InObjects)
-	{
-		TArray<UPackage*> Packages;
-		Algo::Transform(InObjects, Packages, [](UObject* InObject) { return InObject->GetPackage(); });
-		UEditorLoadingAndSavingUtils::SavePackages(Packages, /* bOnlyDirty = */ false);
-	}
-
-	void SaveLandscapeProxies(TArrayView<ALandscapeProxy*> Proxies, UWorldPartition* WorldPartition);
 }
diff --git a/Engine/Source/Editor/LandscapeEditor/Public/LandscapeImportHelper.h b/Engine/Source/Editor/LandscapeEditor/Public/LandscapeImportHelper.h
--- a/Engine/Source/Editor/LandscapeEditor/Public/LandscapeImportHelper.h
+++ b/Engine/Source/Editor/LandscapeEditor/Public/LandscapeImportHelper.h
@@ -50,8 +50,7 @@ enum class ELandscapeImportTransformType : int8
 	None UMETA(DisplayName="Original", ToolTip="Will Import the data at the gizmo location in the original size"),
 	ExpandOffset UMETA(DisplayName="Expand", ToolTip="Will Import the data at the gizmo location and expand the data to fill the landscape") ,
 	ExpandCentered UMETA(Hidden), 
-	Resample UMETA(ToolTip="Will resample Import data to fit landscape"),
-	Subregion UMETA(ToolTop="Import Sub-region of the Image to Landscape")
+	Resample UMETA(ToolTip="Will resample Import data to fit landscape")
 };
 
 FORCEINLINE bool operator==(const FLandscapeImportResolution& Lhs, const FLandscapeImportResolution& Rhs)
diff --git a/Engine/Source/Runtime/Landscape/Classes/LandscapeInfo.h b/Engine/Source/Runtime/Landscape/Classes/LandscapeInfo.h
--- a/Engine/Source/Runtime/Landscape/Classes/LandscapeInfo.h
+++ b/Engine/Source/Runtime/Landscape/Classes/LandscapeInfo.h
@@ -144,9 +144,6 @@ class ULandscapeInfo : public UObject
 #if WITH_EDITORONLY_DATA
 	UPROPERTY()
 	TArray<FLandscapeInfoLayerSettings> Layers;
-
-	UPROPERTY()
-	int32 RegionSizeInComponents;
 #endif // WITH_EDITORONLY_DATA
 
 public:
@@ -196,10 +193,10 @@ public:
 	virtual void Serialize(FArchive& Ar) override;
 	//~ End UObject Interface
 
-	LANDSCAPE_API FBox GetLoadedBounds() const;
+	FBox GetLoadedBounds() const;
 
 #if WITH_EDITOR
-	LANDSCAPE_API FBox GetCompleteBounds() const;
+	FBox GetCompleteBounds() const;
 #endif
 
 #if WITH_EDITOR
@@ -362,19 +359,8 @@ public:
 	 *     // Code
 	 * });
 	 */
-	UE_DEPRECATED(5.3, "This function has been deprecated, please use the ForEachLandscapeProxy property instead")
 	LANDSCAPE_API void ForAllLandscapeProxies(TFunctionRef<void(ALandscapeProxy*)> Fn) const;
 
-	/**
-	 * Runs the given function on the root landscape actor and all streaming proxies, with the posibility of early exit
-	 * Most easily used with a lambda as follows:
-	 * ForEachLandscapeProxy([](ALandscapeProxy* Proxy) -> bool
-	 * {
-	 *     return continueLoop ? true : false;
-	 * });
-	 */
-	LANDSCAPE_API void ForEachLandscapeProxy(TFunctionRef<bool(ALandscapeProxy*)> Fn) const;
-
 	/** Associates passed actor with this info object
  *  @param	Proxy		Landscape actor to register
  *  @param  bMapCheck	Whether to warn about landscape errors
diff --git a/Engine/Source/Runtime/Landscape/Private/Landscape.cpp b/Engine/Source/Runtime/Landscape/Private/Landscape.cpp
--- a/Engine/Source/Runtime/Landscape/Private/Landscape.cpp
+++ b/Engine/Source/Runtime/Landscape/Private/Landscape.cpp
@@ -3082,7 +3082,7 @@ FLandscapeEditorLayerSettings& ULandscapeInfo::GetLayerEditorSettings(ULandscape
 
 void ULandscapeInfo::CreateLayerEditorSettingsFor(ULandscapeLayerInfoObject* LayerInfo)
 {
-	ForEachLandscapeProxy([LayerInfo](ALandscapeProxy* Proxy)
+	ForAllLandscapeProxies([LayerInfo](ALandscapeProxy* Proxy)
 	{
 		FLandscapeEditorLayerSettings* EditorLayerSettings = Proxy->EditorLayerSettings.FindByKey(LayerInfo);
 		if (!EditorLayerSettings)
@@ -3090,7 +3090,6 @@ void ULandscapeInfo::CreateLayerEditorSettingsFor(ULandscapeLayerInfoObject* Lay
 			Proxy->Modify();
 			Proxy->EditorLayerSettings.Add(FLandscapeEditorLayerSettings(LayerInfo));
 		}
-		return true;
 	});
 }
 
@@ -3306,14 +3305,13 @@ bool ULandscapeInfo::UpdateLayerInfoMapInternal(ALandscapeProxy* Proxy, bool bIn
 
 			if (!bInvalidate)
 			{
-				ForEachLandscapeProxy([this](ALandscapeProxy* EachProxy)
+				ForAllLandscapeProxies([this](ALandscapeProxy* EachProxy)
 				{
 					if (!EachProxy->IsPendingKillPending())
 					{
 						checkSlow(EachProxy->GetLandscapeInfo() == this);
 						UpdateLayerInfoMapInternal(EachProxy, false);
 					}
-					return true;
 				});
 			}
 		}
@@ -4000,16 +3998,13 @@ void ULandscapeInfo::ModifyObject(UObject* InObject, bool bAlwaysMarkDirty)
 
 ALandscapeProxy* ULandscapeInfo::GetLandscapeProxyForLevel(ULevel* Level) const
 {
-	TRACE_CPUPROFILER_EVENT_SCOPE(ULandscapeInfo::GetLandscapeProxyForLevel);
 	ALandscapeProxy* LandscapeProxy = nullptr;
-	ForEachLandscapeProxy([&LandscapeProxy, Level](ALandscapeProxy* Proxy) -> bool 
+	ForAllLandscapeProxies([&LandscapeProxy, Level](ALandscapeProxy* Proxy)
 	{
 		if (Proxy->GetLevel() == Level)
 		{
 			LandscapeProxy = Proxy;
-			return false;
 		}
-		return true;
 	});
 	return LandscapeProxy;
 }
@@ -4018,10 +4013,8 @@ ALandscapeProxy* ULandscapeInfo::GetLandscapeProxyForLevel(ULevel* Level) const
 
 ALandscapeProxy* ULandscapeInfo::GetCurrentLevelLandscapeProxy(bool bRegistered) const
 {
-	TRACE_CPUPROFILER_EVENT_SCOPE(ULandscapeInfo::GetCurrentLevelLandscapeProxy);
-
 	ALandscapeProxy* LandscapeProxy = nullptr;
-	ForEachLandscapeProxy([&LandscapeProxy, bRegistered](ALandscapeProxy* Proxy) -> bool
+	ForAllLandscapeProxies([&LandscapeProxy, bRegistered](ALandscapeProxy* Proxy)
 	{
 		if (!bRegistered || Proxy->GetRootComponent()->IsRegistered())
 		{
@@ -4030,10 +4023,8 @@ ALandscapeProxy* ULandscapeInfo::GetCurrentLevelLandscapeProxy(bool bRegistered)
 				ProxyWorld->GetCurrentLevel() == Proxy->GetOuter())
 			{
 				LandscapeProxy = Proxy;
-				return false;
 			}
 		}
-		return true;
 	});
 	return LandscapeProxy;
 }
@@ -4155,13 +4146,12 @@ void ULandscapeInfo::FixupProxiesTransform(bool bDirty)
 
 void ULandscapeInfo::UpdateComponentLayerAllowList()
 {
-	ForEachLandscapeProxy([](ALandscapeProxy* Proxy)
+	ForAllLandscapeProxies([](ALandscapeProxy* Proxy)
 	{
 		for (ULandscapeComponent* Comp : Proxy->LandscapeComponents)
 		{
 			Comp->UpdateLayerAllowListFromPaintedLayers();
 		}
-		return true;
 	});
 }
 
@@ -4286,28 +4276,6 @@ void ULandscapeInfo::ForAllLandscapeProxies(TFunctionRef<void(ALandscapeProxy*)>
 	}
 }
 
-void ULandscapeInfo::ForEachLandscapeProxy(TFunctionRef<bool(ALandscapeProxy*)> Fn) const
-{
-	if (ALandscape* Landscape = LandscapeActor.Get())
-	{
-		if (!Fn(Landscape))
-		{
-			return;
-		}
-	}
-
-	for (TWeakObjectPtr<ALandscapeStreamingProxy> StreamingProxyPtr : StreamingProxies)
-	{
-		if (ALandscapeProxy* LandscapeProxy = StreamingProxyPtr.Get())
-		{
-			if (!Fn(LandscapeProxy))
-			{
-				return;
-			}
-		}
-	}
-}
-
 bool ULandscapeInfo::IsRegistered(const ALandscapeProxy* Proxy) const
 {
 	if (Proxy == nullptr)
diff --git a/Engine/Source/Runtime/Landscape/Private/LandscapeCollision.cpp b/Engine/Source/Runtime/Landscape/Private/LandscapeCollision.cpp
--- a/Engine/Source/Runtime/Landscape/Private/LandscapeCollision.cpp
+++ b/Engine/Source/Runtime/Landscape/Private/LandscapeCollision.cpp
@@ -1803,8 +1803,6 @@ bool ULandscapeHeightfieldCollisionComponent::RecreateCollision()
 #if WITH_EDITORONLY_DATA
 void ULandscapeHeightfieldCollisionComponent::SnapFoliageInstances()
 {
-	TRACE_CPUPROFILER_EVENT_SCOPE(ULandscapeHeightfieldCollisionComponent::SnapFoliageInstances);
-
 	SnapFoliageInstances(FBox(FVector(-WORLD_MAX), FVector(WORLD_MAX)));
 }
 
@@ -2243,8 +2241,6 @@ void ULandscapeHeightfieldCollisionComponent::PreSave(FObjectPreSaveContext Obje
 #if WITH_EDITOR
 void ULandscapeInfo::UpdateAllAddCollisions()
 {
-
-	TRACE_CPUPROFILER_EVENT_SCOPE(ULandscapeInfo::UpdateAllAddCollisions);
 	XYtoAddCollisionMap.Reset();
 
 	// Don't recreate add collisions if the landscape is not registered. This can happen during Undo.
diff --git a/Engine/Source/Runtime/Landscape/Private/LandscapeConfigHelper.cpp b/Engine/Source/Runtime/Landscape/Private/LandscapeConfigHelper.cpp
--- a/Engine/Source/Runtime/Landscape/Private/LandscapeConfigHelper.cpp
+++ b/Engine/Source/Runtime/Landscape/Private/LandscapeConfigHelper.cpp
@@ -354,15 +354,14 @@ void FLandscapeConfigHelper::MoveSplinesToLandscape(ULandscapeInfo* InLandscapeI
 		InLandscape->CreateSplineComponent();
 	}
 
-	InLandscapeInfo->ForEachLandscapeProxy([InLandscapeInfo, InLandscape](ALandscapeProxy* LandscapeProxy)
+	InLandscapeInfo->ForAllLandscapeProxies([InLandscapeInfo, InLandscape](ALandscapeProxy* LandscapeProxy)
 	{
 		if (InLandscape == LandscapeProxy || !LandscapeProxy->GetSplinesComponent())
 		{
-			return true;
+			return;
 		}
 
 		InLandscapeInfo->MoveSplinesToProxy(LandscapeProxy->GetSplinesComponent(), InLandscape);
-		return true;
 	});
 }
 
diff --git a/Engine/Source/Runtime/Landscape/Private/LandscapeEdit.cpp b/Engine/Source/Runtime/Landscape/Private/LandscapeEdit.cpp
--- a/Engine/Source/Runtime/Landscape/Private/LandscapeEdit.cpp
+++ b/Engine/Source/Runtime/Landscape/Private/LandscapeEdit.cpp
@@ -1128,8 +1128,6 @@ void ULandscapeComponent::ClearDirtyCollisionHeightData()
 
 void ULandscapeComponent::UpdateCollisionHeightData(const FColor* const HeightmapTextureMipData, const FColor* const SimpleCollisionHeightmapTextureData, int32 ComponentX1/*=0*/, int32 ComponentY1/*=0*/, int32 ComponentX2/*=MAX_int32*/, int32 ComponentY2/*=MAX_int32*/, bool bUpdateBounds/*=false*/, const FColor* XYOffsetTextureMipData/*=nullptr*/, bool bInUpdateHeightfieldRegion/*=true*/)
 {
-	TRACE_CPUPROFILER_EVENT_SCOPE(ULandscapeComponent::UpdateCollisionHeightData);
-	
 	ALandscapeProxy* Proxy = GetLandscapeProxy();
 	FIntPoint ComponentKey = GetSectionBase() / ComponentSizeQuads;
 	ULandscapeHeightfieldCollisionComponent* CollisionComp = GetCollisionComponent();
@@ -1919,8 +1917,6 @@ void ULandscapeComponent::UpdateCollisionPhysicalMaterialData(TArray<UPhysicalMa
 
 void ULandscapeComponent::GenerateHeightmapMips(TArray<FColor*>& HeightmapTextureMipData, int32 ComponentX1/*=0*/, int32 ComponentY1/*=0*/, int32 ComponentX2/*=MAX_int32*/, int32 ComponentY2/*=MAX_int32*/, FLandscapeTextureDataInfo* TextureDataInfo/*=nullptr*/)
 {
-	TRACE_CPUPROFILER_EVENT_SCOPE(ULandscapeComponent::GenerateHeightmapMips);
-	
 	bool EndX = false;
 	bool EndY = false;
 
@@ -2530,44 +2526,36 @@ FIntRect ULandscapeComponent::GetComponentExtent() const
 bool ULandscapeInfo::SupportsLandscapeEditing() const
 {
 	bool bSupportsEditing = true;
-	ForEachLandscapeProxy([&bSupportsEditing](ALandscapeProxy* Proxy)
+	ForAllLandscapeProxies([&bSupportsEditing](ALandscapeProxy* Proxy)
 	{
 		if(Proxy->GetOutermost()->bIsCookedForEditor)
 		{
 			bSupportsEditing = false;
-			return false;
 		}
-		return true;
-
 	});
 	return bSupportsEditing;
 }
 
 bool ULandscapeInfo::AreAllComponentsRegistered() const
 {
-	bool bAllRegistered = true;
-	ForEachLandscapeProxy([&bAllRegistered, this](ALandscapeProxy* LandscapeProxy) {
+	const TArray<ALandscapeProxy*>& LandscapeProxies = ALandscapeProxy::GetLandscapeProxies();
+	for(ALandscapeProxy* LandscapeProxy : LandscapeProxies)
+	{
 		if (!IsValid(LandscapeProxy))
 		{
-			return true;
+			continue;
 		}
 
 		if (LandscapeProxy->GetLandscapeGuid() == LandscapeGuid)
-		{	
+		{
 			for (ULandscapeComponent* LandscapeComponent : LandscapeProxy->LandscapeComponents)
 			{
 				if (LandscapeComponent && !LandscapeComponent->IsRegistered())
 				{
-					bAllRegistered = false;
+					return false;
 				}
 			}
 		}
-		return true;
-	});
-
-	if (!bAllRegistered)
-	{
-		return false;
 	}
 		
 	for (TScriptInterface<ILandscapeSplineInterface> SplineOwner : SplineActors)
@@ -2758,8 +2746,6 @@ ENGINE_API extern bool GDisableAutomaticTextureMaterialUpdateDependencies;
 LANDSCAPE_API void ALandscapeProxy::Import(const FGuid& InGuid, int32 InMinX, int32 InMinY, int32 InMaxX, int32 InMaxY, int32 InNumSubsections, int32 InSubsectionSizeQuads, const TMap<FGuid, TArray<uint16>>& InImportHeightData, 
 										   const TCHAR* const InHeightmapFileName, const TMap<FGuid, TArray<FLandscapeImportLayerInfo>>& InImportMaterialLayerInfos, ELandscapeImportAlphamapType InImportMaterialLayerType, const TArray<FLandscapeLayer>* InImportLayers)
 {
-	TRACE_CPUPROFILER_EVENT_SCOPE(ALandscapeProxy::Import);
-	
 	check(InGuid.IsValid());
 	check(InImportHeightData.Num() == InImportMaterialLayerInfos.Num());
 
@@ -3337,8 +3323,6 @@ LANDSCAPE_API void ALandscapeProxy::Import(const FGuid& InGuid, int32 InMinX, in
 	SlowTask.EnterProgressFrame(1.0f);
 
 	{
-		TRACE_CPUPROFILER_EVENT_SCOPE(ALandscapeProxy::Import - Update Materials);
-		
 		// We disable automatic material update context, to manage it manually
 		GDisableAutomaticTextureMaterialUpdateDependencies = true;
 	
@@ -4051,7 +4035,7 @@ LANDSCAPE_API bool ULandscapeInfo::GetLandscapeXYComponentBounds(FIntRect& OutXY
 
 LANDSCAPE_API void ULandscapeInfo::ForAllLandscapeComponents(TFunctionRef<void(ULandscapeComponent*)> Fn) const
 {
-	ForEachLandscapeProxy([&](ALandscapeProxy* Proxy)
+	ForAllLandscapeProxies([&](ALandscapeProxy* Proxy)
 	{
 		for (ULandscapeComponent* Component : Proxy->LandscapeComponents)
 		{
@@ -4060,7 +4044,6 @@ LANDSCAPE_API void ULandscapeInfo::ForAllLandscapeComponents(TFunctionRef<void(U
 				Fn(Component);
 			}
 		}
-		return true;
 	});
 }
 
@@ -4262,7 +4245,7 @@ void ULandscapeInfo::DeleteLayer(ULandscapeLayerInfoObject* LayerInfo, const FNa
 		}
 	}
 
-	ForEachLandscapeProxy([LayerInfo](ALandscapeProxy* Proxy)
+	ForAllLandscapeProxies([LayerInfo](ALandscapeProxy* Proxy)
 	{
 		Proxy->Modify();
 		int32 Index = Proxy->EditorLayerSettings.IndexOfByKey(LayerInfo);
@@ -4270,7 +4253,6 @@ void ULandscapeInfo::DeleteLayer(ULandscapeLayerInfoObject* LayerInfo, const FNa
 		{
 			Proxy->EditorLayerSettings.RemoveAt(Index);
 		}
-		return true;
 	});
 
 	//UpdateLayerInfoMap();
@@ -4297,7 +4279,7 @@ void ULandscapeInfo::ReplaceLayer(ULandscapeLayerInfoObject* FromLayerInfo, ULan
 			}
 		}
 
-		ForEachLandscapeProxy([FromLayerInfo, ToLayerInfo](ALandscapeProxy* Proxy)
+		ForAllLandscapeProxies([FromLayerInfo, ToLayerInfo](ALandscapeProxy* Proxy)
 		{
 			Proxy->Modify();
 			FLandscapeEditorLayerSettings* ToEditorLayerSettings = Proxy->EditorLayerSettings.FindByKey(ToLayerInfo);
@@ -4324,7 +4306,6 @@ void ULandscapeInfo::ReplaceLayer(ULandscapeLayerInfoObject* FromLayerInfo, ULan
 					Proxy->EditorLayerSettings.Add(FLandscapeEditorLayerSettings(ToLayerInfo));
 				}
 			}
-			return true;
 		});
 
 		//UpdateLayerInfoMap();
@@ -4336,7 +4317,7 @@ void ULandscapeInfo::ReplaceLayer(ULandscapeLayerInfoObject* FromLayerInfo, ULan
 void ULandscapeInfo::GetUsedPaintLayers(const FGuid& InLayerGuid, TArray<ULandscapeLayerInfoObject*>& OutUsedLayerInfos) const
 {
 	OutUsedLayerInfos.Empty();
-	ForEachLandscapeProxy([&](ALandscapeProxy* Proxy)
+	ForAllLandscapeProxies([&](ALandscapeProxy* Proxy)
 	{
 		for (ULandscapeComponent* Component : Proxy->LandscapeComponents)
 		{
@@ -4345,8 +4326,6 @@ void ULandscapeInfo::GetUsedPaintLayers(const FGuid& InLayerGuid, TArray<ULandsc
 				Component->GetUsedPaintLayers(InLayerGuid, OutUsedLayerInfos);
 			}
 		}
-
-		return true;
 	});
 }
 
@@ -4799,19 +4778,17 @@ void ALandscapeProxy::RecreateCollisionComponents()
 
 void ULandscapeInfo::RecreateCollisionComponents()
 {
-	ForEachLandscapeProxy([](ALandscapeProxy* Proxy)
+	ForAllLandscapeProxies([](ALandscapeProxy* Proxy)
 	{
 		Proxy->RecreateCollisionComponents();
-		return true;
 	});
 }
 
 void ULandscapeInfo::RemoveXYOffsets()
 {
-	ForEachLandscapeProxy([](ALandscapeProxy* Proxy)
+	ForAllLandscapeProxies([](ALandscapeProxy* Proxy)
 	{
 		Proxy->RemoveXYOffsets();
-		return true;
 	});
 }
 
@@ -4837,10 +4814,9 @@ void ULandscapeInfo::ClearDirtyData()
 
 void ULandscapeInfo::UpdateAllComponentMaterialInstances(bool bInInvalidateCombinationMaterials)
 {
-	ForEachLandscapeProxy([=](ALandscapeProxy* Proxy)
+	ForAllLandscapeProxies([=](ALandscapeProxy* Proxy)
 	{
 		Proxy->UpdateAllComponentMaterialInstances(bInInvalidateCombinationMaterials);
-		return true;
 	});
 }
 
@@ -4914,8 +4890,6 @@ ALandscapeProxy* ULandscapeInfo::MoveComponentsToLevel(const TArray<ULandscapeCo
 
 ALandscapeProxy* ULandscapeInfo::MoveComponentsToProxy(const TArray<ULandscapeComponent*>& InComponents, ALandscapeProxy* LandscapeProxy, bool bSetPositionAndOffset, ULevel* TargetLevel)
 {
-	TRACE_CPUPROFILER_EVENT_SCOPE(ULandscapeInfo::MoveComponentsToProxy);
-	
 	ALandscape* Landscape = LandscapeActor.Get();
 	check(Landscape != nullptr);
 	
@@ -5116,8 +5090,6 @@ ALandscapeProxy* ULandscapeInfo::MoveComponentsToProxy(const TArray<ULandscapeCo
 
 void ALandscape::SplitHeightmap(ULandscapeComponent* Comp, ALandscapeProxy* TargetProxy, FMaterialUpdateContext* InOutUpdateContext, TArray<FComponentRecreateRenderStateContext>* InOutRecreateRenderStateContext, bool InReregisterComponent)
 {
-	TRACE_CPUPROFILER_EVENT_SCOPE(ALandscape::SplitHeightmap);
-	
 	ULandscapeInfo* Info = Comp->GetLandscapeInfo();
 
 	// Make sure the heightmap UVs are powers of two.
diff --git a/Engine/Source/Runtime/Landscape/Private/LandscapeEditInterface.cpp b/Engine/Source/Runtime/Landscape/Private/LandscapeEditInterface.cpp
--- a/Engine/Source/Runtime/Landscape/Private/LandscapeEditInterface.cpp
+++ b/Engine/Source/Runtime/Landscape/Private/LandscapeEditInterface.cpp
@@ -261,7 +261,7 @@ bool FLandscapeEditDataInterface::GetComponentsInRegion(int32 X1, int32 Y1, int3
 void FLandscapeEditDataInterface::SetHeightData(int32 X1, int32 Y1, int32 X2, int32 Y2, const uint16* InData, int32 InStride, bool InCalcNormals, const uint16* InNormalData, const uint16* InHeightAlphaBlendData, const uint8* InHeightFlagsData, bool InCreateComponents, UTexture2D* InHeightmap, UTexture2D* InXYOffsetmapTexture,
 											   bool InUpdateBounds, bool InUpdateCollision, bool InGenerateMips)
 {
-	TRACE_CPUPROFILER_EVENT_SCOPE(FLandscapeEditDataInterface::SetHeightData);
+	TRACE_CPUPROFILER_EVENT_SCOPE(LandscapeEditDataInterface_SetHeightData);
 
 	const int32 NumVertsX = 1 + X2 - X1;
 	const int32 NumVertsY = 1 + Y2 - Y1;
diff --git a/Engine/Source/Runtime/Landscape/Private/LandscapeEditLayers.cpp b/Engine/Source/Runtime/Landscape/Private/LandscapeEditLayers.cpp
--- a/Engine/Source/Runtime/Landscape/Private/LandscapeEditLayers.cpp
+++ b/Engine/Source/Runtime/Landscape/Private/LandscapeEditLayers.cpp
@@ -33,7 +33,6 @@ LandscapeEditLayers.cpp: Landscape editing layers mode
 #include "Rendering/Texture2DResource.h"
 #include "SceneView.h"
 #include "MaterialCachedData.h"
-#include "ContentStreaming.h"
 
 #if WITH_EDITOR
 #include "LandscapeEditorModule.h"
@@ -192,7 +191,7 @@ void OnLandscapeEditLayersLocalMergeChanged(IConsoleVariable* CVar)
 	}
 }
 
-int32 LandscapeEditLayersLocalMerge = 1;
+int32 LandscapeEditLayersLocalMerge = 0;
 static FAutoConsoleVariableRef CVarLandscapeEditLayersLocalMerge(
 	TEXT("landscape.EditLayersLocalMerge.Enable"),
 	LandscapeEditLayersLocalMerge,
@@ -2055,7 +2054,7 @@ void ALandscape::ReleaseLayersRenderingResource()
 		return;
 	}
 
-	Info->ForEachLandscapeProxy([&](ALandscapeProxy* Proxy)
+	Info->ForAllLandscapeProxies([&](ALandscapeProxy* Proxy)
 	{
 		for (auto& ItPair : Proxy->HeightmapsCPUReadback)
 		{
@@ -2070,7 +2069,6 @@ void ALandscape::ReleaseLayersRenderingResource()
 			delete WeightmapCPUReadback;
 		}
 		Proxy->WeightmapsCPUReadback.Empty();
-		return true;
 	});
 
 	if (CombinedLayersWeightmapAllMaterialLayersResource != nullptr)
@@ -2122,7 +2120,7 @@ FIntPoint ALandscape::ComputeComponentCounts() const
 	FIntPoint MaxSectionBase(TNumericLimits<int32>::Min(), TNumericLimits<int32>::Min());
 	FIntPoint MinSectionBase(TNumericLimits<int32>::Max(), TNumericLimits<int32>::Max());
 
-	Info->ForEachLandscapeProxy([&MaxSectionBase, &MinSectionBase](ALandscapeProxy* Proxy)
+	Info->ForAllLandscapeProxies([&MaxSectionBase, &MinSectionBase](ALandscapeProxy* Proxy)
 	{
 		for (ULandscapeComponent* Component : Proxy->LandscapeComponents)
 		{
@@ -2132,7 +2130,6 @@ FIntPoint ALandscape::ComputeComponentCounts() const
 			MinSectionBase.X = FMath::Min(MinSectionBase.X, Component->SectionBaseX);
 			MinSectionBase.Y = FMath::Min(MinSectionBase.Y, Component->SectionBaseY);
 		}
-		return true;
 	});
 
 	if ((MaxSectionBase.X >= MinSectionBase.X) && (MaxSectionBase.Y >= MinSectionBase.Y))
@@ -2153,10 +2150,9 @@ void ALandscape::CopyOldDataToDefaultLayer()
 		return;
 	}
 
-	Info->ForEachLandscapeProxy([&](ALandscapeProxy* Proxy)
+	Info->ForAllLandscapeProxies([&](ALandscapeProxy* Proxy)
 	{
 		CopyOldDataToDefaultLayer(Proxy);
-		return true;
 	});
 }
 
@@ -2321,10 +2317,9 @@ void ALandscape::UpdateProxyLayersWeightmapUsage()
 		return;
 	}
 
-	Info->ForEachLandscapeProxy([&](ALandscapeProxy* Proxy)
+	Info->ForAllLandscapeProxies([&](ALandscapeProxy* Proxy)
 	{
 		Proxy->UpdateProxyLayersWeightmapUsage();
-		return true;
 	});
 }
 
@@ -2353,10 +2348,9 @@ void ALandscape::InitializeLandscapeLayersWeightmapUsage()
 		return;
 	}
 
-	Info->ForEachLandscapeProxy([&](ALandscapeProxy* Proxy)
+	Info->ForAllLandscapeProxies([&](ALandscapeProxy* Proxy)
 	{
 		Proxy->InitializeProxyLayersWeightmapUsage();
-		return true;
 	});
 }
 
@@ -2470,10 +2464,9 @@ void ALandscape::ValidateProxyLayersWeightmapUsage() const
 		return;
 	}
 
-	Info->ForEachLandscapeProxy([=](const ALandscapeProxy* Proxy)
+	Info->ForAllLandscapeProxies([=](const ALandscapeProxy* Proxy)
 	{
 		Proxy->ValidateProxyLayersWeightmapUsage();
-		return true;
 	});
 }
 
@@ -3488,7 +3481,7 @@ bool ALandscape::PrepareTextureResources(bool bInWaitForStreaming)
 	TSet<ULandscapeComponent*> StreamedInHeightmapComponents;
 
 	bool bIsReady = true;
-	Info->ForEachLandscapeProxy([&](ALandscapeProxy* Proxy)
+	Info->ForAllLandscapeProxies([&](ALandscapeProxy* Proxy)
 	{
 		for (ULandscapeComponent* Component : Proxy->LandscapeComponents)
 		{
@@ -3531,7 +3524,6 @@ bool ALandscape::PrepareTextureResources(bool bInWaitForStreaming)
 				bIsReady &= bIsTextureReady;
 			}
 		}
-		return true;
 	});
 
 	// The assets that were streaming in before and are not anymore can be considered streamed in: 
@@ -3631,7 +3623,7 @@ bool ALandscape::PrepareLayersTextureResources(const TArray<FLandscapeLayer>& In
 	}
 
 	bool bIsReady = true;
-	Info->ForEachLandscapeProxy([&](ALandscapeProxy* Proxy)
+	Info->ForAllLandscapeProxies([&](ALandscapeProxy* Proxy)
 	{
 		for (const FLandscapeLayer& Layer : InLayers)
 		{
@@ -3648,7 +3640,6 @@ bool ALandscape::PrepareLayersTextureResources(const TArray<FLandscapeLayer>& In
 				}
 			}
 		}
-		return true;
 	});
 
 	return bIsReady;
@@ -4410,8 +4401,8 @@ namespace EditLayersHeightmapLocalMerge_RenderThread
 			const FLandscapeRDGTrackedTexture* TrackedTexture = InTrackedTextures.Find(TextureResolveInfo.Texture);
 			check((TrackedTexture != nullptr) && (TrackedTexture->ScratchTextureRef != nullptr));
 
-			//check(TextureResolveInfo.NumMips == TrackedTexture->ScratchTextureRef->Desc.NumMips);
-			//check(TrackedTexture->ScratchTextureMipsSRVRefs.Num() == TextureResolveInfo.NumMips);
+			check(TextureResolveInfo.NumMips == TrackedTexture->ScratchTextureRef->Desc.NumMips);
+			check(TrackedTexture->ScratchTextureMipsSRVRefs.Num() == TextureResolveInfo.NumMips);
 
 			FIntPoint CurrentMipSubregionSize = InLocalMergeInfo.ComponentSizeVerts;
 			for (int32 MipLevel = 1; MipLevel < TextureResolveInfo.NumMips; ++MipLevel)
@@ -5879,7 +5870,7 @@ void ALandscape::ReallocateLayersWeightmaps(FUpdateLayersContentContext& InUpdat
 	//GDisableAutomaticTextureMaterialUpdateDependencies = false;
 
 	// Clean-up unused weightmap CPUReadback resources
-	Info->ForEachLandscapeProxy([](ALandscapeProxy* Proxy)
+	Info->ForAllLandscapeProxies([](ALandscapeProxy* Proxy)
 	{
 		TArray<UTexture2D*, TInlineAllocator<64>> EntriesToRemoveFromMap;
 		for (auto& Pair : Proxy->WeightmapsCPUReadback)
@@ -5915,8 +5906,6 @@ void ALandscape::ReallocateLayersWeightmaps(FUpdateLayersContentContext& InUpdat
 				}
 			}
 		}
-
-		return true;
 	});
 
 	ValidateProxyLayersWeightmapUsage();
@@ -7725,7 +7714,7 @@ void ALandscape::RequestLayersContentUpdateForceAll(ELandscapeLayerUpdateMode In
 	const bool bUpdateClientUdpateEditing = (InModeMask & ELandscapeLayerUpdateMode::Update_Client_Editing) != 0;
 	if (ULandscapeInfo* LandscapeInfo = GetLandscapeInfo())
 	{
-		LandscapeInfo->ForEachLandscapeProxy([bUpdateHeightmap, bUpdateWeightmap, bUpdateAllHeightmap, bUpdateAllWeightmap, bUpdateHeightCollision, bUpdateWeightCollision, bUpdateClientUdpateEditing, bInUserTriggered](ALandscapeProxy* Proxy)
+		LandscapeInfo->ForAllLandscapeProxies([bUpdateHeightmap, bUpdateWeightmap, bUpdateAllHeightmap, bUpdateAllWeightmap, bUpdateHeightCollision, bUpdateWeightCollision, bUpdateClientUdpateEditing, bInUserTriggered](ALandscapeProxy* Proxy)
 		{
 			if (Proxy)
 			{
@@ -7747,7 +7736,6 @@ void ALandscape::RequestLayersContentUpdateForceAll(ELandscapeLayerUpdateMode In
 					}
 				}
 			}
-			return true;
 		});
 	}
 
@@ -8439,14 +8427,6 @@ void ALandscape::ForceUpdateLayersContent(bool bIntermediateRender)
 
 void ALandscape::ForceLayersFullUpdate()
 {
-	TRACE_CPUPROFILER_EVENT_SCOPE(ALandscape::ForceLayersFullUpdate);
-
-	FAssetCompilingManager::Get().FinishAllCompilation();
-
-	FStreamingManagerCollection& StreamingManagers = IStreamingManager::Get();
-	StreamingManagers.UpdateResourceStreaming(GetWorld()->GetDeltaSeconds(), /* bProcessEverything */ true);
-	StreamingManagers.BlockTillAllRequestsFinished();
-
 	RequestSplineLayerUpdate();
 	RequestLayersContentUpdateForceAll();
 	ForceUpdateLayersContent(/* bIntermediateRender */ false);
@@ -8932,10 +8912,9 @@ void ALandscape::DeleteLayer(int32 InLayerIndex)
 	FGuid LayerGuid = Layer->Guid;
 
 	// Clean up Weightmap usage in LandscapeProxies
-	LandscapeInfo->ForEachLandscapeProxy([&LayerGuid](ALandscapeProxy* Proxy)
+	LandscapeInfo->ForAllLandscapeProxies([&LayerGuid](ALandscapeProxy* Proxy)
 	{
 		Proxy->DeleteLayer(LayerGuid);
-		return true;
 	});
 
 	const FLandscapeLayer* SplinesReservedLayer = GetLandscapeSplinesReservedLayer();
@@ -9058,14 +9037,13 @@ void ALandscape::ClearPaintLayer(const FGuid& InLayerGuid, ULandscapeLayerInfoOb
 	FScopedSetLandscapeEditingLayer Scope(this, InLayerGuid, [=] { RequestLayersContentUpdate(ELandscapeLayerUpdateMode::Update_Weightmap_All); });
 
 	FLandscapeEditDataInterface LandscapeEdit(LandscapeInfo);
-	LandscapeInfo->ForEachLandscapeProxy([&](ALandscapeProxy* Proxy)
+	LandscapeInfo->ForAllLandscapeProxies([&](ALandscapeProxy* Proxy)
 	{
 		Proxy->Modify();
 		for (ULandscapeComponent* Component : Proxy->LandscapeComponents)
 		{
 			Component->DeleteLayer(InLayerInfo, LandscapeEdit);
 		}
-		return true;
 	});
 }
 
@@ -9133,11 +9111,10 @@ void ALandscape::ClearLayer(const FGuid& InLayerGuid, TSet<TObjectPtr<ULandscape
 	}
 	else
 	{
-		LandscapeInfo->ForEachLandscapeProxy([&](ALandscapeProxy* Proxy)
+		LandscapeInfo->ForAllLandscapeProxies([&](ALandscapeProxy* Proxy)
 		{
 			Proxy->Modify(bMarkPackageDirty);
 			Components.Append(Proxy->LandscapeComponents);
-			return true;
 		});
 	}
 
@@ -9466,10 +9443,9 @@ FLandscapeLayer* ALandscape::DuplicateLayerAndMoveBrushes(const FLandscapeLayer&
 	int32 AddedIndex = LandscapeLayers.Add(NewLayer);
 
 	// Create associated layer data in each landscape proxy
-	LandscapeInfo->ForEachLandscapeProxy([&NewLayer](ALandscapeProxy* Proxy)
+	LandscapeInfo->ForAllLandscapeProxies([&NewLayer](ALandscapeProxy* Proxy)
 	{
 		Proxy->AddLayer(NewLayer.Guid);
-		return true;
 	});
 
 	return &LandscapeLayers[AddedIndex];
@@ -9489,10 +9465,9 @@ int32 ALandscape::CreateLayer(FName InName)
 	int32 LayerIndex = LandscapeLayers.Add(NewLayer);
 
 	// Create associated layer data in each landscape proxy
-	LandscapeInfo->ForEachLandscapeProxy([&NewLayer](ALandscapeProxy* Proxy)
+	LandscapeInfo->ForAllLandscapeProxies([&NewLayer](ALandscapeProxy* Proxy)
 	{
 		Proxy->AddLayer(NewLayer.Guid);
-		return true;
 	});
 
 	return LayerIndex;
diff --git a/Engine/Source/Runtime/Landscape/Private/LandscapeSubsystem.cpp b/Engine/Source/Runtime/Landscape/Private/LandscapeSubsystem.cpp
--- a/Engine/Source/Runtime/Landscape/Private/LandscapeSubsystem.cpp
+++ b/Engine/Source/Runtime/Landscape/Private/LandscapeSubsystem.cpp
@@ -493,8 +493,6 @@ bool ULandscapeSubsystem::IsGridBased() const
 
 void ULandscapeSubsystem::ChangeGridSize(ULandscapeInfo* LandscapeInfo, uint32 GridSizeInComponents)
 {
-	TRACE_CPUPROFILER_EVENT_SCOPE(ULandscapeSubsystem::ChangeGridSize);
-	
 	if (!IsGridBased())
 	{
 		return;
diff --git a/Engine/Source/Runtime/Landscape/Public/LandscapeConfigHelper.h b/Engine/Source/Runtime/Landscape/Public/LandscapeConfigHelper.h
--- a/Engine/Source/Runtime/Landscape/Public/LandscapeConfigHelper.h
+++ b/Engine/Source/Runtime/Landscape/Public/LandscapeConfigHelper.h
@@ -32,6 +32,7 @@ struct LANDSCAPE_API FLandscapeConfig
 		, SubsectionSizeQuads(InSubsectionSizeQuads)
 		, GridSizeInComponents(InGridSizeInComponents)
 	{
+
 	}
 
 	FLandscapeConfig(ULandscapeInfo* InLandscapeInfo);
@@ -130,32 +131,6 @@ public:
 			}
 		}
 	}
-
-	template<typename T>
-	static void CopySubregion(const TArrayView<T>& InData, TArray<T>& OutData, const FIntRect& InSrcRegion, uint32 InSrcDataPitch)
-	{
-		check(InData.Size() >= InSrcRegion.Area());
-
-		const int32 SrcWidth = InSrcRegion.Width();
-		const int32 SrcHeight = InSrcRegion.Height();
-		
-		OutData.Empty(SrcWidth * SrcHeight);
-		OutData.AddUninitialized(SrcWidth * SrcHeight);
-
-		for (int32 Y = 0; Y < SrcHeight; ++Y)
-		{
-			for (int32 X = 0; X < SrcWidth; ++X)
-			{
-				const int32 SrcX = X + InSrcRegion.Min.X;
-				const int32 SrcY = Y + InSrcRegion.Min.Y;
-
-				const int32 SrcIndex = SrcX + SrcY * InSrcDataPitch;
-				const int32 DstIndex = X + Y * InSrcRegion.Width();
-
-				OutData[DstIndex] = InData[SrcIndex];
-			}
-		}
-	}
 		
 	template<typename T>
 	static void ResampleData(const TArray<T>& InData, TArray<T>& OutData, const FIntRect& InSrcRegion, const FIntRect& InDestRegion)
diff --git a/Engine/Source/Runtime/Landscape/Public/LandscapeSettings.h b/Engine/Source/Runtime/Landscape/Public/LandscapeSettings.h
--- a/Engine/Source/Runtime/Landscape/Public/LandscapeSettings.h
+++ b/Engine/Source/Runtime/Landscape/Public/LandscapeSettings.h
@@ -4,14 +4,11 @@
 
 #include "CoreMinimal.h"
 #include "Engine/DeveloperSettings.h"
-#include "Delegates/Delegate.h"
 #include "LandscapeSettings.generated.h"
 
 class UMaterialInterface;
 class ULandscapeLayerInfoObject;
 
-DECLARE_DELEGATE_OneParam(FnOnMaxComponentsChanged, uint64);
-
 UENUM()
 enum class ELandscapeDirtyingMode : uint8
 {
@@ -45,16 +42,11 @@ public:
 
 	/** Returns the default landscape layer info object that will be assigned to unset layers when creating a new landscape. */
 	TSoftObjectPtr<ULandscapeLayerInfoObject> GetDefaultLayerInfoObject() const { return DefaultLayerInfoObject; }
+
 public:
 	UPROPERTY(config, EditAnywhere, Category = "Layers", meta=(UIMin = "1", UIMax = "32", ClampMin = "1", ClampMax = "32", ToolTip = "This option controls the maximum editing layers that can be added to a Landscape"))
 	int32 MaxNumberOfLayers = 8;
 
-	UPROPERTY(config, EditAnywhere, Category = "Configuration", meta=(ToolTip = "Maximum Dimension of Landscape in Components"))
-	int32 MaxComponents = 256;
-
-	UPROPERTY(config, EditAnywhere, Category = "Configuration", meta = (UIMin = "1", UIMax = "1024", ClampMin = "1", ClampMax = "1024", ToolTip = "Maximum Size of Import Image Cache in MB"))
-	uint32 MaxImageImportCacheSizeMegaBytes = 256;
-	
 	UPROPERTY(Config, Category = "Configuration", EditAnywhere)
 	ELandscapeDirtyingMode LandscapeDirtyingMode;
 
@@ -70,5 +62,4 @@ protected:
 
 	UPROPERTY(transient)
 	bool bRestrictiveMode = false;
-
 };
