Landscape: Fix Localization keys & set local merge to be disabled by default

diff --git a/Engine/Source/Editor/LandscapeEditor/Private/LandscapeEditorDetailCustomization_ImportExport.cpp b/Engine/Source/Editor/LandscapeEditor/Private/LandscapeEditorDetailCustomization_ImportExport.cpp
--- a/Engine/Source/Editor/LandscapeEditor/Private/LandscapeEditorDetailCustomization_ImportExport.cpp
+++ b/Engine/Source/Editor/LandscapeEditor/Private/LandscapeEditorDetailCustomization_ImportExport.cpp
@@ -595,7 +595,7 @@ FReply FLandscapeEditorDetailCustomization_ImportExport::OnImportExportButtonCli
 
 				int32 NumRegions = LandscapeRegions.Num();
 					
-				FScopedSlowTask Progress(NumRegions, LOCTEXT("ImportingLandscapeRegions", "Importing Landscape Regions..."));
+				FScopedSlowTask Progress(NumRegions, LOCTEXT("ImportingLandscapeRegions", "Importing Landscape Regions"));
 				Progress.MakeDialog(/*bShowCancelButton = */ true);
 
 				auto RegionImporter = [&Progress, LandscapeInfo, CurrentLayerGuid, LandscapeEdMode, ImportRegion, TransformType, ImportOffset, PaintRestriction](const FBox& RegionBounds, const TArray<ALandscapeProxy*>& Proxies)
@@ -707,7 +707,7 @@ FReply FLandscapeEditorDetailCustomization_ImportExport::OnImportExportButtonCli
 			}
 			else
 			{
-				FScopedSlowTask Progress(0, LOCTEXT("ExportingLandscapeRegions", "Exporting Landscape Regions..."));
+				FScopedSlowTask Progress(0, LOCTEXT("ExportingLandscapeRegions", "Exporting Landscape Regions"));
 				Progress.MakeDialog(/*bShowCancelButton = */ true);
 
 				auto Exporter = [&Progress, LandscapeInfo, Landscape, LandscapeEdMode, LandscapeExtent, PerformExport, World](const FBox& RegionBounds, const TArray<ALandscapeProxy*>& Proxies)
diff --git a/Engine/Source/Editor/LandscapeEditor/Private/LandscapeEditorDetailCustomization_NewLandscape.cpp b/Engine/Source/Editor/LandscapeEditor/Private/LandscapeEditorDetailCustomization_NewLandscape.cpp
--- a/Engine/Source/Editor/LandscapeEditor/Private/LandscapeEditorDetailCustomization_NewLandscape.cpp
+++ b/Engine/Source/Editor/LandscapeEditor/Private/LandscapeEditorDetailCustomization_NewLandscape.cpp
@@ -997,6 +997,11 @@ FReply FLandscapeEditorDetailCustomization_NewLandscape::OnCreateButtonClicked()
 	TRACE_CPUPROFILER_EVENT_SCOPE(FLandscapeEditorDetailCustomization_NewLandscape::OnCreateButtonClicked);
 
 	FEdModeLandscape* LandscapeEdMode = GetEditorMode();
+	if (!LandscapeEdMode)
+	{
+		return FReply::Handled();
+	}
+
 	const bool bIsNewLandscape = LandscapeEdMode->NewLandscapePreviewMode == ENewLandscapePreviewMode::NewLandscape;
 	
 	UWorld* World = LandscapeEdMode->GetWorld();
diff --git a/Engine/Source/Editor/LandscapeEditor/Private/LandscapeEditorObject.cpp b/Engine/Source/Editor/LandscapeEditor/Private/LandscapeEditorObject.cpp
--- a/Engine/Source/Editor/LandscapeEditor/Private/LandscapeEditorObject.cpp
+++ b/Engine/Source/Editor/LandscapeEditor/Private/LandscapeEditorObject.cpp
@@ -157,7 +157,7 @@ void ULandscapeEditorObject::PostEditChangeProperty(FPropertyChangedEvent& Prope
 		PropertyChangedEvent.MemberProperty->GetFName() == GET_MEMBER_NAME_CHECKED(ULandscapeEditorObject, NewLandscape_ComponentCount))
 	{
 		// Only clamp the landscape size if we're not in World Partition
-		if (!ParentMode->IsGridBased())
+		if (ParentMode && !ParentMode->IsGridBased())
 		{
 			NewLandscape_ClampSize();
 		}
diff --git a/Engine/Source/Editor/LandscapeEditor/Private/LandscapeFileFormatRaw.cpp b/Engine/Source/Editor/LandscapeEditor/Private/LandscapeFileFormatRaw.cpp
--- a/Engine/Source/Editor/LandscapeEditor/Private/LandscapeFileFormatRaw.cpp
+++ b/Engine/Source/Editor/LandscapeEditor/Private/LandscapeFileFormatRaw.cpp
@@ -122,7 +122,7 @@ FLandscapeFileInfo FLandscapeHeightmapFileFormat_Raw::Validate(const TCHAR* Heig
 			if (BitsPerPixel != 16)
 			{
 				Result.ResultCode = ELandscapeImportResult::Error;
-				Result.ErrorMessage = LOCTEXT("Import_WeightmapFileBitsPerPixel", "Height file has an invalid number of bits per pixel");
+				Result.ErrorMessage = LOCTEXT("Import_HeightmapFileBitsPerPixel", "Height file has an invalid number of bits per pixel");
 			}
 			else
 			{
diff --git a/Engine/Source/Editor/LandscapeEditor/Private/LandscapeTiledImage.cpp b/Engine/Source/Editor/LandscapeEditor/Private/LandscapeTiledImage.cpp
--- a/Engine/Source/Editor/LandscapeEditor/Private/LandscapeTiledImage.cpp
+++ b/Engine/Source/Editor/LandscapeEditor/Private/LandscapeTiledImage.cpp
@@ -138,7 +138,7 @@ FLandscapeFileInfo FLandscapeTiledImage::Load(const TCHAR* Filename)
 		else if (Width != TileResolution.X && Height != TileResolution.Y)
 		{
 			Result.ResultCode = ELandscapeImportResult::Error;
-			Result.ErrorMessage = LOCTEXT("LandscapeFileReadError", "Mimatched resolution found in tiled image");
+			Result.ErrorMessage = LOCTEXT("FileReadErrorTiledResolutionMismatch", "Mimatched resolution found in tiled image");
 			return Result;
 		}
 	}
@@ -146,7 +146,7 @@ FLandscapeFileInfo FLandscapeTiledImage::Load(const TCHAR* Filename)
 	if (TileFilenames.Num() == 0)
 	{
 		Result.ResultCode = ELandscapeImportResult::Error;
-		Result.ErrorMessage = LOCTEXT("LandscapeFileReadError", "No files found");
+		Result.ErrorMessage = LOCTEXT("FileReadErrorNoFilesFound", "No files found");
 		return Result;
 	}
 
diff --git a/Engine/Source/Runtime/Landscape/Private/LandscapeEditLayers.cpp b/Engine/Source/Runtime/Landscape/Private/LandscapeEditLayers.cpp
--- a/Engine/Source/Runtime/Landscape/Private/LandscapeEditLayers.cpp
+++ b/Engine/Source/Runtime/Landscape/Private/LandscapeEditLayers.cpp
@@ -192,7 +192,7 @@ void OnLandscapeEditLayersLocalMergeChanged(IConsoleVariable* CVar)
 	}
 }
 
-int32 LandscapeEditLayersLocalMerge = 1;
+int32 LandscapeEditLayersLocalMerge = 0;
 static FAutoConsoleVariableRef CVarLandscapeEditLayersLocalMerge(
 	TEXT("landscape.EditLayersLocalMerge.Enable"),
 	LandscapeEditLayersLocalMerge,
