Fix for missing landscape actor in landscape tools for debug visual log

diff --git a/Engine/Source/Editor/LandscapeEditor/Private/LandscapeEdModeTools.h b/Engine/Source/Editor/LandscapeEditor/Private/LandscapeEdModeTools.h
--- a/Engine/Source/Editor/LandscapeEditor/Private/LandscapeEdModeTools.h
+++ b/Engine/Source/Editor/LandscapeEditor/Private/LandscapeEdModeTools.h
@@ -663,10 +663,12 @@ private:
 	void VisualizeLandscapeRegion(int32 InX1, int32 InY1, int32 InX2, int32 InY2, const FColor& InColor, const FString& InDescription)
 	{
 		check(LandscapeInfo != nullptr);
-		const FTransform& LandscapeTransform = LandscapeInfo->LandscapeActor->GetTransform();
+		ALandscapeProxy* LandscapeProxy = LandscapeInfo->GetLandscapeProxy();
+		check(LandscapeProxy != nullptr);
+		const FTransform& LandscapeTransform = LandscapeProxy->GetTransform();
 		FVector Min = LandscapeTransform.TransformPosition(FVector(InX1, InY1, 0));
 		FVector Max = LandscapeTransform.TransformPosition(FVector(InX2, InY2, 0));
-		UE_VLOG_BOX(LandscapeInfo->LandscapeActor.Get(), LogLandscapeTools, Log, FBox(Min, Max), InColor, TEXT("%s"), *InDescription);
+		UE_VLOG_BOX(LandscapeProxy, LogLandscapeTools, Log, FBox(Min, Max), InColor, TEXT("%s"), *InDescription);
 	}
 
 	TMap<FIntPoint, AccessorType> CachedData;
