LandscapeEditor: Fix truncation warnings

diff --git a/Engine/Source/Editor/LandscapeEditor/Private/LandscapeEdMode.cpp b/Engine/Source/Editor/LandscapeEditor/Private/LandscapeEdMode.cpp
--- a/Engine/Source/Editor/LandscapeEditor/Private/LandscapeEdMode.cpp
+++ b/Engine/Source/Editor/LandscapeEditor/Private/LandscapeEdMode.cpp
@@ -552,8 +552,8 @@ void FEdModeLandscape::Enter()
 				{
 					for (int32 X = 0; X < CurrentGizmoActor->SampleSizeX; ++X)
 					{
-						float TexX = X * SizeX / CurrentGizmoActor->SampleSizeX;
-						float TexY = Y * SizeY / CurrentGizmoActor->SampleSizeY;
+						float TexX = static_cast<float>(X * SizeX / CurrentGizmoActor->SampleSizeX);
+						float TexY = static_cast<float>(Y * SizeY / CurrentGizmoActor->SampleSizeY);
 						int32 LX = FMath::FloorToInt(TexX);
 						int32 LY = FMath::FloorToInt(TexY);
 
@@ -565,11 +565,11 @@ void FEdModeLandscape::Enter()
 						FGizmoSelectData* Data01 = CurrentGizmoActor->SelectedData.Find(FIntPoint(LX, LY + 1));
 						FGizmoSelectData* Data11 = CurrentGizmoActor->SelectedData.Find(FIntPoint(LX + 1, LY + 1));
 
-						TexData[X + Y*ALandscapeGizmoActiveActor::DataTexSize] = FMath::Lerp(
+						TexData[X + Y*ALandscapeGizmoActiveActor::DataTexSize] = static_cast<uint8>(FMath::Lerp(
 							FMath::Lerp(Data00 ? Data00->Ratio : 0, Data10 ? Data10->Ratio : 0, FracX),
 							FMath::Lerp(Data01 ? Data01->Ratio : 0, Data11 ? Data11->Ratio : 0, FracX),
 							FracY
-							) * 255;
+							) * 255);
 					}
 				}
 			}
@@ -1097,8 +1097,8 @@ bool FEdModeLandscape::LandscapeMouseTrace(FEditorViewportClient* ViewportClient
 {
 	FVector HitLocation;
 	bool bResult = LandscapeMouseTrace(ViewportClient, MouseX, MouseY, HitLocation);
-	OutHitX = HitLocation.X;
-	OutHitY = HitLocation.Y;
+	OutHitX = static_cast<float>(HitLocation.X);
+	OutHitY = static_cast<float>(HitLocation.Y);
 	return bResult;
 }
 
@@ -1270,7 +1270,8 @@ bool FEdModeLandscape::LandscapeTrace(const FVector& InRayOrigin, const FVector&
 			double InitialAngle = (ConcentricIndex % 2) * AngleIncrement / 2.0;
 			for (int32 AngularSampleIndex = 0; AngularSampleIndex < NumAngularSamples; ++AngularSampleIndex)
 			{
-				if (TOptional<FProcessLandscapeTraceHitsResult> HitResult = LineTraceAroundCenter(AngularSampleIndex * AngleIncrement + InitialAngle, Radius); HitResult.IsSet())
+				if (TOptional<FProcessLandscapeTraceHitsResult> HitResult = LineTraceAroundCenter(
+					static_cast<float>(AngularSampleIndex * AngleIncrement + InitialAngle), static_cast<float>(Radius)); HitResult.IsSet())
 				{ 
 					// Compute the mean height in order to extrapolate the landscape hit to a "fake" landscape plane that extends beyond the borders : 
 					if (TOptional<float> Height = HitResult->LandscapeProxy->GetHeightAtLocation(HitResult->HitLocation, EHeightfieldSource::Editor); Height.IsSet())
@@ -1740,7 +1741,7 @@ void FEdModeLandscape::ChangeBrushSize(bool bIncrease)
 			NewValue = FMath::Min(NewValue, Radius - 1.0f);
 		}
 
-		NewValue = (int32)FMath::Clamp(NewValue, SliderMin, SliderMax);
+		NewValue = FMath::Clamp(NewValue, SliderMin, SliderMax);
 		UISettings->SetCurrentToolBrushRadius(NewValue);
 	}
 }
@@ -3668,7 +3669,7 @@ ALandscape* FEdModeLandscape::ChangeComponentSetting(int32 NumComponentsX, int32
 				}
 			}
 
-			Progress.EnterProgressFrame(CurrentTaskProgress++);
+			Progress.EnterProgressFrame(static_cast<float>(CurrentTaskProgress++));
 
 			const FVector Location = OldLandscape->GetActorLocation() + LandscapeOffset;
 			FActorSpawnParameters SpawnParams;
@@ -3750,7 +3751,7 @@ ALandscape* FEdModeLandscape::ChangeComponentSetting(int32 NumComponentsX, int32
 				// TODO: Foliage on spline meshes
 			}
 
-			Progress.EnterProgressFrame(CurrentTaskProgress++);
+			Progress.EnterProgressFrame(static_cast<float>(CurrentTaskProgress++));
 
 			if (bResample)
 			{
@@ -3771,7 +3772,7 @@ ALandscape* FEdModeLandscape::ChangeComponentSetting(int32 NumComponentsX, int32
 					}
 				}
 
-				Progress.EnterProgressFrame(CurrentTaskProgress++);
+				Progress.EnterProgressFrame(static_cast<float>(CurrentTaskProgress++));
 
 				// delete any components that were deleted in the original
 				TSet<ULandscapeComponent*> ComponentsToDelete;
@@ -3829,14 +3830,14 @@ ALandscape* FEdModeLandscape::ChangeComponentSetting(int32 NumComponentsX, int32
 					}
 				}
 
-				Progress.EnterProgressFrame(CurrentTaskProgress++);
+				Progress.EnterProgressFrame(static_cast<float>(CurrentTaskProgress++));
 
 				// delete any components that are in areas that were entirely deleted in the original
 				TSet<ULandscapeComponent*> ComponentsToDelete;
 				for (const TPair<FIntPoint, ULandscapeComponent*>& Entry : NewLandscapeInfo->XYtoComponentMap)
 				{
-					float OldX = Entry.Key.X * NewComponentSizeQuads + LandscapeOffsetQuads.X;
-					float OldY = Entry.Key.Y * NewComponentSizeQuads + LandscapeOffsetQuads.Y;
+					const int32 OldX = Entry.Key.X * NewComponentSizeQuads + LandscapeOffsetQuads.X;
+					const int32 OldY = Entry.Key.Y * NewComponentSizeQuads + LandscapeOffsetQuads.Y;
 					TSet<ULandscapeComponent*> OverlapComponents;
 					LandscapeInfo->GetComponentsInRegion(OldX, OldY, OldX + NewComponentSizeQuads, OldY + NewComponentSizeQuads, OverlapComponents, false);
 					if (OverlapComponents.Num() == 0)
diff --git a/Engine/Source/Editor/LandscapeEditor/Private/LandscapeEdModeBrushes.cpp b/Engine/Source/Editor/LandscapeEditor/Private/LandscapeEdModeBrushes.cpp
--- a/Engine/Source/Editor/LandscapeEditor/Private/LandscapeEdModeBrushes.cpp
+++ b/Engine/Source/Editor/LandscapeEditor/Private/LandscapeEdModeBrushes.cpp
@@ -181,7 +181,7 @@ public:
 		ULandscapeInfo* LandscapeInfo = EdMode->CurrentToolTarget.LandscapeInfo.Get();
 		ALandscapeProxy* Proxy = LandscapeInfo->GetLandscapeProxy();
 
-		const float ScaleXY = FMath::Abs(LandscapeInfo->DrawScale.X);
+		const float ScaleXY = static_cast<float>(FMath::Abs(LandscapeInfo->DrawScale.X));
 		const float TotalRadius = EdMode->UISettings->GetCurrentToolBrushRadius() / ScaleXY;
 		const float Radius = (1.0f - EdMode->UISettings->GetCurrentToolBrushFalloff()) * TotalRadius;
 		const float Falloff = EdMode->UISettings->GetCurrentToolBrushFalloff() * TotalRadius;
@@ -273,7 +273,7 @@ public:
 		}
 
 		ULandscapeInfo* LandscapeInfo = EdMode->CurrentToolTarget.LandscapeInfo.Get();
-		const float ScaleXY = FMath::Abs(LandscapeInfo->DrawScale.X);
+		const float ScaleXY = static_cast<float>(FMath::Abs(LandscapeInfo->DrawScale.X));
 		const float TotalRadius = EdMode->UISettings->GetCurrentToolBrushRadius() / ScaleXY;
 		const float Radius = (1.0f - EdMode->UISettings->GetCurrentToolBrushFalloff()) * TotalRadius;
 		const float Falloff = EdMode->UISettings->GetCurrentToolBrushFalloff() * TotalRadius;
@@ -297,10 +297,10 @@ public:
 		for (const FLandscapeToolInteractorPosition& InteractorPosition : InteractorPositions)
 		{
 			FIntRect SpotBounds;
-			SpotBounds.Min.X = FMath::FloorToInt(InteractorPosition.Position.X - TotalRadius);
-			SpotBounds.Min.Y = FMath::FloorToInt(InteractorPosition.Position.Y - TotalRadius);
-			SpotBounds.Max.X = FMath::CeilToInt( InteractorPosition.Position.X + TotalRadius);
-			SpotBounds.Max.Y = FMath::CeilToInt( InteractorPosition.Position.Y + TotalRadius);
+			SpotBounds.Min.X = FMath::FloorToInt32(InteractorPosition.Position.X - TotalRadius);
+			SpotBounds.Min.Y = FMath::FloorToInt32(InteractorPosition.Position.Y - TotalRadius);
+			SpotBounds.Max.X = FMath::CeilToInt32( InteractorPosition.Position.X + TotalRadius);
+			SpotBounds.Max.Y = FMath::CeilToInt32( InteractorPosition.Position.Y + TotalRadius);
 
 			if (Bounds.IsEmpty())
 			{
@@ -327,10 +327,10 @@ public:
 		for (const FLandscapeToolInteractorPosition& InteractorPosition : InteractorPositions)
 		{
 			FIntRect SpotBounds;
-			SpotBounds.Min.X = FMath::Max(FMath::FloorToInt(InteractorPosition.Position.X - TotalRadius), Bounds.Min.X);
-			SpotBounds.Min.Y = FMath::Max(FMath::FloorToInt(InteractorPosition.Position.Y - TotalRadius), Bounds.Min.Y);
-			SpotBounds.Max.X = FMath::Min(FMath::CeilToInt( InteractorPosition.Position.X + TotalRadius), Bounds.Max.X);
-			SpotBounds.Max.Y = FMath::Min(FMath::CeilToInt( InteractorPosition.Position.Y + TotalRadius), Bounds.Max.Y);
+			SpotBounds.Min.X = FMath::Max(FMath::FloorToInt32(InteractorPosition.Position.X - TotalRadius), Bounds.Min.X);
+			SpotBounds.Min.Y = FMath::Max(FMath::FloorToInt32(InteractorPosition.Position.Y - TotalRadius), Bounds.Min.Y);
+			SpotBounds.Max.X = FMath::Min(FMath::CeilToInt32( InteractorPosition.Position.X + TotalRadius), Bounds.Max.X);
+			SpotBounds.Max.Y = FMath::Min(FMath::CeilToInt32( InteractorPosition.Position.Y + TotalRadius), Bounds.Max.Y);
 
 			for (int32 Y = SpotBounds.Min.Y; Y < SpotBounds.Max.Y; Y++)
 			{
@@ -342,7 +342,7 @@ public:
 					if (PrevAmount < 1.0f)
 					{
 						// Distance from mouse
-						float MouseDist = FMath::Sqrt(FMath::Square(InteractorPosition.Position.X - (float)X) + FMath::Square(InteractorPosition.Position.Y - (float)Y));
+						float MouseDist = static_cast<float>(FMath::Sqrt(FMath::Square(InteractorPosition.Position.X - X) + FMath::Square(InteractorPosition.Position.Y - Y)));
 
 						float PaintAmount = CalculateFalloff(MouseDist, Radius, Falloff);
 
@@ -443,8 +443,8 @@ public:
 			{
 				const int32 BrushSize = FMath::Max(EdMode->UISettings->BrushComponentSize, 0);
 
-				const float BrushOriginX = LastMousePosition.X / LandscapeInfo->ComponentSizeQuads - (BrushSize - 1) / 2.0f;
-				const float BrushOriginY = LastMousePosition.Y / LandscapeInfo->ComponentSizeQuads - (BrushSize - 1) / 2.0f;
+				const float BrushOriginX = static_cast<float>(LastMousePosition.X / LandscapeInfo->ComponentSizeQuads - (BrushSize - 1) / 2.0f);
+				const float BrushOriginY = static_cast<float>(LastMousePosition.Y / LandscapeInfo->ComponentSizeQuads - (BrushSize - 1) / 2.0f);
 				const int32 ComponentIndexX = FMath::FloorToInt(BrushOriginX);
 				const int32 ComponentIndexY = FMath::FloorToInt(BrushOriginY);
 				BrushExtentsInclusive.Min = FIntPoint(ComponentIndexX * LandscapeInfo->ComponentSizeQuads, ComponentIndexY * LandscapeInfo->ComponentSizeQuads);
@@ -518,8 +518,8 @@ public:
 		{
 			const int32 BrushSize = FMath::Max(EdMode->UISettings->BrushComponentSize, 0);
 
-			const float BrushOriginX = LastMousePosition.X / LandscapeInfo->ComponentSizeQuads - (BrushSize - 1) / 2.0f;
-			const float BrushOriginY = LastMousePosition.Y / LandscapeInfo->ComponentSizeQuads - (BrushSize - 1) / 2.0f;
+			const float BrushOriginX = static_cast<float>(LastMousePosition.X / LandscapeInfo->ComponentSizeQuads - (BrushSize - 1) / 2.0f);
+			const float BrushOriginY = static_cast<float>(LastMousePosition.Y / LandscapeInfo->ComponentSizeQuads - (BrushSize - 1) / 2.0f);
 			const int32 ComponentIndexX = FMath::FloorToInt(BrushOriginX);
 			const int32 ComponentIndexY = FMath::FloorToInt(BrushOriginY);
 
@@ -662,7 +662,7 @@ public:
 				ULandscapeInfo* LandscapeInfo = Gizmo->TargetLandscapeInfo;
 				if (LandscapeInfo && LandscapeInfo->GetLandscapeProxy())
 				{
-					float ScaleXY = FMath::Abs(LandscapeInfo->DrawScale.X);
+					float ScaleXY = static_cast<float>(FMath::Abs(LandscapeInfo->DrawScale.X));
 					FMatrix LToW = LandscapeInfo->GetLandscapeProxy()->LandscapeActorToWorld().ToMatrixWithScale();
 					FMatrix WToL = LToW.InverseFast();
 
@@ -688,12 +688,12 @@ public:
 					FLinearColor AlphaScaleBias(
 						SquaredScaleXY / FMath::Max<float>(1.f, Gizmo->GetWidth() * DataTexture->GetSizeX()),
 						SquaredScaleXY / FMath::Max<float>(1.f, Gizmo->GetHeight() * DataTexture->GetSizeY()),
-						Gizmo->TextureScale.X,
-						Gizmo->TextureScale.Y
+						static_cast<float>(Gizmo->TextureScale.X),
+						static_cast<float>(Gizmo->TextureScale.Y)
 						);
 					BrushMaterial->SetVectorParameterValue(FName(TEXT("AlphaScaleBias")), AlphaScaleBias);
 
-					float Angle = (-EdMode->CurrentGizmoActor->GetActorRotation().Euler().Z) * PI / 180.0f;
+					float Angle = static_cast<float>(-EdMode->CurrentGizmoActor->GetActorRotation().Euler().Z * PI / 180.0f);
 					FVector4 LandscapeLocation(EdMode->CurrentGizmoActor->GetActorLocation().X, EdMode->CurrentGizmoActor->GetActorLocation().Y, EdMode->CurrentGizmoActor->GetActorLocation().Z, Angle);
 
 					BrushMaterial->SetDoubleVectorParameterValue(FName(TEXT("LandscapeLocation")), LandscapeLocation);
@@ -765,7 +765,7 @@ public:
 		}
 
 		Gizmo->TargetLandscapeInfo = LandscapeInfo;
-		float ScaleXY = FMath::Abs(LandscapeInfo->DrawScale.X);
+		float ScaleXY = static_cast<float>(FMath::Abs(LandscapeInfo->DrawScale.X));
 
 		// Get extent for all components
 		FIntRect Bounds;
@@ -812,10 +812,10 @@ public:
 						FVector TransformedLocal(FMath::Abs(GizmoLocal.X - LW), FMath::Abs(GizmoLocal.Y - LH) * (W / H), 0);
 						float FalloffRadius = LW * EdMode->UISettings->GetCurrentToolBrushFalloff();
 						float SquareRadius = LW - FalloffRadius;
-						float Cos = FMath::Abs(TransformedLocal.X) / TransformedLocal.Size2D();
-						float Sin = FMath::Abs(TransformedLocal.Y) / TransformedLocal.Size2D();
-						float RatioX = FalloffRadius > 0.0f ? 1.0f - FMath::Clamp<float>((FMath::Abs(TransformedLocal.X) - Cos*SquareRadius) / FalloffRadius, 0.0f, 1.0f) : 1.0f;
-						float RatioY = FalloffRadius > 0.0f ? 1.0f - FMath::Clamp<float>((FMath::Abs(TransformedLocal.Y) - Sin*SquareRadius) / FalloffRadius, 0.0f, 1.0f) : 1.0f;
+						float Cos = static_cast<float>(FMath::Abs(TransformedLocal.X) / TransformedLocal.Size2D());
+						float Sin = static_cast<float>(FMath::Abs(TransformedLocal.Y) / TransformedLocal.Size2D());
+						float RatioX = FalloffRadius > 0.0f ? 1.0f - FMath::Clamp<float>(static_cast<float>((FMath::Abs(TransformedLocal.X) - Cos*SquareRadius) / FalloffRadius), 0.0f, 1.0f) : 1.0f;
+						float RatioY = FalloffRadius > 0.0f ? 1.0f - FMath::Clamp<float>(static_cast<float>((FMath::Abs(TransformedLocal.Y) - Sin*SquareRadius) / FalloffRadius), 0.0f, 1.0f) : 1.0f;
 						float Ratio = TransformedLocal.Size2D() > SquareRadius ? RatioX * RatioY : 1.0f; //TransformedLocal.X / LW * TransformedLocal.Y / LW;
 						PaintAmount = Ratio*Ratio*(3 - 2 * Ratio); //FMath::Lerp(SquareFalloff, RectFalloff*RectFalloff, Ratio);
 					}
@@ -1127,7 +1127,7 @@ public:
 		}
 
 		ULandscapeInfo* LandscapeInfo = EdMode->CurrentToolTarget.LandscapeInfo.Get();
-		const float ScaleXY = FMath::Abs(LandscapeInfo->DrawScale.X);
+		const float ScaleXY = static_cast<float>(FMath::Abs(LandscapeInfo->DrawScale.X));
 		const float TotalRadius = EdMode->UISettings->GetCurrentToolBrushRadius() / ScaleXY;
 		const float Radius = (1.0f - EdMode->UISettings->GetCurrentToolBrushFalloff()) * TotalRadius;
 		const float Falloff = EdMode->UISettings->GetCurrentToolBrushFalloff() * TotalRadius;
@@ -1191,8 +1191,8 @@ public:
 				FVector2D SamplePos = FVector2D(X, Y) * Scale + Bias;
 				SamplePos = SamplePos.GetRotated(Angle);
 
-				float ModSampleX = FMath::Fmod(SamplePos.X, (float)SizeX);
-				float ModSampleY = FMath::Fmod(SamplePos.Y, (float)SizeY);
+				float ModSampleX = FMath::Fmod(static_cast<float>(SamplePos.X), static_cast<float>(SizeX));
+				float ModSampleY = FMath::Fmod(static_cast<float>(SamplePos.Y), static_cast<float>(SizeY));
 
 				if (ModSampleX < 0.0f)
 				{
@@ -1238,7 +1238,7 @@ public:
 		ALandscapeProxy* Proxy = EdMode->CurrentToolTarget.LandscapeInfo.IsValid() ? EdMode->CurrentToolTarget.LandscapeInfo->GetLandscapeProxy() : nullptr;
 		if (Proxy)
 		{
-			const float ScaleXY = FMath::Abs(EdMode->CurrentToolTarget.LandscapeInfo->DrawScale.X);
+			const float ScaleXY = static_cast<float>(FMath::Abs(EdMode->CurrentToolTarget.LandscapeInfo->DrawScale.X));
 
 			const bool bHasValidAlphaTextureData = EdMode->UISettings->HasValidAlphaTextureData();
 			int32 SizeX = EdMode->UISettings->AlphaTextureSizeX;
@@ -1264,17 +1264,17 @@ public:
 					LocalOrigin += FVector2D(0.5f, 0.5f).GetRotated(-Angle);
 				}
 				AlphaScaleBias = FLinearColor(
-					Scale.X,
-					Scale.Y,
-					LocalOrigin.X,
-					LocalOrigin.Y);
+					static_cast<float>(Scale.X),
+					static_cast<float>(Scale.Y),
+					static_cast<float>(LocalOrigin.X),
+					static_cast<float>(LocalOrigin.Y));
 			}
 			else 
 			{
 				const FVector2D Scale(EdMode->UISettings->AlphaBrushScale * SizeX, EdMode->UISettings->AlphaBrushScale * SizeY);
 				AlphaScaleBias = FLinearColor(
-					FMath::IsNearlyZero(Scale.X) ? 1.0f : 1.0f / (Scale.X),
-					FMath::IsNearlyZero(Scale.Y) ? 1.0f : 1.0f / (Scale.Y),
+					static_cast<float>(FMath::IsNearlyZero(Scale.X) ? 1.0 : 1.0 / (Scale.X)),
+					static_cast<float>(FMath::IsNearlyZero(Scale.Y) ? 1.0 : 1.0 / (Scale.Y)),
 					EdMode->UISettings->AlphaBrushPanU,
 					EdMode->UISettings->AlphaBrushPanV);
 				Angle = EdMode->UISettings->AlphaBrushRotation;
@@ -1285,10 +1285,10 @@ public:
 			FVector4 LandscapeLocationParam(LandscapeLocation.X, LandscapeLocation.Y, LandscapeLocation.Z, Angle);
 
 			FLinearColor AlphaTextureMask(
-				EdMode->UISettings->AlphaTextureChannel == ELandscapeTextureColorChannel::Red ? 1 : 0, 
-				EdMode->UISettings->AlphaTextureChannel == ELandscapeTextureColorChannel::Green ? 1 : 0, 
-				EdMode->UISettings->AlphaTextureChannel == ELandscapeTextureColorChannel::Blue ? 1 : 0, 
-				EdMode->UISettings->AlphaTextureChannel == ELandscapeTextureColorChannel::Alpha ? 1 : 0);
+				EdMode->UISettings->AlphaTextureChannel == ELandscapeTextureColorChannel::Red ? 1.0f : 0.0f, 
+				EdMode->UISettings->AlphaTextureChannel == ELandscapeTextureColorChannel::Green ? 1.0f : 0.0f, 
+				EdMode->UISettings->AlphaTextureChannel == ELandscapeTextureColorChannel::Blue ? 1.0f : 0.0f, 
+				EdMode->UISettings->AlphaTextureChannel == ELandscapeTextureColorChannel::Alpha ? 1.0f : 0.0f);
 
 			for (const auto& BrushMaterialInstancePair : BrushMaterialInstanceMap)
 			{
@@ -1341,7 +1341,7 @@ public:
 		int32 SizeY = EdMode->UISettings->AlphaTextureSizeY;
 		bool bIsValidTextureSize = (SizeX > 0) && (SizeY > 0);
 
-		float ScaleXY = FMath::Abs(LandscapeInfo->DrawScale.X);
+		float ScaleXY = static_cast<float>(FMath::Abs(LandscapeInfo->DrawScale.X));
 		float Radius = FMath::IsNearlyZero(ScaleXY) ? 0.0f : EdMode->UISettings->GetCurrentToolBrushRadius() / ScaleXY;
 
 		float MaxSize = 2.0f * FMath::Sqrt(FMath::Square(Radius) / 2.0f);
@@ -1523,10 +1523,10 @@ public:
 				EdMode->UISettings->bAlphaBrushAutoRotate ? LastMouseAngle : FMath::DegreesToRadians(EdMode->UISettings->AlphaBrushRotation)
 				);
 
-			FLinearColor AlphaTextureMask(EdMode->UISettings->AlphaTextureChannel == ELandscapeTextureColorChannel::Red ? 1 : 0,
-				EdMode->UISettings->AlphaTextureChannel == ELandscapeTextureColorChannel::Green ? 1 : 0,
-				EdMode->UISettings->AlphaTextureChannel == ELandscapeTextureColorChannel::Blue ? 1 : 0,
-				EdMode->UISettings->AlphaTextureChannel == ELandscapeTextureColorChannel::Alpha ? 1 : 0);
+			FLinearColor AlphaTextureMask(EdMode->UISettings->AlphaTextureChannel == ELandscapeTextureColorChannel::Red ? 1.0f : 0.0f,
+				EdMode->UISettings->AlphaTextureChannel == ELandscapeTextureColorChannel::Green ? 1.0f : 0.0f,
+				EdMode->UISettings->AlphaTextureChannel == ELandscapeTextureColorChannel::Blue ? 1.0f : 0.0f,
+				EdMode->UISettings->AlphaTextureChannel == ELandscapeTextureColorChannel::Alpha ? 1.0f : 0.0f);
 
 
 			for (const auto& BrushMaterialInstancePair : BrushMaterialInstanceMap)
diff --git a/Engine/Source/Editor/LandscapeEditor/Private/LandscapeEdModeComponentTools.cpp b/Engine/Source/Editor/LandscapeEditor/Private/LandscapeEdModeComponentTools.cpp
--- a/Engine/Source/Editor/LandscapeEditor/Private/LandscapeEdModeComponentTools.cpp
+++ b/Engine/Source/Editor/LandscapeEditor/Private/LandscapeEdModeComponentTools.cpp
@@ -87,8 +87,8 @@ public:
 			if (!bInitializedComponentInvert)
 			{
 				// Get the component under the mouse location. Copied from FLandscapeBrushComponent::ApplyBrush()
-				const float MouseX = InteractorPositions[0].Position.X;
-				const float MouseY = InteractorPositions[0].Position.Y;
+				const float MouseX = static_cast<float>(InteractorPositions[0].Position.X);
+				const float MouseY = static_cast<float>(InteractorPositions[0].Position.Y);
 				const int32 MouseComponentIndexX = (MouseX >= 0.0f) ? FMath::FloorToInt(MouseX / LandscapeInfo->ComponentSizeQuads) : FMath::CeilToInt(MouseX / LandscapeInfo->ComponentSizeQuads);
 				const int32 MouseComponentIndexY = (MouseY >= 0.0f) ? FMath::FloorToInt(MouseY / LandscapeInfo->ComponentSizeQuads) : FMath::CeilToInt(MouseY / LandscapeInfo->ComponentSizeQuads);
 				ULandscapeComponent* MouseComponent = LandscapeInfo->XYtoComponentMap.FindRef(FIntPoint(MouseComponentIndexX, MouseComponentIndexY));
@@ -239,7 +239,7 @@ public:
 							LandscapeInfo->SelectedRegion.Remove(Key);
 						}
 
-						DataScanline[X] = FMath::Clamp<int32>(FMath::RoundToInt(Value * 255), 0, 255);
+						DataScanline[X] = static_cast<uint8>(FMath::Clamp<int32>(FMath::RoundToInt(Value * 255), 0, 255));
 					}
 				}
 			}
@@ -1039,8 +1039,8 @@ public:
 		{
 			const int32 BrushSize = FMath::Max(EdMode->UISettings->BrushComponentSize, 0);
 			const int32 ComponentSizeQuads = ToolTarget.LandscapeInfo->ComponentSizeQuads;
-			const float BrushOriginX = LastMousePosition.GetValue().X / ComponentSizeQuads - (BrushSize - 1) / 2.0f;
-			const float BrushOriginY = LastMousePosition.GetValue().Y / ComponentSizeQuads - (BrushSize - 1) / 2.0f;
+			const float BrushOriginX = static_cast<float>(LastMousePosition.GetValue().X / ComponentSizeQuads - (BrushSize - 1) / 2.0);
+			const float BrushOriginY = static_cast<float>(LastMousePosition.GetValue().Y / ComponentSizeQuads - (BrushSize - 1) / 2.0);
 			const int32 ComponentIndexX = FMath::FloorToInt(BrushOriginX);
 			const int32 ComponentIndexY = FMath::FloorToInt(BrushOriginY);
 
@@ -1224,7 +1224,7 @@ public:
 				Cache.GetCachedData(X1, Y1, X2, Y2, Data);
 			}
 
-			const float ScaleXY = LandscapeInfo->DrawScale.X;
+			const float ScaleXY = static_cast<float>(LandscapeInfo->DrawScale.X);
 			float Width = Gizmo->GetWidth();
 			float Height = Gizmo->GetHeight();
 
@@ -1256,8 +1256,8 @@ public:
 				for (int32 X = 0; X < SizeX; ++X)
 				{
 					FVector LandscapeLocal = GizmoLocalToLandscape.TransformPosition(FVector(-W + X, -H + Y, 0));
-					int32 LX = FMath::FloorToInt(LandscapeLocal.X);
-					int32 LY = FMath::FloorToInt(LandscapeLocal.Y);
+					const int32 LX = FMath::FloorToInt32(LandscapeLocal.X);
+					const int32 LY = FMath::FloorToInt32(LandscapeLocal.Y);
 
 					{
 						for (int32 i = -1; (!bApplyToAll && i < 0) || i < LayerNum; ++i)
@@ -1307,8 +1307,8 @@ public:
 							}
 
 							FGizmoPreData LerpedData;
-							float FracX = LandscapeLocal.X - LX;
-							float FracY = LandscapeLocal.Y - LY;
+							const float FracX = static_cast<float>(LandscapeLocal.X - LX);
+							const float FracY = static_cast<float>(LandscapeLocal.Y - LY);
 							LerpedData.Ratio = bFullCopy ? 1.0f :
 								FMath::Lerp(
 								FMath::Lerp(GizmoPreData[0].Ratio, GizmoPreData[1].Ratio, FracX),
@@ -1567,7 +1567,7 @@ public:
 			}
 
 			Gizmo->TargetLandscapeInfo = LandscapeInfo;
-			float ScaleXY = LandscapeInfo->DrawScale.X;
+			float ScaleXY = static_cast<float>(LandscapeInfo->DrawScale.X);
 
 			//LandscapeInfo->Modify();
 
@@ -1654,11 +1654,11 @@ public:
 						GizmoLocal.X *= ScaleX * SignX;
 						GizmoLocal.Y *= ScaleY * SignY;
 
-						int32 LX = FMath::FloorToInt(GizmoLocal.X);
-						int32 LY = FMath::FloorToInt(GizmoLocal.Y);
+						const int32 LX = FMath::FloorToInt32(GizmoLocal.X);
+						const int32 LY = FMath::FloorToInt32(GizmoLocal.Y);
 
-						float FracX = GizmoLocal.X - LX;
-						float FracY = GizmoLocal.Y - LY;
+						const float FracX = static_cast<float>(GizmoLocal.X - LX);
+						const float FracY = static_cast<float>(GizmoLocal.Y - LY);
 
 						FGizmoSelectData* Data00 = Gizmo->SelectedData.Find(FIntPoint(LX, LY));
 						FGizmoSelectData* Data10 = Gizmo->SelectedData.Find(FIntPoint(LX + 1, LY));
@@ -1679,7 +1679,7 @@ public:
 									OriginalValue = Data[index];
 								}
 
-								float Value = LandscapeDataAccess::GetLocalHeight(OriginalValue);
+								float Value = LandscapeDataAccess::GetLocalHeight(static_cast<uint16>(OriginalValue));
 
 								float DestValue = FLandscapeHeightCache::ClampValue(
 									LandscapeDataAccess::GetTexHeight(
@@ -1706,11 +1706,11 @@ public:
 
 								if (bApplyToAll)
 								{
-									HeightData[index] = FMath::Lerp(OriginalValue, DestValue, PaintAmount);
+									HeightData[index] = static_cast<uint16>(FMath::Lerp(OriginalValue, DestValue, PaintAmount));
 								}
 								else
 								{
-									Data[index] = FMath::Lerp(OriginalValue, DestValue, PaintAmount);
+									Data[index] = static_cast<uint16>(FMath::Lerp(OriginalValue, DestValue, PaintAmount));
 								}
 							}
 							else
@@ -1728,22 +1728,22 @@ public:
 									OriginalValue = Data[index];
 								}
 
-								float DestValue = FLandscapeAlphaCache::ClampValue(
+								float DestValue = FLandscapeAlphaCache::ClampValue(static_cast<int32>(
 									FMath::Lerp(
 									FMath::Lerp(Data00 ? FMath::Lerp(OriginalValue, Data00->WeightDataMap.FindRef(LayerInfo), Data00->Ratio) : OriginalValue,
 									Data10 ? FMath::Lerp(OriginalValue, Data10->WeightDataMap.FindRef(LayerInfo), Data10->Ratio) : OriginalValue, FracX),
 									FMath::Lerp(Data01 ? FMath::Lerp(OriginalValue, Data01->WeightDataMap.FindRef(LayerInfo), Data01->Ratio) : OriginalValue,
 									Data11 ? FMath::Lerp(OriginalValue, Data11->WeightDataMap.FindRef(LayerInfo), Data11->Ratio) : OriginalValue, FracX),
 									FracY
-									));
+									)));
 
 								if (bApplyToAll)
 								{
-									WeightDatas[index*LayerNum + i] = FMath::Lerp(OriginalValue, DestValue, PaintAmount);
+									WeightDatas[index*LayerNum + i] = static_cast<uint8>(FMath::Lerp(OriginalValue, DestValue, PaintAmount));
 								}
 								else
 								{
-									Data[index] = FMath::Lerp(OriginalValue, DestValue, PaintAmount);
+									Data[index] = static_cast<typename ToolTarget::CacheClass::DataType>(FMath::Lerp(OriginalValue, DestValue, PaintAmount));
 								}
 							}
 						}
@@ -2132,13 +2132,13 @@ public:
 			auto DragEdge = [&UISettings, &HitLocation, &Transform](const ELandscapeEdge::Type Edge)
 			{
 				int32& ComponentCount = Edge == ELandscapeEdge::X_Negative || Edge == ELandscapeEdge::X_Positive ? UISettings->NewLandscape_ComponentCount.X : UISettings->NewLandscape_ComponentCount.Y;
-				const float Hit = Edge == ELandscapeEdge::X_Negative || Edge == ELandscapeEdge::X_Positive ? HitLocation.X : HitLocation.Y;
+				const float Hit = static_cast<float>(Edge == ELandscapeEdge::X_Negative || Edge == ELandscapeEdge::X_Positive ? HitLocation.X : HitLocation.Y);
 				const float PosOrNeg = Edge == ELandscapeEdge::X_Negative || Edge == ELandscapeEdge::Y_Negative ? -1.0f : 1.0f;
 				const FVector XOrY = Edge == ELandscapeEdge::X_Negative || Edge == ELandscapeEdge::X_Positive ? FVector(1, 0, 0) : FVector(0, 1, 0);
 				
 				const int32 InitialComponentCount = ComponentCount;
 				const int32 Delta = FMath::RoundToInt(Hit - PosOrNeg * static_cast<float>(InitialComponentCount) / 2.0f);
-				ComponentCount = InitialComponentCount + PosOrNeg * Delta;
+				ComponentCount = static_cast<int32>(InitialComponentCount + PosOrNeg * Delta);
 				UISettings->NewLandscape_ClampSize();
 				const float ActualDelta = static_cast<float>(ComponentCount - InitialComponentCount) / 2.0f;
 				UISettings->NewLandscape_Location += PosOrNeg * XOrY * Transform.TransformVector(FVector(ActualDelta, ActualDelta, 0));
@@ -2194,7 +2194,7 @@ public:
 			const int32 ComponentCountX = EdMode->UISettings->NewLandscape_ComponentCount.X;
 			const int32 ComponentCountY = EdMode->UISettings->NewLandscape_ComponentCount.Y;
 			const int32 QuadsPerComponent = EdMode->UISettings->NewLandscape_SectionsPerComponent * EdMode->UISettings->NewLandscape_QuadsPerSection;
-			const float ComponentSize = QuadsPerComponent;
+			const float ComponentSize = static_cast<float>(QuadsPerComponent);
 			const int32 GridSize = EdMode->UISettings->WorldPartitionGridSize;
 			const FVector Offset = EdMode->UISettings->NewLandscape_Location + FTransform(EdMode->UISettings->NewLandscape_Rotation, FVector::ZeroVector, EdMode->UISettings->NewLandscape_Scale).TransformVector(FVector(-ComponentCountX * ComponentSize / 2, -ComponentCountY * ComponentSize / 2, 0));
 			const FTransform Transform = FTransform(EdMode->UISettings->NewLandscape_Rotation, Offset, EdMode->UISettings->NewLandscape_Scale);
@@ -2420,7 +2420,7 @@ public:
 					for (int32 X = 1; X < ComponentCountX * QuadsPerComponent; ++X)
 					{
 						const FLinearColor CurrentColor = GetColor(X);
-						const uint8 DepthPriority = CurrentColor == InnerColor ? SDPG_World : SDPG_Foreground;
+						const uint8 DepthPriority = static_cast<uint8>(CurrentColor == InnerColor ? SDPG_World : SDPG_Foreground);
 						DrawLine({FVector(X, 0, 0), FVector(X, ComponentCountY * ComponentSize, 0)}, CurrentColor, DepthPriority);
 					}
 
@@ -2428,7 +2428,7 @@ public:
 					for (int32 Y = 1; Y < ComponentCountY * QuadsPerComponent; ++Y)
 					{
 						const FLinearColor CurrentColor = GetColor(Y);
-						const uint8 DepthPriority = CurrentColor == InnerColor ? SDPG_World : SDPG_Foreground;
+						const uint8 DepthPriority = static_cast<uint8>(CurrentColor == InnerColor ? SDPG_World : SDPG_Foreground);
 						DrawLine({FVector(0, Y, 0), FVector(ComponentCountX * ComponentSize, Y, 0)}, CurrentColor, DepthPriority);
 					}
 				}
@@ -2675,7 +2675,7 @@ public:
 		const ELevelViewportType ViewportType = ((FEditorViewportClient*)Viewport->GetClient())->ViewportType;
 
 		const int32 ComponentSizeInt = EdMode->UISettings->NewLandscape_SectionsPerComponent * EdMode->UISettings->NewLandscape_QuadsPerSection;
-		const float ComponentSize = ComponentSizeInt;
+		const float ComponentSize = static_cast<float>(ComponentSizeInt);
 		const FTransform GizmoTransform = FTransform(FRotator(0,0,0), GetWidgetLocation(), LandscapeInfo->DrawScale);
 		const int32 Height = EdMode->UISettings->ImportType == ELandscapeImportTransformType::None ? EdMode->UISettings->ImportLandscape_Height : (LandscapeExtent.Height()+1);
 		const int32 Width = EdMode->UISettings->ImportType == ELandscapeImportTransformType::None ? EdMode->UISettings->ImportLandscape_Width: (LandscapeExtent.Width()+1);
diff --git a/Engine/Source/Editor/LandscapeEditor/Private/LandscapeEdModeErosionTools.cpp b/Engine/Source/Editor/LandscapeEditor/Private/LandscapeEdModeErosionTools.cpp
--- a/Engine/Source/Editor/LandscapeEditor/Private/LandscapeEdModeErosionTools.cpp
+++ b/Engine/Source/Editor/LandscapeEditor/Private/LandscapeEdModeErosionTools.cpp
@@ -118,7 +118,7 @@ public:
 		WeightCache.GetCachedData(X1, Y1, X2, Y2, WeightDatas, LayerNum);
 
 		// Apply the brush	
-		uint16 Thresh = UISettings->ErodeThresh;
+		uint16 Thresh = static_cast<uint16>(UISettings->ErodeThresh);
 		int32 WeightMoveThresh = FMath::Min<int32>(FMath::Max<int32>(Thickness >> 2, Thresh), Thickness >> 1);
 
 		TArray<float> CenterWeights;
@@ -197,7 +197,7 @@ public:
 											float WeightDiff = Softness * UISettings->GetCurrentToolStrength() * Pressure * ((float)Slope / SlopeTotal) * BrushValue;
 											//uint16 HeightDiff = (uint16)((SlopeMax - Thresh) * WeightDiff);
 											float HeightDiff = (SlopeMax - Thresh) * WeightDiff;
-											HeightData[Neighbor[Idx]] += HeightDiff;
+											HeightData[Neighbor[Idx]] += static_cast<uint16>(HeightDiff);
 											TotalHeightDiff += HeightDiff;
 
 											if (bWeightApplied)
@@ -219,7 +219,7 @@ public:
 									}
 								}
 
-								HeightData[Center] -= TotalHeightDiff;
+								HeightData[Center] -= static_cast<uint16>(TotalHeightDiff);
 
 								if (bWeightApplied)
 								{
@@ -270,7 +270,7 @@ public:
 				{
 					FNoiseParameter NoiseParam(0, UISettings->ErosionNoiseScale, BrushValue * Thresh * UISettings->GetCurrentToolStrength() * BrushSizeAdjust);
 					float PaintAmount = NoiseModeConversion((ELandscapeToolNoiseMode)UISettings->ErosionNoiseMode, NoiseParam.NoiseAmount, NoiseParam.Sample(X, Y));
-					HeightData[(X - X1) + (Y - Y1)*(1 + X2 - X1)] = FLandscapeHeightCache::ClampValue(HeightData[(X - X1) + (Y - Y1)*(1 + X2 - X1)] + PaintAmount);
+					HeightData[(X - X1) + (Y - Y1)*(1 + X2 - X1)] = FLandscapeHeightCache::ClampValue(static_cast<int32>(HeightData[(X - X1) + (Y - Y1)*(1 + X2 - X1)] + PaintAmount));
 				}
 			}
 		}
@@ -351,13 +351,11 @@ public:
 		X2 += 1;
 		Y2 += 1;
 
-		const int32 LayerNum = this->LandscapeInfo->Layers.Num();
-
 		const int32 Iteration = UISettings->HErodeIterationNum;
-		const uint16 RainAmount = UISettings->RainAmount;
-		const float DissolvingRatio = 0.07 * UISettings->GetCurrentToolStrength() * Pressure;  //0.01;
-		const float EvaporateRatio = 0.5;
-		const float SedimentCapacity = 0.10 * UISettings->SedimentCapacity; //DissolvingRatio; //0.01;
+		const uint16 RainAmount = static_cast<uint16>(UISettings->RainAmount);
+		const float DissolvingRatio = 0.07f * UISettings->GetCurrentToolStrength() * Pressure;  //0.01;
+		const float EvaporateRatio = 0.5f;
+		const float SedimentCapacity = 0.10f * UISettings->SedimentCapacity; //DissolvingRatio; //0.01;
 
 		TArray<uint16> HeightData;
 		LayerHeightDataCache.Initialize(this->LandscapeInfo, bCombinedLayerOperation);
@@ -386,7 +384,7 @@ public:
 				{
 					float PaintAmount = NoiseModeConversion((ELandscapeToolNoiseMode)UISettings->RainDistMode, NoiseParam.NoiseAmount, NoiseParam.Sample(X, Y));
 					if (PaintAmount > 0) // Raining only for positive region...
-						WaterDataScanline[X] += PaintAmount;
+						WaterDataScanline[X] += static_cast<uint16>(PaintAmount);
 				}
 			}
 		}
@@ -421,8 +419,8 @@ public:
 						float DissolvedAmount = DissolvingRatio * WaterData[Center] * BrushValue;
 						if (DissolvedAmount > 0 && HeightData[Center] >= DissolvedAmount)
 						{
-							HeightData[Center] -= DissolvedAmount;
-							SedimentData[Center] += DissolvedAmount;
+							HeightData[Center] -= static_cast<uint16>(DissolvedAmount);
+							SedimentData[Center] += static_cast<uint16>(DissolvedAmount);
 						}
 
 						uint32 TotalHeightDiff = 0;
@@ -461,27 +459,27 @@ public:
 							// This is not mathematically correct, but makes good result
 							if (TotalHeightDiff)
 							{
-								AverageAltitude *= (1.0f - 0.1 * UISettings->GetCurrentToolStrength() * Pressure);
+								AverageAltitude *= static_cast<float>(1.0f - 0.1 * UISettings->GetCurrentToolStrength() * Pressure);
 								//AverageAltitude -= 4000.0f * UISettings->ToolStrength;
 							}
 
-							uint32 WaterTransfer = FMath::Min<uint32>(WaterData[Center], Altitude - (uint32)AverageAltitude) * BrushValue;
+							uint32 WaterTransfer = static_cast<uint32>(FMath::Min<uint32>(WaterData[Center], Altitude - (uint32)AverageAltitude) * BrushValue);
 
 							for (int32 Idx = 0; Idx < NeighborNum; Idx++)
 							{
 								if (AltitudeDiff[Idx] > 0)
 								{
 									uint32 WaterDiff = (uint32)(WaterTransfer * (float)AltitudeDiff[Idx] / TotalAltitudeDiff);
-									WaterData[Neighbor[Idx]] += WaterDiff;
+									WaterData[Neighbor[Idx]] += static_cast<uint16>(WaterDiff);
 									TotalWaterDiff += WaterDiff;
 									uint32 SedimentDiff = (uint32)(SedimentData[Center] * (float)WaterDiff / WaterData[Center]);
-									SedimentData[Neighbor[Idx]] += SedimentDiff;
+									SedimentData[Neighbor[Idx]] += static_cast<uint16>(SedimentDiff);
 									TotalSedimentDiff += SedimentDiff;
 								}
 							}
 
-							WaterData[Center] -= TotalWaterDiff;
-							SedimentData[Center] -= TotalSedimentDiff;
+							WaterData[Center] -= static_cast<uint16>(TotalWaterDiff);
+							SedimentData[Center] -= static_cast<uint16>(TotalSedimentDiff);
 						}
 
 						// evaporation
@@ -493,8 +491,8 @@ public:
 							float SedimentDiff = SedimentData[Center] - SedimentCap;
 							if (SedimentDiff > 0)
 							{
-								SedimentData[Center] -= SedimentDiff;
-								HeightData[Center] = FMath::Clamp<uint16>(HeightData[Center] + SedimentDiff, 0, LandscapeDataAccess::MaxValue);
+								SedimentData[Center] -= static_cast<uint16>(SedimentDiff);
+								HeightData[Center] = FMath::Clamp<uint16>(static_cast<uint16>(HeightData[Center] + SedimentDiff), 0, LandscapeDataAccess::MaxValue);
 							}
 						}
 					}
diff --git a/Engine/Source/Editor/LandscapeEditor/Private/LandscapeEdModeMirrorTool.cpp b/Engine/Source/Editor/LandscapeEditor/Private/LandscapeEdModeMirrorTool.cpp
--- a/Engine/Source/Editor/LandscapeEditor/Private/LandscapeEdModeMirrorTool.cpp
+++ b/Engine/Source/Editor/LandscapeEditor/Private/LandscapeEdModeMirrorTool.cpp
@@ -141,7 +141,7 @@ public:
 					MirrorPlaneScale.Y = (MaxX - MinX) / 2.0f;
 				}
 
-				MirrorPoint3D.Z = GetLocalZAtPoint(LandscapeInfo, FMath::RoundToInt(MirrorPoint3D.X), FMath::RoundToInt(MirrorPoint3D.Y));
+				MirrorPoint3D.Z = GetLocalZAtPoint(LandscapeInfo, FMath::RoundToInt32(MirrorPoint3D.X), FMath::RoundToInt32(MirrorPoint3D.Y));
 				MirrorPoint3D = LandscapeToWorld.TransformPosition(MirrorPoint3D);
 
 				FMatrix Matrix;
@@ -166,7 +166,7 @@ public:
 				const FBox Box = FBox(FVector(-1, -1, 0), FVector(+1, +1, 0));
 				DrawWireBox(PDI, Matrix, Box, FLinearColor::Green, SDPG_World);
 
-				const float LandscapeScaleRatio = LandscapeToWorld.GetScale3D().Z / LandscapeToWorld.GetScale3D().X;
+				const float LandscapeScaleRatio = static_cast<float>(LandscapeToWorld.GetScale3D().Z / LandscapeToWorld.GetScale3D().X);
 				FVector2D UVScale = FVector2D(FMath::RoundToFloat(MirrorPlaneScale.Y / 10), FMath::RoundToFloat(MirrorPlaneScale.Z * LandscapeScaleRatio / 10 / 2) * 2);
 				MirrorPlaneMaterial->SetVectorParameterValue(FName("GridSize"), FVector(UVScale, 0));
 				DrawPlane10x10(PDI, Matrix, 1, FVector2D(0, 0), FVector2D(1, 1), MirrorPlaneMaterial->GetRenderProxy(), SDPG_World);
@@ -259,7 +259,7 @@ public:
 			{
 				MirrorPoint3D.Y = EdMode->UISettings->MirrorPoint.Y;
 			}
-			MirrorPoint3D.Z = GetLocalZAtPoint(LandscapeInfo, FMath::RoundToInt(MirrorPoint3D.X), FMath::RoundToInt(MirrorPoint3D.Y));
+			MirrorPoint3D.Z = GetLocalZAtPoint(LandscapeInfo, FMath::RoundToInt32(MirrorPoint3D.X), FMath::RoundToInt32(MirrorPoint3D.Y));
 			MirrorPoint3D = LandscapeToWorld.TransformPosition(MirrorPoint3D);
 			MirrorPoint3D.Z += 1000.f; // place the widget a little off the ground for better visibility
 			return MirrorPoint3D;
@@ -541,7 +541,7 @@ public:
 		case ELandscapeMirrorOperation::PlusXToMinusX:
 		case ELandscapeMirrorOperation::RotatePlusXToMinusX:
 			{
-				MirrorPos = FMath::RoundToInt(MirrorPoint.X);
+				MirrorPos = FMath::RoundToInt32(MirrorPoint.X);
 				if (MirrorPos <= MinX || MirrorPos >= MaxX)
 				{
 					return;
@@ -575,7 +575,7 @@ public:
 		case ELandscapeMirrorOperation::PlusYToMinusY:
 		case ELandscapeMirrorOperation::RotatePlusYToMinusY:
 			{
-				MirrorPos = FMath::RoundToInt(MirrorPoint.Y);
+				MirrorPos = FMath::RoundToInt32(MirrorPoint.Y);
 				if (MirrorPos <= MinY || MirrorPos >= MaxY)
 				{
 					return;
diff --git a/Engine/Source/Editor/LandscapeEditor/Private/LandscapeEdModePaintTools.cpp b/Engine/Source/Editor/LandscapeEditor/Private/LandscapeEdModePaintTools.cpp
--- a/Engine/Source/Editor/LandscapeEditor/Private/LandscapeEdModePaintTools.cpp
+++ b/Engine/Source/Editor/LandscapeEditor/Private/LandscapeEdModePaintTools.cpp
@@ -252,7 +252,7 @@ public:
 
 		// Adjust strength based on brush size and drawscale, so strength 1 = one hemisphere
 		const float AdjustedStrength = ToolTarget::StrengthMultiplier(this->LandscapeInfo, UISettings->GetCurrentToolBrushRadius());
-		FWeightmapToolTarget::CacheClass::DataType DestValue = FWeightmapToolTarget::CacheClass::ClampValue(255.0f * UISettings->WeightTargetValue);
+		FWeightmapToolTarget::CacheClass::DataType DestValue = FWeightmapToolTarget::CacheClass::ClampValue(static_cast<int32>(255.0f * UISettings->WeightTargetValue));
 
 		float PaintStrength = UISettings->GetCurrentToolStrength() * Pressure * AdjustedStrength;
 
@@ -404,11 +404,11 @@ public:
 					int32 Delta = DataScanline[X] - FlattenHeight;
 					if (Delta > 0)
 					{
-						DataScanline[X] = FMath::FloorToInt(FMath::Lerp((float)DataScanline[X], (float)FlattenHeight, Strength));
+						DataScanline[X] = static_cast<ToolTarget::CacheClass::DataType>(FMath::FloorToInt(FMath::Lerp((float)DataScanline[X], (float)FlattenHeight, Strength)));
 					}
 					else
 					{
-						DataScanline[X] = FMath::CeilToInt(FMath::Lerp((float)DataScanline[X], (float)FlattenHeight, Strength));
+						DataScanline[X] = static_cast<ToolTarget::CacheClass::DataType>(FMath::CeilToInt(FMath::Lerp((float)DataScanline[X], (float)FlattenHeight, Strength)));
 					}
 				}
 			}
@@ -484,7 +484,7 @@ public:
 		const float AdjustedStrength = ToolTarget::StrengthMultiplier(this->LandscapeInfo, UISettings->GetCurrentToolBrushRadius());
 
 		float SculptStrength = UISettings->GetCurrentToolStrength() * Pressure * AdjustedStrength;
-		const float DeltaTime = FMath::Min<float>(FApp::GetDeltaTime(), 0.1f); // Under 10 fps slow down paint speed
+		const float DeltaTime = FMath::Min<float>(static_cast<float>(FApp::GetDeltaTime()), 0.1f); // Under 10 fps slow down paint speed
 		SculptStrength *= DeltaTime * 3.0f; // * 3.0f to partially compensate for impact of DeltaTime on slowing the tools down compared to the old framerate-dependent version
 
 		if (SculptStrength <= 0.0f)
@@ -615,7 +615,7 @@ public:
 						WorldLoc.Z = (BrushPlane.W - BrushPlane.X*WorldLoc.X - BrushPlane.Y*WorldLoc.Y) / BrushPlane.Z;
 
 						// Painted amount lerps based on brush falloff.
-						float PaintValue = FMath::Lerp<float>((float)SourceValue, FromWorld.TransformPosition(WorldLoc).Z, BrushValue);
+						float PaintValue = FMath::Lerp<float>(SourceValue, static_cast<float>(FromWorld.TransformPosition(WorldLoc).Z), BrushValue);
 
 						if (bInvert)
 						{
@@ -852,8 +852,8 @@ public:
 		if (InEdMode->UISettings->bUseFlattenTarget && bTargetIsHeightmap)
 		{
 			FTransform LocalToWorld = InTarget.LandscapeInfo->GetLandscapeProxy()->ActorToWorld();
-			float Height = (InEdMode->UISettings->FlattenTarget - LocalToWorld.GetTranslation().Z) / LocalToWorld.GetScale3D().Z;
-			FlattenValue = LandscapeDataAccess::GetTexHeight(Height);
+			float Height = static_cast<float>((InEdMode->UISettings->FlattenTarget - LocalToWorld.GetTranslation().Z) / LocalToWorld.GetScale3D().Z);
+			FlattenValue = static_cast<typename ToolTarget::CacheClass::DataType>(LandscapeDataAccess::GetTexHeight(Height));
 			bInitializedFlattenValue = true;
 		}
 	}
@@ -876,8 +876,8 @@ public:
 		if (!bInitializedFlattenValue || (UISettings->bPickValuePerApply && bTargetIsHeightmap))
 		{
 			bInitializedFlattenValue = false;
-			float FlattenX = InteractorPositions[0].Position.X;
-			float FlattenY = InteractorPositions[0].Position.Y;
+			float FlattenX = static_cast<float>(InteractorPositions[0].Position.X);
+			float FlattenY = static_cast<float>(InteractorPositions[0].Position.Y);
 			int32 FlattenHeightX = FMath::FloorToInt(FlattenX);
 			int32 FlattenHeightY = FMath::FloorToInt(FlattenY);
 
@@ -910,7 +910,7 @@ public:
 				this->Cache.CacheData(FlattenHeightX, FlattenHeightY, FlattenHeightX + 1, FlattenHeightY + 1);
 				InterpolatedValue = this->Cache.GetValue(FlattenX, FlattenY);
 			}
-			FlattenValue = InterpolatedValue;
+			FlattenValue = static_cast<typename ToolTarget::CacheClass::DataType>(InterpolatedValue);
 
 			if (bUseSlopeFlatten && bTargetIsHeightmap)
 			{
@@ -929,7 +929,7 @@ public:
 				{
 					FlattenNormal = this->Cache.GetNormal(FlattenHeightX, FlattenHeightY);
 				}
-				FlattenPlaneDist = -(FlattenNormal | FVector(FlattenX, FlattenY, InterpolatedValue));
+				FlattenPlaneDist = static_cast<float>(-(FlattenNormal | FVector(FlattenX, FlattenY, InterpolatedValue)));
 			}
 
 			bInitializedFlattenValue = true;
@@ -986,8 +986,8 @@ public:
 							if (bTargetIsHeightmap)
 							{
 								const FTransform& LocalToWorld = this->Target.LandscapeInfo->GetLandscapeProxy()->ActorToWorld();
-								float ScaleZ = LocalToWorld.GetScale3D().Z;
-								float TranslateZ = LocalToWorld.GetTranslation().Z;
+								float ScaleZ = static_cast<float>(LocalToWorld.GetScale3D().Z);
+								float TranslateZ = static_cast<float>(LocalToWorld.GetTranslation().Z);
 								float TerraceInterval = FMath::Clamp(UISettings->TerraceInterval, 1.0f, 32768.f);
 								float Smoothness = UISettings->TerraceSmooth;
 								float WorldHeight = LandscapeDataAccess::GetLocalHeight(DataScanline[X]);
@@ -1010,15 +1010,15 @@ public:
 
 								float FinalHeight = FMath::Lerp(CurrentHeight, WorldHeight, Strength);
 								FinalHeight = (FinalHeight - TranslateZ) / ScaleZ;
-								DataScanline[X] = LandscapeDataAccess::GetTexHeight(FinalHeight);
+								DataScanline[X] = static_cast<typename ToolTarget::CacheClass::DataType>(LandscapeDataAccess::GetTexHeight(FinalHeight));
 							}
 							break;
 						case ELandscapeToolFlattenMode::Interval:
 							if (bTargetIsHeightmap)
 							{
 								const FTransform& LocalToWorld = this->Target.LandscapeInfo->GetLandscapeProxy()->ActorToWorld();
-								float ScaleZ = LocalToWorld.GetScale3D().Z;
-								float TranslateZ = LocalToWorld.GetTranslation().Z;
+								float ScaleZ = static_cast<float>(LocalToWorld.GetScale3D().Z);
+								float TranslateZ = static_cast<float>(LocalToWorld.GetTranslation().Z);
 								float TerraceInterval = UISettings->TerraceInterval;
 								float TargetHeight = LandscapeDataAccess::GetLocalHeight(FlattenValue);
 								float CurrentHeight = LandscapeDataAccess::GetLocalHeight(DataScanline[X]);
@@ -1032,64 +1032,64 @@ public:
 							
 								//back to local space of landscape object
 								TargetHeight = (TargetHeight - TranslateZ) / ScaleZ;
-								DataScanline[X] = LandscapeDataAccess::GetTexHeight(TargetHeight);
+								DataScanline[X] = static_cast<typename ToolTarget::CacheClass::DataType>(LandscapeDataAccess::GetTexHeight(TargetHeight));
 
 							}
 							break;
 						case ELandscapeToolFlattenMode::Raise:
 							if (Delta < 0)
 							{
-								DataScanline[X] = FMath::CeilToInt(FMath::Lerp((float)DataScanline[X], (float)FlattenValue, Strength));
+								DataScanline[X] = static_cast<typename ToolTarget::CacheClass::DataType>(FMath::CeilToInt(FMath::Lerp((float)DataScanline[X], (float)FlattenValue, Strength)));
 							}
 							break;
 						case ELandscapeToolFlattenMode::Lower:
 							if (Delta > 0)
 							{
-								DataScanline[X] = FMath::FloorToInt(FMath::Lerp((float)DataScanline[X], (float)FlattenValue, Strength));
+								DataScanline[X] = static_cast<typename ToolTarget::CacheClass::DataType>(FMath::FloorToInt(FMath::Lerp((float)DataScanline[X], (float)FlattenValue, Strength)));
 							}
 							break;
 						default:
 						case ELandscapeToolFlattenMode::Both:
 							if (Delta > 0)
 							{
-								DataScanline[X] = FMath::FloorToInt(FMath::Lerp((float)DataScanline[X], (float)FlattenValue, Strength));
+								DataScanline[X] = static_cast<typename ToolTarget::CacheClass::DataType>(FMath::FloorToInt(FMath::Lerp((float)DataScanline[X], (float)FlattenValue, Strength)));
 							}
 							else
 							{
-								DataScanline[X] = FMath::CeilToInt(FMath::Lerp((float)DataScanline[X], (float)FlattenValue, Strength));
+								DataScanline[X] = static_cast<typename ToolTarget::CacheClass::DataType>(FMath::CeilToInt(FMath::Lerp((float)DataScanline[X], (float)FlattenValue, Strength)));
 							}
 							break;
 						}
 					}
 					else
 					{
-						typename ToolTarget::CacheClass::DataType DestValue = -(FlattenNormal.X * X + FlattenNormal.Y * Y + FlattenPlaneDist) / FlattenNormal.Z;
+						typename ToolTarget::CacheClass::DataType DestValue = static_cast<typename ToolTarget::CacheClass::DataType>(-(FlattenNormal.X * X + FlattenNormal.Y * Y + FlattenPlaneDist) / FlattenNormal.Z);
 						//float PlaneDist = FlattenNormal | FVector(X, Y, HeightData(HeightDataIndex)) + FlattenPlaneDist;
-						float PlaneDist = DataScanline[X] - DestValue;
-						DestValue = DataScanline[X] - PlaneDist * Strength;
+						float PlaneDist = static_cast<float>(DataScanline[X] - DestValue);
+						DestValue = static_cast<typename ToolTarget::CacheClass::DataType>(DataScanline[X] - PlaneDist * Strength);
 						switch (UISettings->FlattenMode)
 						{
 						case ELandscapeToolFlattenMode::Raise:
 							if (PlaneDist < 0)
 							{
-								DataScanline[X] = FMath::CeilToInt(FMath::Lerp((float)DataScanline[X], (float)DestValue, Strength));
+								DataScanline[X] = static_cast<typename ToolTarget::CacheClass::DataType>(FMath::CeilToInt(FMath::Lerp((float)DataScanline[X], (float)DestValue, Strength)));
 							}
 							break;
 						case ELandscapeToolFlattenMode::Lower:
 							if (PlaneDist > 0)
 							{
-								DataScanline[X] = FMath::FloorToInt(FMath::Lerp((float)DataScanline[X], (float)DestValue, Strength));
+								DataScanline[X] = static_cast<typename ToolTarget::CacheClass::DataType>(FMath::FloorToInt(FMath::Lerp((float)DataScanline[X], (float)DestValue, Strength)));
 							}
 							break;
 						default:
 						case ELandscapeToolFlattenMode::Both:
 							if (PlaneDist > 0)
 							{
-								DataScanline[X] = FMath::FloorToInt(FMath::Lerp((float)DataScanline[X], (float)DestValue, Strength));
+								DataScanline[X] = static_cast<typename ToolTarget::CacheClass::DataType>(FMath::FloorToInt(FMath::Lerp((float)DataScanline[X], (float)DestValue, Strength)));
 							}
 							else
 							{
-								DataScanline[X] = FMath::CeilToInt(FMath::Lerp((float)DataScanline[X], (float)DestValue, Strength));
+								DataScanline[X] = static_cast<typename ToolTarget::CacheClass::DataType>(FMath::CeilToInt(FMath::Lerp((float)DataScanline[X], (float)DestValue, Strength)));
 							}
 							break;
 						}
@@ -1193,11 +1193,11 @@ public:
 				HeightmapFlattenPreviewComponent->SetRelativeLocation(Origin, false);
 
 				// Clamp the value to the height map
-				uint16 TexHeight = LandscapeDataAccess::GetTexHeight(MousePosition.Z);
+				uint16 TexHeight = LandscapeDataAccess::GetTexHeight(static_cast<float>(MousePosition.Z));
 				float Height = LandscapeDataAccess::GetLocalHeight(TexHeight);
 
 				// Convert the height back to world space
-				this->EdMode->UISettings->FlattenEyeDropperModeDesiredTarget = (Height * LocalToWorld.GetScale3D().Z) + LocalToWorld.GetTranslation().Z;
+				this->EdMode->UISettings->FlattenEyeDropperModeDesiredTarget = static_cast<float>(Height * LocalToWorld.GetScale3D().Z + LocalToWorld.GetTranslation().Z);
 			}
 		}
 
@@ -1342,7 +1342,7 @@ public:
 						float TotalStrength = BrushValue * UISettings->GetCurrentToolStrength() * Pressure * ToolTarget::StrengthMultiplier(this->LandscapeInfo, UISettings->GetCurrentToolBrushRadius());
 						FNoiseParameter NoiseParam(0, UISettings->NoiseScale, TotalStrength * BrushSizeAdjust);
 						float PaintAmount = NoiseModeConversion(UISettings->NoiseMode, NoiseParam.NoiseAmount, NoiseParam.Sample(X, Y));
-						DataScanline[X] = ToolTarget::CacheClass::ClampValue(OriginalValue + PaintAmount);
+						DataScanline[X] = static_cast<typename ToolTarget::CacheClass::DataType>(ToolTarget::CacheClass::ClampValue(static_cast<int32>(OriginalValue + PaintAmount)));
 					}
 				}
 			}
diff --git a/Engine/Source/Editor/LandscapeEditor/Private/LandscapeEdModeRampTool.cpp b/Engine/Source/Editor/LandscapeEditor/Private/LandscapeEdModeRampTool.cpp
--- a/Engine/Source/Editor/LandscapeEditor/Private/LandscapeEdModeRampTool.cpp
+++ b/Engine/Source/Editor/LandscapeEditor/Private/LandscapeEdModeRampTool.cpp
@@ -55,11 +55,11 @@ protected:
 
 	void ProcessPixel(int32 X, int32 Y, const InterpolantType& Interpolant, bool BackFacing)
 	{
-		const float CosInterpX = (Interpolant.X >= 1 ? 1 : 0.5f - 0.5f * FMath::Cos(Interpolant.X * PI));
+		const float CosInterpX = static_cast<float>(Interpolant.X >= 1 ? 1 : 0.5 - 0.5 * FMath::Cos(Interpolant.X * PI));
 		const float Alpha = CosInterpX;
 		uint16& Dest = Data[(Y - MinY)*(1 + MaxX - MinX) + X - MinX];
 		float Value = FMath::Lerp((float)Dest, (float)Interpolant.Y, Alpha);
-		uint16 DValue = (uint32)FMath::Clamp<float>(Value, 0, LandscapeDataAccess::MaxValue);
+		uint16 DValue = static_cast<uint16>(FMath::Clamp<float>(Value, 0, LandscapeDataAccess::MaxValue));
 		if ((bRaiseTerrain && DValue > Dest) ||
 			(bLowerTerrain && DValue < Dest))
 		{
@@ -245,8 +245,8 @@ public:
 		{
 			if (SelectedPoint != INDEX_NONE)
 			{
-				const int32 MinX = FMath::FloorToInt(Points[SelectedPoint].X);
-				const int32 MinY = FMath::FloorToInt(Points[SelectedPoint].Y);
+				const int32 MinX = FMath::FloorToInt32(Points[SelectedPoint].X);
+				const int32 MinY = FMath::FloorToInt32(Points[SelectedPoint].Y);
 				const int32 MaxX = MinX + 1;
 				const int32 MaxY = MinY + 1;
 
@@ -342,11 +342,11 @@ public:
 			float SpriteScale = EdMode->UISettings->RampWidth / 4;
 			if (NumPoints > 1)
 			{
-				SpriteScale = FMath::Min(SpriteScale, (WorldPoints[1] - WorldPoints[0]).Size() / 2);
+				SpriteScale = FMath::Min(SpriteScale, static_cast<float>((WorldPoints[1] - WorldPoints[0]).Size() / 2));
 			}
 			SpriteScale = FMath::Clamp<float>(SpriteScale, 10, 500);
 
-			for (int32 i = 0; i < NumPoints; i++)
+			for (int8 i = 0; i < NumPoints; i++)
 			{
 				const FLinearColor SpriteColor = (i == SelectedPoint) ? SelectedSpriteColor : FLinearColor::White;
 
@@ -357,8 +357,8 @@ public:
 					SpriteTexture->GetResource(),
 					SpriteColor,
 					SDPG_Foreground,
-					0, SpriteTexture->GetResource()->GetSizeX(),
-					0, SpriteTexture->GetResource()->GetSizeY(),
+					0, static_cast<float>(SpriteTexture->GetResource()->GetSizeX()),
+					0, static_cast<float>(SpriteTexture->GetResource()->GetSizeY()),
 					SE_BLEND_Masked);
 			}
 			PDI->SetHitProxy(NULL);
@@ -513,14 +513,15 @@ public:
 		OuterVerts[1][0] = FVector2D(Points[1]) - OuterSide;
 		OuterVerts[1][1] = FVector2D(Points[1]) + OuterSide;
 
-		float Heights[2];
-		Heights[0] = Points[0].Z * LANDSCAPE_INV_ZSCALE + LandscapeDataAccess::MidValue;
-		Heights[1] = Points[1].Z * LANDSCAPE_INV_ZSCALE + LandscapeDataAccess::MidValue;
+		const double Heights[2] = {
+			Points[0].Z * LANDSCAPE_INV_ZSCALE + LandscapeDataAccess::MidValue,
+			Points[1].Z * LANDSCAPE_INV_ZSCALE + LandscapeDataAccess::MidValue
+		};
 
-		int32 MinX = FMath::CeilToInt(FMath::Min(FMath::Min(OuterVerts[0][0].X, OuterVerts[0][1].X), FMath::Min(OuterVerts[1][0].X, OuterVerts[1][1].X))) - 1; // +/- 1 to make sure we have enough data for calculating correct normals
-		int32 MinY = FMath::CeilToInt(FMath::Min(FMath::Min(OuterVerts[0][0].Y, OuterVerts[0][1].Y), FMath::Min(OuterVerts[1][0].Y, OuterVerts[1][1].Y))) - 1;
-		int32 MaxX = FMath::FloorToInt(FMath::Max(FMath::Max(OuterVerts[0][0].X, OuterVerts[0][1].X), FMath::Max(OuterVerts[1][0].X, OuterVerts[1][1].X))) + 1;
-		int32 MaxY = FMath::FloorToInt(FMath::Max(FMath::Max(OuterVerts[0][0].Y, OuterVerts[0][1].Y), FMath::Max(OuterVerts[1][0].Y, OuterVerts[1][1].Y))) + 1;
+		int32 MinX = FMath::CeilToInt32(FMath::Min(FMath::Min(OuterVerts[0][0].X, OuterVerts[0][1].X), FMath::Min(OuterVerts[1][0].X, OuterVerts[1][1].X))) - 1; // +/- 1 to make sure we have enough data for calculating correct normals
+		int32 MinY = FMath::CeilToInt32(FMath::Min(FMath::Min(OuterVerts[0][0].Y, OuterVerts[0][1].Y), FMath::Min(OuterVerts[1][0].Y, OuterVerts[1][1].Y))) - 1;
+		int32 MaxX = FMath::FloorToInt32(FMath::Max(FMath::Max(OuterVerts[0][0].X, OuterVerts[0][1].X), FMath::Max(OuterVerts[1][0].X, OuterVerts[1][1].X))) + 1;
+		int32 MaxY = FMath::FloorToInt32(FMath::Max(FMath::Max(OuterVerts[0][0].Y, OuterVerts[0][1].Y), FMath::Max(OuterVerts[1][0].Y, OuterVerts[1][1].Y))) + 1;
 
 		// I'd dearly love to use FIntRect in this code, but Landscape works with "Inclusive Max" and FIntRect is "Exclusive Max"
 		int32 LandscapeMinX, LandscapeMinY, LandscapeMaxX, LandscapeMaxY;
diff --git a/Engine/Source/Editor/LandscapeEditor/Private/LandscapeEdModeSplineTools.cpp b/Engine/Source/Editor/LandscapeEditor/Private/LandscapeEdModeSplineTools.cpp
--- a/Engine/Source/Editor/LandscapeEditor/Private/LandscapeEdModeSplineTools.cpp
+++ b/Engine/Source/Editor/LandscapeEditor/Private/LandscapeEdModeSplineTools.cpp
@@ -324,7 +324,7 @@ public:
 		End->GetConnectionLocationAndRotation(NewSegment->Connections[1].SocketName, EndLocation, EndRotation);
 
 		// Set up tangent lengths
-		NewSegment->Connections[0].TangentLen = (EndLocation - StartLocation).Size();
+		NewSegment->Connections[0].TangentLen = static_cast<float>((EndLocation - StartLocation).Size());
 		NewSegment->Connections[1].TangentLen = NewSegment->Connections[0].TangentLen;
 
 		NewSegment->AutoFlipTangents();
@@ -432,7 +432,7 @@ public:
 		End->GetConnectionLocationAndRotation(Segment->Connections[1].SocketName, EndLocation, EndRotation);
 
 		// Set up tangent lengths
-		Segment->Connections[0].TangentLen = (EndLocation - StartLocation).Size();
+		Segment->Connections[0].TangentLen = static_cast<float>((EndLocation - StartLocation).Size());
 		Segment->Connections[1].TangentLen = Segment->Connections[0].TangentLen;
 
 		Segment->AutoFlipTangents();
@@ -694,7 +694,7 @@ public:
 			DuplicateCacheSplitSegmentParam = t;
 			DuplicateCacheSplitSegmentTangentLenStart = Segment->Connections[0].TangentLen;
 			DuplicateCacheSplitSegmentTangentLenEnd = Segment->Connections[1].TangentLen;
-			DuplicateCacheSplitSegmentTangentLen = Tangent.Size();
+			DuplicateCacheSplitSegmentTangentLen = static_cast<float>(Tangent.Size());
 		}
 
 		ULandscapeSplineControlPoint* NewControlPoint = NewObject<ULandscapeSplineControlPoint>(SplinesComponent, NAME_None, RF_Transactional);
@@ -725,7 +725,7 @@ public:
 		SplinesComponent->Segments.Add(NewSegment);
 
 		NewSegment->Connections[0].ControlPoint = NewControlPoint;
-		NewSegment->Connections[0].TangentLen = Tangent.Size() * (1 - t);
+		NewSegment->Connections[0].TangentLen = static_cast<float>(Tangent.Size() * (1 - t));
 		NewSegment->Connections[0].ControlPoint->ConnectedSegments.Add(FLandscapeSplineConnection(NewSegment, 0));
 		NewSegment->Connections[1].ControlPoint = Segment->Connections[1].ControlPoint;
 		NewSegment->Connections[1].TangentLen = Segment->Connections[1].TangentLen * (1 - t);
@@ -749,7 +749,7 @@ public:
 		Segment->Connections[0].TangentLen *= t;
 		Segment->Connections[1].ControlPoint->ConnectedSegments.Remove(FLandscapeSplineConnection(Segment, 1));
 		Segment->Connections[1].ControlPoint = NewControlPoint;
-		Segment->Connections[1].TangentLen = -Tangent.Size() * t;
+		Segment->Connections[1].TangentLen = static_cast<float>(-Tangent.Size() * t);
 		Segment->Connections[1].ControlPoint->ConnectedSegments.Add(FLandscapeSplineConnection(Segment, 1));
 
 		Segment->UpdateSplinePoints();
@@ -784,24 +784,17 @@ public:
 		Segment->FindNearest(LocalLocation, t0, Location0, Tangent0);
 		NewSegment->FindNearest(LocalLocation, t1, Location1, Tangent1);
 
-		float Len0 = Tangent0.Size();
-		float Len1 = Tangent1.Size();
-
 		ULandscapeSplineSegment* UseSegment;
 		if (FVector::Distance(LocalLocation, Location0) < FVector::Distance(LocalLocation, Location1))
 		{
 			t = DuplicateCacheSplitSegmentParam * t0;
 			tseg = t0;
-			Location = Location0;
-			Tangent = Tangent0;
 			UseSegment = Segment;
 		}
 		else
 		{
 			t = DuplicateCacheSplitSegmentParam + (1 - DuplicateCacheSplitSegmentParam) * t1;
 			tseg = t1;
-			Location = Location1;
-			Tangent = Tangent1;
 			UseSegment = NewSegment;
 		}
 		DuplicateCacheSplitSegmentParam = t;
@@ -1601,7 +1594,7 @@ public:
 					continue;
 				}
 
-				float CurrentAngle = FMath::Acos(FVector::DotProduct(LocalDrag, SegmentDirection) / (LocalDrag.Size() * SegmentDirection.Size()));
+				float CurrentAngle = static_cast<float>(FMath::Acos(FVector::DotProduct(LocalDrag, SegmentDirection) / (LocalDrag.Size() * SegmentDirection.Size())));
 
 				// Create a new segment if there is no segment within 90 degrees of drag direction.
 				// Otherwise split segment that is closest to the drag direction.
@@ -1718,9 +1711,9 @@ public:
 			FVector ForwardVector = FQuatRotationMatrix(StartRotation.Quaternion()).TransformVector(FVector(1.0f, 0.0f, 0.0f));
 
 			FVector DragLocal = SplinesComponent->GetComponentTransform().InverseTransformVector(Drag);
-			float Angle = FMath::Acos(FVector::DotProduct(DragLocal, ForwardVector) / DragLocal.Size());
+			float Angle = static_cast<float>(FMath::Acos(FVector::DotProduct(DragLocal, ForwardVector) / DragLocal.Size()));
 			float OldTangentLen = Connection.TangentLen;
-			Connection.TangentLen = DraggingTangent_Length + (Angle < HALF_PI ? 2.0 : -2.0) * DragLocal.Size();
+			Connection.TangentLen = static_cast<float>(DraggingTangent_Length + (Angle < HALF_PI ? 2.0 : -2.0) * DragLocal.Size());
 
 			// Disallow a tangent of exactly 0 and don't allow tangents to flip
 			if ((Connection.TangentLen > 0 && OldTangentLen < 0) || 
diff --git a/Engine/Source/Editor/LandscapeEditor/Private/LandscapeEdModeTools.h b/Engine/Source/Editor/LandscapeEditor/Private/LandscapeEdModeTools.h
--- a/Engine/Source/Editor/LandscapeEditor/Private/LandscapeEdModeTools.h
+++ b/Engine/Source/Editor/LandscapeEditor/Private/LandscapeEdModeTools.h
@@ -53,7 +53,7 @@ struct FNoiseParameter
 		{
 			for (uint32 Octave = 0; Octave < 4; Octave++)
 			{
-				float	OctaveShift = 1 << Octave;
+				float	OctaveShift = static_cast<float>(1 << Octave);
 				float	OctaveScale = OctaveShift / NoiseScale;
 				Noise += PerlinNoise2D(X * OctaveScale, Y * OctaveScale) / OctaveShift;
 			}
@@ -71,7 +71,7 @@ struct FNoiseParameter
 		{
 			for (uint32 Octave = 0; Octave < 4; Octave++)
 			{
-				float	OctaveShift = 1 << Octave;
+				float	OctaveShift = static_cast<float>(1 << Octave);
 				float	OctaveAmplitude = NoiseAmount / OctaveShift;
 
 				// Attempt to avoid calculating noise if the test value is outside of the noise amplitude.
@@ -187,13 +187,13 @@ inline void LowPassFilter(int32 X1, int32 Y1, int32 X2, int32 Y2, FLandscapeBrus
 
 	for (int32 Y = Y1 + 1; Y <= Y2 - 1; Y++)
 	{
-		auto* DataScanline = Data.GetData() + (Y - Y1) * (X2 - X1 + 1) + (0 - X1);
-		auto* bufScanline = buf + (Y - (Y1 + 1)) * Dims[1] + (0 - (X1 + 1));
+		const typename TArray<DataType>::ElementType* DataScanline = Data.GetData() + (Y - Y1) * (X2 - X1 + 1) + (0 - X1);
+		kiss_fft_cpx* BufScanline = buf + (Y - (Y1 + 1)) * Dims[1] + (0 - (X1 + 1));
 
 		for (int32 X = X1 + 1; X <= X2 - 1; X++)
 		{
-			bufScanline[X].r = DataScanline[X];
-			bufScanline[X].i = 0;
+			BufScanline[X].r = DataScanline[X];
+			BufScanline[X].i = 0;
 		}
 	}
 
@@ -203,7 +203,7 @@ inline void LowPassFilter(int32 X1, int32 Y1, int32 X2, int32 Y2, FLandscapeBrus
 	int32 CenterPos[2] = { Dims[0] >> 1, Dims[1] >> 1 };
 	for (int32 Y = 0; Y < Dims[0]; Y++)
 	{
-		float DistFromCenter = 0.0f;
+		float DistFromCenter;
 		for (int32 X = 0; X < Dims[1]; X++)
 		{
 			if (Y < CenterPos[0])
@@ -211,12 +211,12 @@ inline void LowPassFilter(int32 X1, int32 Y1, int32 X2, int32 Y2, FLandscapeBrus
 				if (X < CenterPos[1])
 				{
 					// 1
-					DistFromCenter = X*X + Y*Y;
+					DistFromCenter = static_cast<float>(X*X + Y*Y);
 				}
 				else
 				{
 					// 2
-					DistFromCenter = (X - Dims[1])*(X - Dims[1]) + Y*Y;
+					DistFromCenter = static_cast<float>((X - Dims[1])*(X - Dims[1]) + Y*Y);
 				}
 			}
 			else
@@ -224,18 +224,18 @@ inline void LowPassFilter(int32 X1, int32 Y1, int32 X2, int32 Y2, FLandscapeBrus
 				if (X < CenterPos[1])
 				{
 					// 3
-					DistFromCenter = X*X + (Y - Dims[0])*(Y - Dims[0]);
+					DistFromCenter = static_cast<float>(X*X + (Y - Dims[0])*(Y - Dims[0]));
 				}
 				else
 				{
 					// 4
-					DistFromCenter = (X - Dims[1])*(X - Dims[1]) + (Y - Dims[0])*(Y - Dims[0]);
+					DistFromCenter = static_cast<float>((X - Dims[1])*(X - Dims[1]) + (Y - Dims[0])*(Y - Dims[0]));
 				}
 			}
 			// High frequency removal
 			float Ratio = 1.0f - DetailScale;
 			float Dist = FMath::Min<float>((Dims[0] * Ratio)*(Dims[0] * Ratio), (Dims[1] * Ratio)*(Dims[1] * Ratio));
-			float Filter = 1.0 / (1.0 + DistFromCenter / Dist);
+			float Filter = 1.0f / (1.0f + DistFromCenter / Dist);
 			CA_SUPPRESS(6385);
 			out[X + Y*Dims[1]].r *= Filter;
 			out[X + Y*Dims[1]].i *= Filter;
@@ -245,7 +245,7 @@ inline void LowPassFilter(int32 X1, int32 Y1, int32 X2, int32 Y2, FLandscapeBrus
 	// Inverse FFT
 	kiss_fftnd(sti, out, buf);
 
-	float Scale = Dims[0] * Dims[1];
+	const float Scale = static_cast<float>(Dims[0] * Dims[1]);
 	const int32 BrushX1 = FMath::Max<int32>(BrushInfo.GetBounds().Min.X, X1 + 1);
 	const int32 BrushY1 = FMath::Max<int32>(BrushInfo.GetBounds().Min.Y, Y1 + 1);
 	const int32 BrushX2 = FMath::Min<int32>(BrushInfo.GetBounds().Max.X, X2);
@@ -253,8 +253,8 @@ inline void LowPassFilter(int32 X1, int32 Y1, int32 X2, int32 Y2, FLandscapeBrus
 	for (int32 Y = BrushY1; Y < BrushY2; Y++)
 	{
 		const float* BrushScanline = BrushInfo.GetDataPtr(FIntPoint(0, Y));
-		auto* DataScanline = Data.GetData() + (Y - Y1) * (X2 - X1 + 1) + (0 - X1);
-		auto* bufScanline = buf + (Y - (Y1 + 1)) * Dims[1] + (0 - (X1 + 1));
+		typename TArray<DataType>::ElementType* DataScanline = Data.GetData() + (Y - Y1) * (X2 - X1 + 1) + (0 - X1);
+		const kiss_fft_cpx* BufScanline = buf + (Y - (Y1 + 1)) * Dims[1] + (0 - (X1 + 1));
 
 		for (int32 X = BrushX1; X < BrushX2; X++)
 		{
@@ -262,7 +262,7 @@ inline void LowPassFilter(int32 X1, int32 Y1, int32 X2, int32 Y2, FLandscapeBrus
 
 			if (BrushValue > 0.0f)
 			{
-				DataScanline[X] = FMath::Lerp((float)DataScanline[X], bufScanline[X].r / Scale, BrushValue * ApplyRatio);
+				DataScanline[X] = static_cast<DataType>(FMath::Lerp(static_cast<float>(DataScanline[X]), BufScanline[X].r / Scale, BrushValue * ApplyRatio));
 			}
 		}
 	}
@@ -701,7 +701,7 @@ struct FHeightmapAccessorTool : public FHeightmapAccessor<bInUseInterp>
 
 struct FLandscapeHeightCache : public TLandscapeEditCache<FHeightmapAccessorTool<true>, uint16>
 {
-	static uint16 ClampValue(int32 Value) { return FMath::Clamp(Value, 0, LandscapeDataAccess::MaxValue); }
+	static uint16 ClampValue(int32 Value) { return static_cast<uint16>(FMath::Clamp(Value, 0, LandscapeDataAccess::MaxValue)); }
 
 	FLandscapeHeightCache(const FLandscapeToolTarget& InTarget)
 		: TLandscapeEditCache<FHeightmapAccessorTool<true>, uint16>(InTarget)
@@ -741,7 +741,7 @@ struct FXYOffsetmapAccessor
 				FVector* Value = Data.Find(FIntPoint(X, Y));
 				if (Value)
 				{
-					Value->Z = LandscapeDataAccess::GetLocalHeight(static_cast<float>(NewHeights.FindRef(FIntPoint(X, Y))));
+					Value->Z = LandscapeDataAccess::GetLocalHeight(static_cast<uint16>(NewHeights.FindRef(FIntPoint(X, Y))));
 				}
 			}
 		}
@@ -760,7 +760,7 @@ struct FXYOffsetmapAccessor
 				FVector* Value = Data.Find(FIntPoint(X, Y));
 				if (Value)
 				{
-					Value->Z = LandscapeDataAccess::GetLocalHeight(static_cast<float>(NewHeights.FindRef(FIntPoint(X, Y))));
+					Value->Z = LandscapeDataAccess::GetLocalHeight(static_cast<uint16>(NewHeights.FindRef(FIntPoint(X, Y))));
 				}
 			}
 		}
@@ -781,7 +781,7 @@ struct FXYOffsetmapAccessor
 			{
 				for (int32 X = X1; X <= X2; ++X)
 				{
-					NewHeights[X - X1 + (Y - Y1) * (X2 - X1 + 1)] = LandscapeDataAccess::GetTexHeight(Data[(X - X1 + (Y - Y1) * (X2 - X1 + 1))].Z);
+					NewHeights[X - X1 + (Y - Y1) * (X2 - X1 + 1)] = LandscapeDataAccess::GetTexHeight(static_cast<float>(Data[(X - X1 + (Y - Y1) * (X2 - X1 + 1))].Z));
 				}
 			}
 						
@@ -887,7 +887,7 @@ struct FAlphamapAccessorTool : public FAlphamapAccessor<bInUseInterp, bInUseTota
 
 struct FLandscapeAlphaCache : public TLandscapeEditCache<FAlphamapAccessorTool<true, false>, uint8>
 {
-	static uint8 ClampValue(int32 Value) { return FMath::Clamp(Value, 0, 255); }
+	static uint8 ClampValue(int32 Value) { return static_cast<uint8>(FMath::Clamp(Value, 0, 255)); }
 
 	FLandscapeAlphaCache(const FLandscapeToolTarget& InTarget)
 		: TLandscapeEditCache<FAlphamapAccessorTool<true, false>, uint8>(InTarget)
@@ -905,7 +905,7 @@ struct FVisibilityAccessor : public FAlphamapAccessorTool<false, false>
 
 struct FLandscapeVisCache : public TLandscapeEditCache<FAlphamapAccessorTool<false, false>, uint8>
 {
-	static uint8 ClampValue(int32 Value) { return FMath::Clamp(Value, 0, 255); }
+	static uint8 ClampValue(int32 Value) { return static_cast<uint8>(FMath::Clamp(Value, 0, 255)); }
 
 	FLandscapeVisCache(const FLandscapeToolTarget& InTarget)
 		: TLandscapeEditCache<FAlphamapAccessorTool<false, false>, uint8>(InTarget)
@@ -949,7 +949,7 @@ public:
 			bCombinedLayerOperation = Landscape && Landscape->HasLayersContent() && InCombinedLayerOperation && bTargetIsHeightmap;
 			if (bCombinedLayerOperation)
 			{
-				for (int i = 0; i < Landscape->GetLayerCount(); ++i)
+				for (uint8 i = 0; i < Landscape->GetLayerCount(); ++i)
 				{
 					FLandscapeLayer* CurrentLayer = Landscape->GetLayer(i);
 					BackupLayerVisibility.Add(CurrentLayer->bVisible);
@@ -1038,7 +1038,7 @@ public:
 			for (int i = 0; i < Data.Num(); ++i)
 			{
 				float Contribution = (LandscapeDataAccess::GetLocalHeight(Data[i]) - LandscapeDataAccess::GetLocalHeight(BottomLayersData[i])) * InverseAlpha;
-				DataContribution[i] = LandscapeDataAccess::GetTexHeight(Contribution);
+				DataContribution[i] = static_cast<typename ToolTarget::CacheClass::DataType>(LandscapeDataAccess::GetTexHeight(Contribution));
 			}
 
 			FGuid CacheAccessorLayerGuid = CacheUpToEditingLayer.DataAccess.GetEditLayer();
@@ -1332,7 +1332,7 @@ private:
 
 struct FLandscapeDataCache : public TLandscapeEditCache<FDatamapAccessor<false>, uint8>
 {
-	static uint8 ClampValue(int32 Value) { return FMath::Clamp(Value, 0, 255); }
+	static uint8 ClampValue(int32 Value) { return static_cast<uint8>(FMath::Clamp(Value, 0, 255)); }
 
 	FLandscapeDataCache(const FLandscapeToolTarget& InTarget)
 		: TLandscapeEditCache<FDatamapAccessor<false>, uint8>(InTarget)
@@ -1354,7 +1354,7 @@ struct FHeightmapToolTarget
 		if (LandscapeInfo)
 		{
 			// Adjust strength based on brush size and drawscale, so strength 1 = one hemisphere
-			return BrushRadius * LANDSCAPE_INV_ZSCALE / (LandscapeInfo->DrawScale.Z);
+			return static_cast<float>(BrushRadius * LANDSCAPE_INV_ZSCALE / LandscapeInfo->DrawScale.Z);
 		}
 		return 5.0f * LANDSCAPE_INV_ZSCALE;
 	}
@@ -1499,7 +1499,7 @@ public:
 		{
 			ToolStroke.Emplace( EdMode, ViewportClient, InTarget );				// construct the tool stroke class
 			ToolStroke->SetEditLayer(this->EdMode->GetCurrentLayerGuid());		// set the edit layer explicitly (if the tool supports this path)
-			EdMode->CurrentBrush->BeginStroke( InHitLocation.X, InHitLocation.Y, this );
+			EdMode->CurrentBrush->BeginStroke(static_cast<float>(InHitLocation.X), static_cast<float>(InHitLocation.Y), this);
 		}
 
 		// Save the mouse position  
@@ -1580,7 +1580,7 @@ public:
 				if (EdMode->CurrentBrush && !EdMode->IsAdjustingBrush(ViewportClient))
 				{
 					// Inform the brush of the current location, to update the cursor
-					EdMode->CurrentBrush->MouseMove(HitLocation.X, HitLocation.Y);
+					EdMode->CurrentBrush->MouseMove(static_cast<float>(HitLocation.X), static_cast<float>(HitLocation.Y));
 				}
 
 				if (IsToolActive())
diff --git a/Engine/Source/Editor/LandscapeEditor/Private/LandscapeEdModeXYOffset.cpp b/Engine/Source/Editor/LandscapeEditor/Private/LandscapeEdModeXYOffset.cpp
--- a/Engine/Source/Editor/LandscapeEditor/Private/LandscapeEdModeXYOffset.cpp
+++ b/Engine/Source/Editor/LandscapeEditor/Private/LandscapeEdModeXYOffset.cpp
@@ -32,12 +32,12 @@ namespace
 		float Height[4];
 		for (int32 Idx = 0; Idx < 4; ++Idx)
 		{
-			int32 XX = FMath::Clamp(FMath::FloorToInt(X + XYOffset.X + XOffsets[Idx]), MinX, MaxX);
-			int32 YY = FMath::Clamp(FMath::FloorToInt(Y + XYOffset.Y + YOffsets[Idx]), MinY, MaxY);
-			Height[Idx] = XYOffsetVectorData[XX - MinX + (YY - MinY) * (MaxX - MinX + 1)].Z;
+			int32 XX = FMath::Clamp(FMath::FloorToInt32(X + XYOffset.X + XOffsets[Idx]), MinX, MaxX);
+			int32 YY = FMath::Clamp(FMath::FloorToInt32(Y + XYOffset.Y + YOffsets[Idx]), MinY, MaxY);
+			Height[Idx] = static_cast<float>(XYOffsetVectorData[XX - MinX + (YY - MinY) * (MaxX - MinX + 1)].Z);
 		}
-		float FracX = FMath::Fractional(X + XYOffset.X);
-		float FracY = FMath::Fractional(Y + XYOffset.Y);
+		float FracX = static_cast<float>(FMath::Fractional(X + XYOffset.X));
+		float FracY = static_cast<float>(FMath::Fractional(Y + XYOffset.Y));
 		return FMath::Lerp(FMath::Lerp(Height[0], Height[1], FracX),
 			FMath::Lerp(Height[2], Height[3], FracX),
 			FracY);
@@ -207,7 +207,7 @@ public:
 											));
 									}
 
-									TotalArea += (((P[3] - P[0]) ^ (P[1] - P[0])).Size() + ((P[3] - P[0]) ^ (P[2] - P[0])).Size()) * 0.5f;
+									TotalArea += static_cast<float>((((P[3] - P[0]) ^ (P[1] - P[0])).Size() + ((P[3] - P[0]) ^ (P[2] - P[0])).Size()) * 0.5f);
 									QuadNum++;
 									QuadX[ComponentIndexX*ComponentSizeQuads + X - X1]++;
 									QuadY[ComponentIndexY*ComponentSizeQuads + Y - Y1]++;
@@ -226,7 +226,7 @@ public:
 				}
 			}
 
-			const float HeightErrorThreshold = DrawScale3D.X * 0.5f;
+			const float HeightErrorThreshold = static_cast<float>(DrawScale3D.X * 0.5f);
 			// Made XYOffset and new Z value allocation...
 			float AreaErrorThreshold = FMath::Square(AreaResolution);
 			float RemainArea = TotalArea;
@@ -245,7 +245,7 @@ public:
 				float AreaBaseError = AverageArea * 0.5f;
 				float TotalLineArea = 0.0f;
 				float TargetLineArea = AverageArea * QuadY[Y - Y1];
-				float YOffset = Y + 1, PreYOffset = Y + 1; // Need to be bigger than previous Y
+				float YOffset = static_cast<float>(Y + 1), PreYOffset = static_cast<float>(Y + 1); // Need to be bigger than previous Y
 				float StepSize = FPlatformMath::Sqrt(2.f) * 0.25f;
 				float LineAreaDiff = FLT_MAX; //Abs(TotalLineArea - TargetLineArea);
 				int32 IterNum = 0;
@@ -256,7 +256,7 @@ public:
 					YOffset = YOffset + 1.0f;
 					if (YOffset >= Y2)
 					{
-						YOffset = Y2;
+						YOffset = static_cast<float>(Y2);
 						break;
 					}
 				}
@@ -294,12 +294,12 @@ public:
 								if (Idx == 2)
 								{
 									FVector P0(GetWorldPos(LocalToWorld, FVector2D(XX - Comp->GetSectionBase().X, YY0 - Comp->GetSectionBase().Y), XYOffsetVectorData[XX - X1 + (YY0 - Y1) * (X2 - X1 + 1)]));
-									TotalHeightError += FMath::Abs(((P[2] - P0) ^ (P2 - P[2])).Size() - ((P1 - P0) ^ (P2 - P1)).Size());
+									TotalHeightError += static_cast<float>(FMath::Abs(((P[2] - P0) ^ (P2 - P[2])).Size() - ((P1 - P0) ^ (P2 - P1)).Size()));
 								}
 							}
 
 							//TotalHeightError += Abs( XYOffsetVectorData( X - X1 + (YY2 - Y1) * (X2-X1+1) ).Z - XYOffsetVectorData( X - X1 + (YY1 - Y1) * (X2-X1+1) ).Z * appFractional(YOffset) );  
-							TotalLineArea += (((P[3] - P[0]) ^ (P[1] - P[0])).Size() + ((P[3] - P[0]) ^ (P[2] - P[0])).Size()) * 0.5f;
+							TotalLineArea += static_cast<float>((((P[3] - P[0]) ^ (P[1] - P[0])).Size() + ((P[3] - P[0]) ^ (P[2] - P[0])).Size()) * 0.5f);
 						}
 					}
 
@@ -335,12 +335,12 @@ public:
 						// clamp
 						if (YOffset < Y1)
 						{
-							YOffset = Y1;
+							YOffset = static_cast<float>(Y1);
 							break;
 						}
 						if (YOffset >= Y2)
 						{
-							YOffset = Y2;
+							YOffset = static_cast<float>(Y2);
 							break;
 						}
 					}
@@ -392,7 +392,7 @@ public:
 				float AreaBaseError = AverageArea * 0.5f;
 				float TotalLineArea = 0.0f;
 				float TargetLineArea = AverageArea * QuadX[X - X1];
-				float XOffset = X + 1, PreXOffset = X + 1; // Need to be bigger than previous Y
+				float XOffset = static_cast<float>(X + 1), PreXOffset = static_cast<float>(X + 1); // Need to be bigger than previous Y
 				float StepSize = FMath::Sqrt(2.f) * 0.25f;
 				float LineAreaDiff = FLT_MAX; // Abs(TotalLineArea - TargetLineArea);
 				int32 IterNum = 0;
@@ -403,7 +403,7 @@ public:
 					XOffset = XOffset + 1.0f;
 					if (XOffset >= X2)
 					{
-						XOffset = X2;
+						XOffset = static_cast<float>(X2);
 						break;
 					}
 				}
@@ -440,12 +440,12 @@ public:
 								if (Idx == 1)
 								{
 									FVector P0(GetWorldPos(LocalToWorld, FVector2D(XX0 - Comp->GetSectionBase().X, YY - Comp->GetSectionBase().Y), NewYOffsets[XX0 - X1 + (YY - Y1) * (X2 - X1 + 1)]));
-									TotalHeightError += FMath::Abs(((P[1] - P0) ^ (P2 - P[1])).Size() - ((P1 - P0) ^ (P2 - P1)).Size());
+									TotalHeightError += static_cast<float>(FMath::Abs(((P[1] - P0) ^ (P2 - P[1])).Size() - ((P1 - P0) ^ (P2 - P1)).Size()));
 								}
 							}
 
 							//TotalHeightError += Abs( NewYOffsets( XX1 - X1 + (Y - Y1) * (X2-X1+1) ).Z - NewYOffsets( XX2 - X1 + (Y - Y1) * (X2-X1+1) ).Z * appFractional(XOffset) );  
-							TotalLineArea += (((P[3] - P[0]) ^ (P[1] - P[0])).Size() + ((P[3] - P[0]) ^ (P[2] - P[0])).Size()) * 0.5f;
+							TotalLineArea += static_cast<float>((((P[3] - P[0]) ^ (P[1] - P[0])).Size() + ((P[3] - P[0]) ^ (P[2] - P[0])).Size()) * 0.5f);
 						}
 					}
 
@@ -481,12 +481,12 @@ public:
 						// clamp
 						if (XOffset <= X1)
 						{
-							XOffset = X1;
+							XOffset = static_cast<float>(X1);
 							break;
 						}
 						else if (XOffset >= X2)
 						{
-							XOffset = X2;
+							XOffset = static_cast<float>(X2);
 							break;
 						}
 					}
@@ -522,8 +522,8 @@ public:
 
 
 		// Same as Gizmo fall off...
-		float W = X2 - X1 + 1;
-		float H = Y2 - Y1 + 1;
+		float W = static_cast<float>(X2 - X1 + 1);
+		float H = static_cast<float>(Y2 - Y1 + 1);
 		float FalloffRadius = W * 0.5f * UISettings->GetCurrentToolBrushFalloff();
 		float SquareRadius = W * 0.5f - FalloffRadius;
 		for (int32 Y = 0; Y <= Y2 - Y1; ++Y)
@@ -532,10 +532,10 @@ public:
 			{
 				int32 Index = X + Y * (X2 - X1 + 1);
 				FVector2D TransformedLocal(FMath::Abs(X - W * 0.5f), FMath::Abs(Y - H * 0.5f) * (W / H));
-				float Cos = FMath::Abs(TransformedLocal.X) / TransformedLocal.Size();
-				float Sin = FMath::Abs(TransformedLocal.Y) / TransformedLocal.Size();
-				float RatioX = FalloffRadius > 0.0f ? 1.0f - FMath::Clamp((FMath::Abs(TransformedLocal.X) - Cos*SquareRadius) / FalloffRadius, 0.0f, 1.0f) : 1.0f;
-				float RatioY = FalloffRadius > 0.0f ? 1.0f - FMath::Clamp((FMath::Abs(TransformedLocal.Y) - Sin*SquareRadius) / FalloffRadius, 0.0f, 1.0f) : 1.0f;
+				float Cos = static_cast<float>(FMath::Abs(TransformedLocal.X) / TransformedLocal.Size());
+				float Sin = static_cast<float>(FMath::Abs(TransformedLocal.Y) / TransformedLocal.Size());
+				float RatioX = FalloffRadius > 0.0f ? 1.0f - FMath::Clamp(static_cast<float>((FMath::Abs(TransformedLocal.X) - Cos*SquareRadius) / FalloffRadius), 0.0f, 1.0f) : 1.0f;
+				float RatioY = FalloffRadius > 0.0f ? 1.0f - FMath::Clamp(static_cast<float>((FMath::Abs(TransformedLocal.Y) - Sin*SquareRadius) / FalloffRadius), 0.0f, 1.0f) : 1.0f;
 				float Ratio = TransformedLocal.SizeSquared() > FMath::Square(SquareRadius) ? RatioX * RatioY : 1.0f; //TransformedLocal.X / LW * TransformedLocal.Y / LW;
 				float PaintAmount = Ratio*Ratio*(3 - 2 * Ratio);
 
diff --git a/Engine/Source/Editor/LandscapeEditor/Private/LandscapeEditorDetailCustomization_AlphaBrush.cpp b/Engine/Source/Editor/LandscapeEditor/Private/LandscapeEditorDetailCustomization_AlphaBrush.cpp
--- a/Engine/Source/Editor/LandscapeEditor/Private/LandscapeEditorDetailCustomization_AlphaBrush.cpp
+++ b/Engine/Source/Editor/LandscapeEditor/Private/LandscapeEditorDetailCustomization_AlphaBrush.cpp
@@ -78,8 +78,8 @@ public:
 			[
 				SNew(SBox)
 				.ToolTipText(this, &STextureMaskThumbnail::OnGetToolTip)
-				.WidthOverride(64)
-				.HeightOverride(64)
+				.WidthOverride(64.0f)
+				.HeightOverride(64.0f)
 				[
 					SNew(SBorder)
 					.Padding(2)
@@ -133,7 +133,7 @@ public:
 				|| NewTextureChannel != CachedTextureChannel)
 			{
 				CachedTexture = NewTexture;
-				CachedTextureChannel = NewTextureChannel;
+				CachedTextureChannel = static_cast<uint8>(NewTextureChannel);
 				Viewport->Invalidate();
 			}
 		}
diff --git a/Engine/Source/Editor/LandscapeEditor/Private/LandscapeEditorDetailCustomization_CopyPaste.cpp b/Engine/Source/Editor/LandscapeEditor/Private/LandscapeEditorDetailCustomization_CopyPaste.cpp
--- a/Engine/Source/Editor/LandscapeEditor/Private/LandscapeEditorDetailCustomization_CopyPaste.cpp
+++ b/Engine/Source/Editor/LandscapeEditor/Private/LandscapeEditorDetailCustomization_CopyPaste.cpp
@@ -413,7 +413,7 @@ void FLandscapeEditorDetailCustomization_CopyPaste::GenerateGuessDimensionList()
 				TArray<uint32> CurrentValues;
 				for (int32 i = 0; i <= It.Value(); ++i)
 				{
-					CurrentValues.Add(FMath::Pow(static_cast<float>(It.Key()), i));
+					CurrentValues.Add(static_cast<uint32>(FMath::Pow(static_cast<float>(It.Key()), i)));
 				}				
 
 				for (int32 i = 0; i < PreviousIterValues.Num(); ++i)
@@ -438,7 +438,7 @@ void FLandscapeEditorDetailCustomization_CopyPaste::GenerateGuessDimensionList()
 
 			for (int32 i = 0; i < FinalValuesLeft.Num(); ++i)
 			{
-				FinalValuesRight.Add(InitialGizmoDimension / FinalValuesLeft[i]);
+				FinalValuesRight.Add(static_cast<uint32>(InitialGizmoDimension / FinalValuesLeft[i]));
 			}
 
 			GuessedDimensionComboList.Empty(FinalValuesLeft.Num());
diff --git a/Engine/Source/Editor/LandscapeEditor/Private/LandscapeEditorDetailCustomization_ImportExport.cpp b/Engine/Source/Editor/LandscapeEditor/Private/LandscapeEditorDetailCustomization_ImportExport.cpp
--- a/Engine/Source/Editor/LandscapeEditor/Private/LandscapeEditorDetailCustomization_ImportExport.cpp
+++ b/Engine/Source/Editor/LandscapeEditor/Private/LandscapeEditorDetailCustomization_ImportExport.cpp
@@ -245,7 +245,7 @@ void FLandscapeEditorDetailCustomization_ImportExport::CustomizeDetails(IDetailL
 			[
 				SNew(SComboButton)
 				.OnGetMenuContent_Static(&FLandscapeEditorDetailCustomization_ImportExport::GetImportLandscapeResolutionMenu)
-				.ContentPadding(2)
+				.ContentPadding(2.0f)
 				.ButtonContent()
 				[
 					SNew(STextBlock)
@@ -561,7 +561,7 @@ FReply FLandscapeEditorDetailCustomization_ImportExport::OnImportExportButtonCli
 			ELandscapeImportTransformType TransformType = LandscapeEdMode->UISettings->ImportType;
 									
 			FVector LocalGizmoPosition = LandscapeInfo->GetLandscapeProxy()->LandscapeActorToWorld().InverseTransformPosition(LandscapeEdMode->CurrentGizmoActor->GetActorLocation());
-			FIntPoint LocalGizmoPoint = FIntPoint(FMath::FloorToInt(LocalGizmoPosition.X), FMath::FloorToInt(LocalGizmoPosition.Y));
+			FIntPoint LocalGizmoPoint = FIntPoint(FMath::FloorToInt32(LocalGizmoPosition.X), FMath::FloorToInt32(LocalGizmoPosition.Y));
 			
 			// Update Gizmo Position if we exit and comeback into tool
 			LandscapeEdMode->UISettings->ImportLandscape_GizmoLocalPosition = LocalGizmoPoint;
@@ -595,7 +595,7 @@ FReply FLandscapeEditorDetailCustomization_ImportExport::OnImportExportButtonCli
 
 				int32 NumRegions = LandscapeRegions.Num();
 					
-				FScopedSlowTask Progress(NumRegions, LOCTEXT("ImportingLandscapeRegions", "Importing Landscape Regions"));
+				FScopedSlowTask Progress(static_cast<float>(NumRegions), LOCTEXT("ImportingLandscapeRegions", "Importing Landscape Regions"));
 				Progress.MakeDialog(/*bShowCancelButton = */ true);
 
 				auto RegionImporter = [&Progress, LandscapeInfo, CurrentLayerGuid, LandscapeEdMode, ImportRegion, TransformType, ImportOffset, PaintRestriction](const FBox& RegionBounds, const TArray<ALandscapeProxy*>& Proxies)
@@ -633,7 +633,7 @@ FReply FLandscapeEditorDetailCustomization_ImportExport::OnImportExportButtonCli
 			}
 			else
 			{
-				FScopedSlowTask Progress(1 + LandscapeEdMode->UISettings->ImportLandscape_Layers.Num(), LOCTEXT("ImportingLandscape", "Importing Landscape"));
+				FScopedSlowTask Progress(static_cast<float>(1 + LandscapeEdMode->UISettings->ImportLandscape_Layers.Num()), LOCTEXT("ImportingLandscape", "Importing Landscape"));
 				Progress.MakeDialog(/*bShowCancelButton = */ true);
 
 				if (LandscapeEdMode->UISettings->bHeightmapSelected)
diff --git a/Engine/Source/Editor/LandscapeEditor/Private/LandscapeEditorDetailCustomization_Layers.cpp b/Engine/Source/Editor/LandscapeEditor/Private/LandscapeEditorDetailCustomization_Layers.cpp
--- a/Engine/Source/Editor/LandscapeEditor/Private/LandscapeEditorDetailCustomization_Layers.cpp
+++ b/Engine/Source/Editor/LandscapeEditor/Private/LandscapeEditorDetailCustomization_Layers.cpp
@@ -175,7 +175,7 @@ TSharedPtr<SWidget> FLandscapeEditorCustomNodeBuilder_Layers::GenerateRow(int32
 {
 	FEdModeLandscape* LandscapeEdMode = GetEditorMode();
 	TSharedPtr<SWidget> RowWidget = SNew(SLandscapeEditorSelectableBorder)
-		.Padding(0)
+		.Padding(0.0f)
 		.VAlign(VAlign_Center)
 		.OnContextMenuOpening(this, &FLandscapeEditorCustomNodeBuilder_Layers::OnLayerContextMenuOpening, InLayerIndex)
 		.OnSelected(this, &FLandscapeEditorCustomNodeBuilder_Layers::OnLayerSelectionChanged, InLayerIndex)
@@ -203,7 +203,7 @@ TSharedPtr<SWidget> FLandscapeEditorCustomNodeBuilder_Layers::GenerateRow(int32
 			.VAlign(VAlign_Center)
 			[
 				SNew(SButton)
-				.ContentPadding(0)
+				.ContentPadding(0.0f)
 				.ButtonStyle(FAppStyle::Get(), "NoBorder")
 				.OnClicked(this, &FLandscapeEditorCustomNodeBuilder_Layers::OnToggleVisibility, InLayerIndex)
 				.ToolTipText(LOCTEXT("LandscapeLayerVisibility", "Toggle Layer Visibility"))
diff --git a/Engine/Source/Editor/LandscapeEditor/Private/LandscapeEditorDetailCustomization_LayersBrushStack.cpp b/Engine/Source/Editor/LandscapeEditor/Private/LandscapeEditorDetailCustomization_LayersBrushStack.cpp
--- a/Engine/Source/Editor/LandscapeEditor/Private/LandscapeEditorDetailCustomization_LayersBrushStack.cpp
+++ b/Engine/Source/Editor/LandscapeEditor/Private/LandscapeEditorDetailCustomization_LayersBrushStack.cpp
@@ -125,7 +125,7 @@ TSharedPtr<SWidget> FLandscapeEditorCustomNodeBuilder_LayersBrushStack::Generate
 {
 	FEdModeLandscape* LandscapeEdMode = GetEditorMode();
 	TSharedPtr<SWidget> RowWidget = SNew(SLandscapeEditorSelectableBorder)
-		.Padding(0)
+		.Padding(0.0f)
 		.VAlign(VAlign_Center)
 		.OnContextMenuOpening(this, &FLandscapeEditorCustomNodeBuilder_LayersBrushStack::OnBrushContextMenuOpening, InBrushIndex)
 		.OnSelected(this, &FLandscapeEditorCustomNodeBuilder_LayersBrushStack::OnBrushSelectionChanged, InBrushIndex)
@@ -137,7 +137,7 @@ TSharedPtr<SWidget> FLandscapeEditorCustomNodeBuilder_LayersBrushStack::Generate
 			.VAlign(VAlign_Center)
 			[
 				SNew(SButton)
-				.ContentPadding(0)
+				.ContentPadding(0.0f)
 				.ButtonStyle(FAppStyle::Get(), "NoBorder")
 				.OnClicked(this, &FLandscapeEditorCustomNodeBuilder_LayersBrushStack::OnToggleVisibility, InBrushIndex)
 				.ToolTipText(LOCTEXT("LandscapeBrushVisibility", "Toggle Brush Visibility"))
@@ -171,7 +171,7 @@ TSharedPtr<SWidget> FLandscapeEditorCustomNodeBuilder_LayersBrushStack::Generate
 			.VAlign(VAlign_Center)
 			[
 				SNew(SButton)
-				.ContentPadding(0)
+				.ContentPadding(0.0f)
 				.ButtonStyle(FAppStyle::Get(), "NoBorder")
 				.OnClicked(this, &FLandscapeEditorCustomNodeBuilder_LayersBrushStack::OnToggleAffectsHeightmap, InBrushIndex)
 				.ToolTipText(LOCTEXT("LandscapeBrushAffectsHeightmap", "Toggle Affects Heightmap"))
@@ -188,7 +188,7 @@ TSharedPtr<SWidget> FLandscapeEditorCustomNodeBuilder_LayersBrushStack::Generate
             .VAlign(VAlign_Center)
             [
 				SNew(SButton)
-                .ContentPadding(0)
+                .ContentPadding(0.0f)
                 .ButtonStyle(FAppStyle::Get(), "NoBorder")
                 .OnClicked(this, &FLandscapeEditorCustomNodeBuilder_LayersBrushStack::OnToggleAffectsWeightmap, InBrushIndex)
                 .ToolTipText(LOCTEXT("LandscapeBrushAffectsWeightmap", "Toggle Affects Weightmap"))
@@ -205,7 +205,7 @@ TSharedPtr<SWidget> FLandscapeEditorCustomNodeBuilder_LayersBrushStack::Generate
 				.VAlign(VAlign_Center)
 				[
 					SNew(SButton)
-					.ContentPadding(0)
+					.ContentPadding(0.0f)
 					.ButtonStyle(FAppStyle::Get(), "NoBorder")
 					.OnClicked(this, &FLandscapeEditorCustomNodeBuilder_LayersBrushStack::OnToggleAffectsVisibilityLayer, InBrushIndex)
 					.ToolTipText(LOCTEXT("LandscapeBrushAffectsVisibilityLayer", "Toggle Affects Visibility Layer"))
diff --git a/Engine/Source/Editor/LandscapeEditor/Private/LandscapeEditorDetailCustomization_NewLandscape.cpp b/Engine/Source/Editor/LandscapeEditor/Private/LandscapeEditorDetailCustomization_NewLandscape.cpp
--- a/Engine/Source/Editor/LandscapeEditor/Private/LandscapeEditorDetailCustomization_NewLandscape.cpp
+++ b/Engine/Source/Editor/LandscapeEditor/Private/LandscapeEditorDetailCustomization_NewLandscape.cpp
@@ -231,7 +231,7 @@ void FLandscapeEditorDetailCustomization_NewLandscape::CustomizeDetails(IDetailL
 		[
 			SNew(SComboButton)
 			.OnGetMenuContent(this, &FLandscapeEditorDetailCustomization_NewLandscape::GetImportLandscapeResolutionMenu)
-			.ContentPadding(2)
+			.ContentPadding(2.0f)
 			.ButtonContent()
 			[
 				SNew(STextBlock)
@@ -368,7 +368,7 @@ void FLandscapeEditorDetailCustomization_NewLandscape::CustomizeDetails(IDetailL
 	[
 		SNew(SComboButton)
 		.OnGetMenuContent_Static(&GetSectionSizeMenu, PropertyHandle_QuadsPerSection)
-		.ContentPadding(2)
+		.ContentPadding(2.0f)
 		.ButtonContent()
 		[
 			SNew(STextBlock)
@@ -388,7 +388,7 @@ void FLandscapeEditorDetailCustomization_NewLandscape::CustomizeDetails(IDetailL
 	[
 		SNew(SComboButton)
 		.OnGetMenuContent_Static(&GetSectionsPerComponentMenu, PropertyHandle_SectionsPerComponent)
-		.ContentPadding(2)
+		.ContentPadding(2.0f)
 		.ButtonContent()
 		[
 			SNew(STextBlock)
@@ -1181,7 +1181,7 @@ FReply FLandscapeEditorDetailCustomization_NewLandscape::OnCreateButtonClicked()
 		int32 NumRegions = FMath::DivideAndRoundUp(TotalLandscapeComponentSize.X, static_cast<int32>(UISettings->WorldPartitionRegionSize)) * FMath::DivideAndRoundUp(TotalLandscapeComponentSize.Y, static_cast<int32>(UISettings->WorldPartitionRegionSize));
 
 
-		FScopedSlowTask Progress(NumRegions, LOCTEXT("CreateLandscapeRegions", "Creating Landscape Editor Regions..."));
+		FScopedSlowTask Progress(static_cast<float>(NumRegions), LOCTEXT("CreateLandscapeRegions", "Creating Landscape Editor Regions..."));
 		Progress.MakeDialog();
 
 		TArray<ALocationVolume*> RegionVolumes;
@@ -1234,7 +1234,7 @@ FReply FLandscapeEditorDetailCustomization_NewLandscape::OnCreateButtonClicked()
 		LandscapeRegionUtils::ForEachComponentByRegion(UISettings->WorldPartitionRegionSize, NewComponents, AddComponentsToRegion);
 
 		// update the zcomponent of the volumes 
-		const float ZScale =  LandscapeBounds.Max.Z - LandscapeBounds.Min.Z;
+		//const float ZScale =  LandscapeBounds.Max.Z - LandscapeBounds.Min.Z;
 		for (ALocationVolume* RegionVolume : RegionVolumes)
 		{
 			const FVector Scale = RegionVolume->GetActorScale();
@@ -1273,15 +1273,15 @@ FReply FLandscapeEditorDetailCustomization_NewLandscape::OnFillWorldButtonClicke
 
 		LandscapeEdMode->UISettings->NewLandscape_Location = FVector::ZeroVector;
 		const int32 QuadsPerComponent = LandscapeEdMode->UISettings->NewLandscape_SectionsPerComponent * LandscapeEdMode->UISettings->NewLandscape_QuadsPerSection;
-		LandscapeEdMode->UISettings->NewLandscape_ComponentCount.X = FMath::CeilToInt(WORLD_MAX / QuadsPerComponent / LandscapeEdMode->UISettings->NewLandscape_Scale.X);
-		LandscapeEdMode->UISettings->NewLandscape_ComponentCount.Y = FMath::CeilToInt(WORLD_MAX / QuadsPerComponent / LandscapeEdMode->UISettings->NewLandscape_Scale.Y);
+		LandscapeEdMode->UISettings->NewLandscape_ComponentCount.X = FMath::CeilToInt32(WORLD_MAX / QuadsPerComponent / LandscapeEdMode->UISettings->NewLandscape_Scale.X);
+		LandscapeEdMode->UISettings->NewLandscape_ComponentCount.Y = FMath::CeilToInt32(WORLD_MAX / QuadsPerComponent / LandscapeEdMode->UISettings->NewLandscape_Scale.Y);
 
 		const ULandscapeSettings* Settings = GetDefault<ULandscapeSettings>();
 		if (Settings->IsLandscapeResolutionRestricted())
 		{
 			auto ClampComponentCount = [Settings, &QuadsPerComponent](int32& ComponentCount)
 			{
-				const float MaxResolution = Settings->GetSideResolutionLimit();
+				const float MaxResolution = static_cast<float>(Settings->GetSideResolutionLimit());
 				ComponentCount = FMath::Clamp(ComponentCount, 1, FMath::Min(32, FMath::FloorToInt((MaxResolution - 1) / QuadsPerComponent)));
 			};
 
diff --git a/Engine/Source/Editor/LandscapeEditor/Private/LandscapeEditorDetailCustomization_ResizeLandscape.cpp b/Engine/Source/Editor/LandscapeEditor/Private/LandscapeEditorDetailCustomization_ResizeLandscape.cpp
--- a/Engine/Source/Editor/LandscapeEditor/Private/LandscapeEditorDetailCustomization_ResizeLandscape.cpp
+++ b/Engine/Source/Editor/LandscapeEditor/Private/LandscapeEditorDetailCustomization_ResizeLandscape.cpp
@@ -97,7 +97,7 @@ void FLandscapeEditorDetailCustomization_ResizeLandscape::CustomizeDetails(IDeta
 		[
 			SNew(SComboButton)
 			.OnGetMenuContent_Static(&GetSectionSizeMenu, PropertyHandle_QuadsPerSection)
-			.ContentPadding(2)
+			.ContentPadding(2.0f)
 			.ButtonContent()
 			[
 				SNew(STextBlock)
@@ -135,7 +135,7 @@ void FLandscapeEditorDetailCustomization_ResizeLandscape::CustomizeDetails(IDeta
 		[
 			SNew(SComboButton)
 			.OnGetMenuContent_Static(&GetSectionsPerComponentMenu, PropertyHandle_SectionsPerComponent)
-			.ContentPadding(2)
+			.ContentPadding(2.0f)
 			.ButtonContent()
 			[
 				SNew(STextBlock)
diff --git a/Engine/Source/Editor/LandscapeEditor/Private/LandscapeEditorDetailCustomization_TargetLayers.cpp b/Engine/Source/Editor/LandscapeEditor/Private/LandscapeEditorDetailCustomization_TargetLayers.cpp
--- a/Engine/Source/Editor/LandscapeEditor/Private/LandscapeEditorDetailCustomization_TargetLayers.cpp
+++ b/Engine/Source/Editor/LandscapeEditor/Private/LandscapeEditorDetailCustomization_TargetLayers.cpp
@@ -508,7 +508,7 @@ TSharedPtr<SWidget> FLandscapeEditorCustomNodeBuilder_TargetLayers::GenerateRow(
 	if (Target->TargetType != ELandscapeToolTargetType::Weightmap)
 	{
 		RowWidget = SNew(SLandscapeEditorSelectableBorder)
-			.Padding(0)
+			.Padding(0.0f)
 			.VAlign(VAlign_Center)
 			.OnContextMenuOpening_Static(&FLandscapeEditorCustomNodeBuilder_TargetLayers::OnTargetLayerContextMenuOpening, Target)
 			.OnSelected_Static(&FLandscapeEditorCustomNodeBuilder_TargetLayers::OnTargetSelectionChanged, Target)
@@ -547,7 +547,7 @@ TSharedPtr<SWidget> FLandscapeEditorCustomNodeBuilder_TargetLayers::GenerateRow(
 		static const FSlateColorBrush SolidWhiteBrush = FSlateColorBrush(FColorList::White);
 
 		RowWidget = SNew(SLandscapeEditorSelectableBorder)
-			.Padding(0)
+			.Padding(0.0f)
 			.VAlign(VAlign_Center)
 			.OnContextMenuOpening_Static(&FLandscapeEditorCustomNodeBuilder_TargetLayers::OnTargetLayerContextMenuOpening, Target)
 			.OnSelected_Static(&FLandscapeEditorCustomNodeBuilder_TargetLayers::OnTargetSelectionChanged, Target)
@@ -575,8 +575,8 @@ TSharedPtr<SWidget> FLandscapeEditorCustomNodeBuilder_TargetLayers::GenerateRow(
 				[
 					SNew(SBox)
 					.Visibility_Static(&FLandscapeEditorCustomNodeBuilder_TargetLayers::GetDebugModeLayerUsageVisibility, Target)
-					.WidthOverride(48)
-					.HeightOverride(48)
+					.WidthOverride(48.0f)
+					.HeightOverride(48.0f)
 					[
 						SNew(SImage)
 						.Image(FCoreStyle::Get().GetBrush("WhiteBrush"))
diff --git a/Engine/Source/Editor/LandscapeEditor/Private/LandscapeEditorDetailWidgets.cpp b/Engine/Source/Editor/LandscapeEditor/Private/LandscapeEditorDetailWidgets.cpp
--- a/Engine/Source/Editor/LandscapeEditor/Private/LandscapeEditorDetailWidgets.cpp
+++ b/Engine/Source/Editor/LandscapeEditor/Private/LandscapeEditorDetailWidgets.cpp
@@ -186,7 +186,7 @@ void SToolSelector::BuildMultiBlockWidget(const ISlateStyle* StyleSet, const FNa
 	const FString MetaTag = FString::Printf(TEXT("LandscapeToolButton.%s"), Label.IsSet() == true ? *Label.Get().ToString() : TEXT("NoLabel"));
 
 	static FTextBlockStyle LabelStyle = FTextBlockStyle(FAppStyle::GetWidgetStyle< FTextBlockStyle >(FAppStyle::Join(StyleName, ".Label"))).SetShadowOffset(FVector2D::UnitVector);
-	static FTextBlockStyle SmallTextStyle = FTextBlockStyle(LabelStyle).SetFontSize(LabelStyle.Font.Size - 1).SetColorAndOpacity(FSlateColor::UseSubduedForeground());
+	static FTextBlockStyle SmallTextStyle = FTextBlockStyle(LabelStyle).SetFontSize(static_cast<uint16>(LabelStyle.Font.Size - 1)).SetColorAndOpacity(FSlateColor::UseSubduedForeground());
 
 	// Create the content for our button
 	TSharedRef< SWidget > ButtonContent =
@@ -233,7 +233,7 @@ void SToolSelector::BuildMultiBlockWidget(const ISlateStyle* StyleSet, const FNa
 	ChildSlot
 	[
 		SNew(SComboButton)
-			.ContentPadding(0)
+			.ContentPadding(0.0f)
 			.ButtonStyle( FAppStyle::Get(), BlockStyle )
 			.ToolTipText(ToolSelectorButtonBlock->ToolTip)
 			.ForegroundColor(FSlateColor::UseForeground())
diff --git a/Engine/Source/Editor/LandscapeEditor/Private/LandscapeEditorObject.cpp b/Engine/Source/Editor/LandscapeEditor/Private/LandscapeEditorObject.cpp
--- a/Engine/Source/Editor/LandscapeEditor/Private/LandscapeEditorObject.cpp
+++ b/Engine/Source/Editor/LandscapeEditor/Private/LandscapeEditorObject.cpp
@@ -829,7 +829,7 @@ void ULandscapeEditorObject::InitializeDefaultHeightData(TArray<uint16>& OutData
 	// Initialize blank heightmap data
 	for (int32 X = 0; X < SizeX; ++X)
 	{
-		StrideData[X] = LandscapeDataAccess::MidValue;
+		StrideData[X] = static_cast<uint16>(LandscapeDataAccess::MidValue);
 	}
 	for (int32 Y = 0; Y < SizeY; ++Y)
 	{
diff --git a/Engine/Source/Editor/LandscapeEditor/Private/LandscapeFileFormatPng.cpp b/Engine/Source/Editor/LandscapeEditor/Private/LandscapeFileFormatPng.cpp
--- a/Engine/Source/Editor/LandscapeEditor/Private/LandscapeFileFormatPng.cpp
+++ b/Engine/Source/Editor/LandscapeEditor/Private/LandscapeFileFormatPng.cpp
@@ -124,7 +124,7 @@ FLandscapeImportData<uint16> FLandscapeHeightmapFileFormat_Png::Import(const TCH
 				else
 				{
 					Result.Data.Empty(ExpectedResolution.Width * ExpectedResolution.Height);
-					Algo::Transform(RawData, Result.Data, [](uint8 Value) { return Value * 0x101; }); // Expand to 16-bit
+					Algo::Transform(RawData, Result.Data, [](uint8 Value) { return static_cast<uint16>(Value * 0x101); }); // Expand to 16-bit
 				}
 			}
 			else
diff --git a/Engine/Source/Editor/LandscapeEditor/Private/LandscapeFileFormatRaw.cpp b/Engine/Source/Editor/LandscapeEditor/Private/LandscapeFileFormatRaw.cpp
--- a/Engine/Source/Editor/LandscapeEditor/Private/LandscapeFileFormatRaw.cpp
+++ b/Engine/Source/Editor/LandscapeEditor/Private/LandscapeFileFormatRaw.cpp
@@ -74,8 +74,8 @@ bool GetRawResolution(const TCHAR* InFilename, FLandscapeFileResolution& OutReso
 	const int32 BytesPerPixel = bIs16Bit ? 2 : 1;
 	int64 FileSize = IFileManager::Get().FileSize(InFilename);
 	
-	const int32 NumPixels = FileSize / BytesPerPixel;
-	int32 Dimension = FMath::TruncToInt(FMath::Sqrt(static_cast<double>(NumPixels)));
+	const int64 NumPixels = FileSize / BytesPerPixel;
+	const uint32 Dimension = FMath::TruncToInt32(FMath::Sqrt(static_cast<double>(NumPixels)));
 	if (Dimension * Dimension == NumPixels)
 	{
 		OutResolution.Width = Dimension;
diff --git a/Engine/Source/Editor/LandscapeEditor/Private/LandscapeImportHelper.cpp b/Engine/Source/Editor/LandscapeEditor/Private/LandscapeImportHelper.cpp
--- a/Engine/Source/Editor/LandscapeEditor/Private/LandscapeImportHelper.cpp
+++ b/Engine/Source/Editor/LandscapeEditor/Private/LandscapeImportHelper.cpp
@@ -337,7 +337,7 @@ ELandscapeImportResult FLandscapeImportHelper::GetWeightmapImportDescriptor(cons
 
 ELandscapeImportResult FLandscapeImportHelper::GetHeightmapImportData(const FLandscapeImportDescriptor& ImportDescriptor, int32 DescriptorIndex, TArray<uint16>& OutData, FText& OutMessage)
 {
-	return GetImportDataInternal<uint16>(ImportDescriptor, DescriptorIndex, NAME_None, LandscapeDataAccess::MidValue, OutData, OutMessage);
+	return GetImportDataInternal<uint16>(ImportDescriptor, DescriptorIndex, NAME_None, static_cast<uint16>(LandscapeDataAccess::MidValue), OutData, OutMessage);
 }
 
 ELandscapeImportResult FLandscapeImportHelper::GetWeightmapImportData(const FLandscapeImportDescriptor& ImportDescriptor, int32 DescriptorIndex, FName LayerName, TArray<uint8>& OutData, FText& OutMessage)
diff --git a/Engine/Source/Editor/LandscapeEditor/Private/LandscapeTiledImage.h b/Engine/Source/Editor/LandscapeEditor/Private/LandscapeTiledImage.h
--- a/Engine/Source/Editor/LandscapeEditor/Private/LandscapeTiledImage.h
+++ b/Engine/Source/Editor/LandscapeEditor/Private/LandscapeTiledImage.h
@@ -19,7 +19,7 @@ template<>
 class LandscapeImageTraits<uint16>
 {
 public:
-	static uint16 DefaultValue() { return LandscapeDataAccess::MidValue; }
+	static uint16 DefaultValue() { return static_cast<uint16>(LandscapeDataAccess::MidValue); }
 };
 
 class FLandscapeTiledImage
diff --git a/Engine/Source/Editor/LandscapeEditor/Private/SLandscapeEditor.cpp b/Engine/Source/Editor/LandscapeEditor/Private/SLandscapeEditor.cpp
--- a/Engine/Source/Editor/LandscapeEditor/Private/SLandscapeEditor.cpp
+++ b/Engine/Source/Editor/LandscapeEditor/Private/SLandscapeEditor.cpp
@@ -31,8 +31,8 @@ void SLandscapeAssetThumbnail::Construct(const FArguments& InArgs, UObject* Asse
 	ChildSlot
 	[
 		SNew(SBox)
-		.WidthOverride(ThumbnailSize.X)
-		.HeightOverride(ThumbnailSize.Y)
+		.WidthOverride(static_cast<float>(ThumbnailSize.X))
+		.HeightOverride(static_cast<float>(ThumbnailSize.Y))
 		[
 			AssetThumbnail->MakeThumbnailWidget()
 		]
