Fix crashes and other weird behavior after undo adding a new LayerInfo to a Landscape Layer.

diff --git a/Engine/Source/Editor/LandscapeEditor/Private/LandscapeEditorDetailCustomization_TargetLayers.cpp b/Engine/Source/Editor/LandscapeEditor/Private/LandscapeEditorDetailCustomization_TargetLayers.cpp
--- a/Engine/Source/Editor/LandscapeEditor/Private/LandscapeEditorDetailCustomization_TargetLayers.cpp
+++ b/Engine/Source/Editor/LandscapeEditor/Private/LandscapeEditorDetailCustomization_TargetLayers.cpp
@@ -1306,7 +1306,10 @@ void FLandscapeEditorCustomNodeBuilder_TargetLayers::OnTargetLayerCreateClicked(
 			LayerObjectName = FName(*NewLayerDlg->GetAssetName().ToString());
 
 			UPackage* Package = CreatePackage( *PackageName);
-			ULandscapeLayerInfoObject* LayerInfo = NewObject<ULandscapeLayerInfoObject>(Package, LayerObjectName, RF_Public | RF_Standalone | RF_Transactional);
+
+			// Do not pass RF_Transactional to NewObject, or the asset will mark itself as garbage on Undo (which is not a well-supported path, potentially causing crashes)
+			ULandscapeLayerInfoObject* LayerInfo = NewObject<ULandscapeLayerInfoObject>(Package, LayerObjectName, RF_Public | RF_Standalone);
+			LayerInfo->SetFlags(RF_Transactional);	// we add RF_Transactional after creation, so that future edits _are_ recorded in undo
 			LayerInfo->LayerName = LayerName;
 			LayerInfo->bNoWeightBlend = bNoWeightBlend;
 
