Fixed visibility weightmaps being uselessly generated when using landscape patches : now, we only generate those when a patch actually writes visibility.
This allows the landscape to use opaque mode rather than masked, which is crucial performance-wise.

diff --git a/Engine/Plugins/Experimental/LandscapePatch/Source/LandscapePatch/Private/LandscapePatchManager.cpp b/Engine/Plugins/Experimental/LandscapePatch/Source/LandscapePatch/Private/LandscapePatchManager.cpp
--- a/Engine/Plugins/Experimental/LandscapePatch/Source/LandscapePatch/Private/LandscapePatchManager.cpp
+++ b/Engine/Plugins/Experimental/LandscapePatch/Source/LandscapePatch/Private/LandscapePatchManager.cpp
@@ -22,9 +22,9 @@ ALandscapePatchManager::ALandscapePatchManager(const FObjectInitializer& ObjectI
 	: ALandscapeBlueprintBrushBase(ObjectInitializer)
 {
 #if WITH_EDITOR
-	SetAffectsHeightmap(true);
-	SetAffectsWeightmap(true);
-	SetAffectsVisibilityLayer(true);
+	SetCanAffectHeightmap(true);
+	SetCanAffectWeightmap(true);
+	SetCanAffectVisibilityLayer(true);
 #endif
 }
 
@@ -248,8 +248,43 @@ void ALandscapePatchManager::MovePatchToIndex(TObjectPtr<ULandscapePatchComponen
 }
 
 #if WITH_EDITOR
+PRAGMA_DISABLE_DEPRECATION_WARNINGS
 bool ALandscapePatchManager::IsAffectingWeightmapLayer(const FName& InLayerName) const
 {
+	return AffectsWeightmapLayer(InLayerName);
+}
+PRAGMA_ENABLE_DEPRECATION_WARNINGS
+
+bool ALandscapePatchManager::AffectsWeightmapLayer(const FName& InLayerName) const
+{
+	if (!CanAffectWeightmap())
+	{
+		return false;
+	}
+	
+	for (const TSoftObjectPtr<ULandscapePatchComponent>& Component : PatchComponents)
+	{
+		if (Component.IsPending())
+		{
+			Component.LoadSynchronous();
+		}
+
+		if (Component.IsValid() && Component->IsEnabled() && Component->AffectsWeightmapLayer(InLayerName))
+		{
+			return true;
+		}
+	}
+
+	return false;
+}
+
+bool ALandscapePatchManager::AffectsVisibilityLayer() const
+{
+	if (!CanAffectVisibilityLayer())
+	{
+		return false;
+	}
+
 	for (const TSoftObjectPtr<ULandscapePatchComponent>& Component : PatchComponents)
 	{
 		if (Component.IsPending())
@@ -257,7 +292,7 @@ bool ALandscapePatchManager::IsAffectingWeightmapLayer(const FName& InLayerName)
 			Component.LoadSynchronous();
 		}
 
-		if (Component.IsValid() && Component->IsEnabled() && Component->IsAffectingWeightmapLayer(InLayerName))
+		if (Component.IsValid() && Component->IsEnabled() && Component->AffectsVisibilityLayer())
 		{
 			return true;
 		}
diff --git a/Engine/Plugins/Experimental/LandscapePatch/Source/LandscapePatch/Private/LandscapeTexturePatch.cpp b/Engine/Plugins/Experimental/LandscapePatch/Source/LandscapePatch/Private/LandscapeTexturePatch.cpp
--- a/Engine/Plugins/Experimental/LandscapePatch/Source/LandscapePatch/Private/LandscapeTexturePatch.cpp
+++ b/Engine/Plugins/Experimental/LandscapePatch/Source/LandscapePatch/Private/LandscapeTexturePatch.cpp
@@ -13,6 +13,7 @@
 #include "RenderGraphUtils.h"
 #include "TextureResource.h"
 #include "RenderingThread.h"
+#include "Algo/AnyOf.h"
 
 #include UE_INLINE_GENERATED_CPP_BY_NAME(LandscapeTexturePatch)
 
@@ -797,7 +798,15 @@ FMatrix44f ULandscapeTexturePatch::GetPatchToHeightmapUVs(int32 PatchSizeX, int3
 	return (FMatrix44f)PatchToLandscapeUVTransposed.GetTransposed();
 }
 
+
+PRAGMA_DISABLE_DEPRECATION_WARNINGS
 bool ULandscapeTexturePatch::IsAffectingWeightmapLayer(const FName& InLayerName) const
+{
+	return AffectsWeightmapLayer(InLayerName);
+}
+PRAGMA_ENABLE_DEPRECATION_WARNINGS
+
+bool ULandscapeTexturePatch::AffectsWeightmapLayer(const FName& InLayerName) const
 {
 	if (!IsEnabled())
 	{
@@ -814,6 +823,16 @@ bool ULandscapeTexturePatch::IsAffectingWeightmapLayer(const FName& InLayerName)
 	return false;
 }
 
+bool ULandscapeTexturePatch::AffectsVisibilityLayer() const
+{
+	if (!IsEnabled())
+	{
+		return false;
+	}
+
+	return Algo::AnyOf(WeightPatches, [](const TObjectPtr<ULandscapeWeightPatchTextureInfo>& InWeightPatch) { return InWeightPatch->bEditVisibilityLayer; });
+}
+
 // We override IsEnabled to make the patch not request updates when all the source modes are "none"
 // (unless we need the update for reinitialization).
 bool ULandscapeTexturePatch::IsEnabled() const
diff --git a/Engine/Plugins/Experimental/LandscapePatch/Source/LandscapePatch/Public/LandscapeCircleHeightPatch.h b/Engine/Plugins/Experimental/LandscapePatch/Source/LandscapePatch/Public/LandscapeCircleHeightPatch.h
--- a/Engine/Plugins/Experimental/LandscapePatch/Source/LandscapePatch/Public/LandscapeCircleHeightPatch.h
+++ b/Engine/Plugins/Experimental/LandscapePatch/Source/LandscapePatch/Public/LandscapeCircleHeightPatch.h
@@ -28,6 +28,8 @@ public:
 
 	virtual UTextureRenderTarget2D* RenderLayer_Native(const FLandscapeBrushParameters& InParameters) override;
 
+	virtual bool AffectsVisibilityLayer() const override { return bEditVisibility; }
+
 	// UActorComponent
 	virtual void OnComponentCreated() override;
 
diff --git a/Engine/Plugins/Experimental/LandscapePatch/Source/LandscapePatch/Public/LandscapePatchComponent.h b/Engine/Plugins/Experimental/LandscapePatch/Source/LandscapePatch/Public/LandscapePatchComponent.h
--- a/Engine/Plugins/Experimental/LandscapePatch/Source/LandscapePatch/Public/LandscapePatchComponent.h
+++ b/Engine/Plugins/Experimental/LandscapePatch/Source/LandscapePatch/Public/LandscapePatchComponent.h
@@ -36,7 +36,10 @@ public:
 
 	virtual UTextureRenderTarget2D* RenderLayer_Native(const FLandscapeBrushParameters& InParameters) { return InParameters.CombinedResult; }
 
+	UE_DEPRECATED(5.3, "Use AffectsWeightmapLayer")
 	virtual bool IsAffectingWeightmapLayer(const FName& InLayerName) const { return false; }
+	virtual bool AffectsWeightmapLayer(const FName& InLayerName) const { return false; }
+	virtual bool AffectsVisibilityLayer() const { return false; }
 
 	UFUNCTION(BlueprintCallable, Category = "LandscapePatch")
 	void RequestLandscapeUpdate(bool bInUserTriggeredUpdate = false);
diff --git a/Engine/Plugins/Experimental/LandscapePatch/Source/LandscapePatch/Public/LandscapePatchManager.h b/Engine/Plugins/Experimental/LandscapePatch/Source/LandscapePatch/Public/LandscapePatchManager.h
--- a/Engine/Plugins/Experimental/LandscapePatch/Source/LandscapePatch/Public/LandscapePatchManager.h
+++ b/Engine/Plugins/Experimental/LandscapePatch/Source/LandscapePatch/Public/LandscapePatchManager.h
@@ -66,7 +66,10 @@ public:
 	
 #if WITH_EDITOR
 	// ALandscapeBlueprintBrushBase
+	UE_DEPRECATED(5.3, "Use AffectsWeightmapLayer")
 	virtual bool IsAffectingWeightmapLayer(const FName& InLayerName) const override;
+	virtual bool AffectsWeightmapLayer(const FName& InLayerName) const override;
+	virtual bool AffectsVisibilityLayer() const override;
 	virtual void SetOwningLandscape(class ALandscape* InOwningLandscape) override;
 
 	// UObject
diff --git a/Engine/Plugins/Experimental/LandscapePatch/Source/LandscapePatch/Public/LandscapeTexturePatch.h b/Engine/Plugins/Experimental/LandscapePatch/Source/LandscapePatch/Public/LandscapeTexturePatch.h
--- a/Engine/Plugins/Experimental/LandscapePatch/Source/LandscapePatch/Public/LandscapeTexturePatch.h
+++ b/Engine/Plugins/Experimental/LandscapePatch/Source/LandscapePatch/Public/LandscapeTexturePatch.h
@@ -223,7 +223,10 @@ public:
 	UTextureRenderTarget2D* RenderLayer_Native(const FLandscapeBrushParameters& InParameters);
 
 	// ULandscapePatchComponent
+	UE_DEPRECATED(5.3, "Use AffectsWeightmapLayer")
 	virtual bool IsAffectingWeightmapLayer(const FName& InLayerName) const override;
+	virtual bool AffectsWeightmapLayer(const FName& InLayerName) const override;
+	virtual bool AffectsVisibilityLayer() const override;
 	virtual bool IsEnabled() const override;
 
 	// UObject
@@ -426,7 +429,6 @@ public:
 	UFUNCTION(BlueprintCallable, Category = LandscapePatch)
 	virtual void SetEditVisibilityLayer(const FName& InWeightmapLayerName, const bool bEditVisibilityLayer);
 
-
 	// These need to be public so that we can take the internal textures and write them to external ones,
 	// but unclear whether we want to expose them to blueprints, since this is a fairly internal thing.
 	UTexture2D* GetHeightInternalTexture() { return HeightInternalData ? HeightInternalData->GetInternalTexture() : nullptr; ; }
diff --git a/Engine/Plugins/Experimental/Water/Source/Editor/Private/WaterLandscapeBrush.cpp b/Engine/Plugins/Experimental/Water/Source/Editor/Private/WaterLandscapeBrush.cpp
--- a/Engine/Plugins/Experimental/Water/Source/Editor/Private/WaterLandscapeBrush.cpp
+++ b/Engine/Plugins/Experimental/Water/Source/Editor/Private/WaterLandscapeBrush.cpp
@@ -28,7 +28,7 @@
 
 AWaterLandscapeBrush::AWaterLandscapeBrush(const FObjectInitializer& ObjectInitializer)
 {
-	SetAffectsHeightmap(true);
+	SetCanAffectHeightmap(true);
 
 	ActorIcon = FWaterIconHelper::EnsureSpriteComponentCreated(this, TEXT("/Water/Icons/WaterLandscapeBrushSprite"));
 }
diff --git a/Engine/Source/Editor/LandscapeEditor/Private/LandscapeEdModeBPCustomTools.cpp b/Engine/Source/Editor/LandscapeEditor/Private/LandscapeEdModeBPCustomTools.cpp
--- a/Engine/Source/Editor/LandscapeEditor/Private/LandscapeEdModeBPCustomTools.cpp
+++ b/Engine/Source/Editor/LandscapeEditor/Private/LandscapeEdModeBPCustomTools.cpp
@@ -88,7 +88,9 @@ public:
 		}
 
 		// Only allow placing brushes that would affect our target type
-		if ((DefaultObject->IsAffectingHeightmap() && Target.TargetType == ELandscapeToolTargetType::Heightmap) || (DefaultObject->IsAffectingWeightmap() && Target.TargetType == ELandscapeToolTargetType::Weightmap))
+		if ((DefaultObject->CanAffectHeightmap() && Target.TargetType == ELandscapeToolTargetType::Heightmap) 
+			|| (DefaultObject->CanAffectWeightmap() && Target.TargetType == ELandscapeToolTargetType::Weightmap)
+			|| (DefaultObject->CanAffectVisibilityLayer() && Target.TargetType == ELandscapeToolTargetType::Visibility))
 		{
 			ULandscapeInfo* Info = EdMode->CurrentToolTarget.LandscapeInfo.Get();
 			check(Info);
diff --git a/Engine/Source/Editor/LandscapeEditor/Private/LandscapeEditorDetailCustomization_LayersBrushStack.cpp b/Engine/Source/Editor/LandscapeEditor/Private/LandscapeEditorDetailCustomization_LayersBrushStack.cpp
--- a/Engine/Source/Editor/LandscapeEditor/Private/LandscapeEditorDetailCustomization_LayersBrushStack.cpp
+++ b/Engine/Source/Editor/LandscapeEditor/Private/LandscapeEditorDetailCustomization_LayersBrushStack.cpp
@@ -313,9 +313,9 @@ FReply FLandscapeEditorCustomNodeBuilder_LayersBrushStack::OnToggleAffectsHeight
 {
 	if (ALandscapeBlueprintBrushBase* Brush = GetBrush(InBrushIndex))
 	{
-		const FScopedTransaction Transaction(LOCTEXT("Landscape_Brush_AffectsHeightmap", "Set Brush Affects Heightmap"));
-		bool bAffectsHeightmap = Brush->IsAffectingHeightmap();
-		Brush->SetAffectsHeightmap(!bAffectsHeightmap);
+		const FScopedTransaction Transaction(LOCTEXT("Landscape_Brush_AffectsHeightmap", "Set Brush Can Affect Heightmap"));
+		bool bCanAffectHeightmap = Brush->CanAffectHeightmap();
+		Brush->SetCanAffectHeightmap(!bCanAffectHeightmap);
 	}
 	return FReply::Handled();
 }
@@ -324,9 +324,9 @@ FReply FLandscapeEditorCustomNodeBuilder_LayersBrushStack::OnToggleAffectsWeight
 {
 	if (ALandscapeBlueprintBrushBase* Brush = GetBrush(InBrushIndex))
 	{
-		const FScopedTransaction Transaction(LOCTEXT("Landscape_Brush_AffectsWeightmap", "Set Brush Affects Weightmap"));
-		bool bAffectsWeightmap = Brush->IsAffectingWeightmap();
-		Brush->SetAffectsWeightmap(!bAffectsWeightmap);
+		const FScopedTransaction Transaction(LOCTEXT("Landscape_Brush_AffectsWeightmap", "Set Brush Can Affect Weightmap"));
+		bool bCanAffectWeightmap = Brush->CanAffectWeightmap();
+		Brush->SetCanAffectWeightmap(!bCanAffectWeightmap);
 	}
 	return FReply::Handled();
 }
@@ -334,9 +334,9 @@ FReply FLandscapeEditorCustomNodeBuilder_LayersBrushStack::OnToggleAffectsVisibi
 {
 	if (ALandscapeBlueprintBrushBase* Brush = GetBrush(InBrushIndex))
 	{
-		const FScopedTransaction Transaction(LOCTEXT("Landscape_Brush_AffectsVisibilityLayer", "Set Brush Affects Visibility Layer"));
-		bool bAffectsVisibilityLayer = Brush->IsAffectingVisibilityLayer();
-		Brush->SetAffectsVisibilityLayer(!bAffectsVisibilityLayer);
+		const FScopedTransaction Transaction(LOCTEXT("Landscape_Brush_AffectsVisibilityLayer", "Set Brush Can Affect Visibility Layer"));
+		bool bCanAffectVisibilityLayer = Brush->CanAffectVisibilityLayer();
+		Brush->SetCanAffectVisibilityLayer(!bCanAffectVisibilityLayer);
 	}
 	return FReply::Handled();
 }
@@ -360,19 +360,19 @@ void FLandscapeEditorCustomNodeBuilder_LayersBrushStack::OnToggleVisibility(ALan
 const FSlateBrush* FLandscapeEditorCustomNodeBuilder_LayersBrushStack::GetAffectsWeightmapBrush(int32 InBrushIndex) const
 {
 	ALandscapeBlueprintBrushBase* Brush = GetBrush(InBrushIndex);
-	return Brush && Brush->IsAffectingWeightmap() ? FAppStyle::GetBrush("LandscapeEditor.Brush.AffectsWeight.Enabled") : FAppStyle::GetBrush("LandscapeEditor.Brush.AffectsWeight.Disabled");
+	return Brush && Brush->CanAffectWeightmap() ? FAppStyle::GetBrush("LandscapeEditor.Brush.AffectsWeight.Enabled") : FAppStyle::GetBrush("LandscapeEditor.Brush.AffectsWeight.Disabled");
 }
 
 const FSlateBrush* FLandscapeEditorCustomNodeBuilder_LayersBrushStack::GetAffectsHeightmapBrush(int32 InBrushIndex) const
 {
 	ALandscapeBlueprintBrushBase* Brush = GetBrush(InBrushIndex);
-	return Brush && Brush->IsAffectingHeightmap() ? FAppStyle::GetBrush("LandscapeEditor.Brush.AffectsHeight.Enabled") : FAppStyle::GetBrush("LandscapeEditor.Brush.AffectsHeight.Disabled");
+	return Brush && Brush->CanAffectHeightmap() ? FAppStyle::GetBrush("LandscapeEditor.Brush.AffectsHeight.Enabled") : FAppStyle::GetBrush("LandscapeEditor.Brush.AffectsHeight.Disabled");
 }
 
 const FSlateBrush* FLandscapeEditorCustomNodeBuilder_LayersBrushStack::GetAffectsVisibilityLayerBrush(int32 InBrushIndex) const
 {
 	ALandscapeBlueprintBrushBase* Brush = GetBrush(InBrushIndex);
-	return Brush && Brush->IsAffectingVisibilityLayer() ? FAppStyle::GetBrush("LandscapeEditor.Brush.AffectsVisibilityLayer.Enabled") : FAppStyle::GetBrush("LandscapeEditor.Brush.AffectsVisibilityLayer.Disabled");
+	return Brush && Brush->CanAffectVisibilityLayer() ? FAppStyle::GetBrush("LandscapeEditor.Brush.AffectsVisibilityLayer.Enabled") : FAppStyle::GetBrush("LandscapeEditor.Brush.AffectsVisibilityLayer.Disabled");
 }
 
 const FSlateBrush* FLandscapeEditorCustomNodeBuilder_LayersBrushStack::GetVisibilityBrush(int32 InBrushIndex) const
@@ -386,7 +386,10 @@ bool FLandscapeEditorCustomNodeBuilder_LayersBrushStack::IsBrushEnabled(int32 In
 {
 	ALandscapeBlueprintBrushBase* Brush = GetBrush(InBrushIndex);
 	const FEdModeLandscape* LandscapeEdMode = GetEditorMode();
-	return Brush && ((Brush->IsAffectingHeightmap() && LandscapeEdMode->CurrentToolTarget.TargetType == ELandscapeToolTargetType::Heightmap) || (Brush->IsAffectingWeightmap() && LandscapeEdMode->CurrentToolTarget.TargetType == ELandscapeToolTargetType::Weightmap) || (Brush->IsAffectingVisibilityLayer() && LandscapeEdMode->CurrentToolTarget.TargetType == ELandscapeToolTargetType::Visibility));
+	return Brush 
+		&& ((Brush->CanAffectHeightmap() && LandscapeEdMode->CurrentToolTarget.TargetType == ELandscapeToolTargetType::Heightmap) 
+		|| (Brush->CanAffectWeightmap() && LandscapeEdMode->CurrentToolTarget.TargetType == ELandscapeToolTargetType::Weightmap) 
+		|| (Brush->CanAffectVisibilityLayer() && LandscapeEdMode->CurrentToolTarget.TargetType == ELandscapeToolTargetType::Visibility));
 }
 
 bool FLandscapeEditorCustomNodeBuilder_LayersBrushStack::IsBrushSelected(int32 InBrushIndex) const
diff --git a/Engine/Source/Runtime/Landscape/Classes/Landscape.h b/Engine/Source/Runtime/Landscape/Classes/Landscape.h
--- a/Engine/Source/Runtime/Landscape/Classes/Landscape.h
+++ b/Engine/Source/Runtime/Landscape/Classes/Landscape.h
@@ -131,9 +131,9 @@ struct FLandscapeLayerBrush
 	UTextureRenderTarget2D* RenderLayer(const FIntRect& InLandscapeSize, const FLandscapeBrushParameters& InParameters);
 
 	LANDSCAPE_API ALandscapeBlueprintBrushBase* GetBrush() const;
-	bool IsAffectingHeightmap() const;
-	bool IsAffectingWeightmapLayer(const FName& InWeightmapLayerName) const;
-	bool IsAffectingVisibilityLayer() const;
+	bool AffectsHeightmap() const;
+	bool AffectsWeightmapLayer(const FName& InWeightmapLayerName) const;
+	bool AffectsVisibilityLayer() const;
 	void SetOwner(ALandscape* InOwner);
 #endif
 
diff --git a/Engine/Source/Runtime/Landscape/Private/LandscapeBlueprintBrushBase.cpp b/Engine/Source/Runtime/Landscape/Private/LandscapeBlueprintBrushBase.cpp
--- a/Engine/Source/Runtime/Landscape/Private/LandscapeBlueprintBrushBase.cpp
+++ b/Engine/Source/Runtime/Landscape/Private/LandscapeBlueprintBrushBase.cpp
@@ -87,11 +87,11 @@ void ALandscapeBlueprintBrushBase::RequestLandscapeUpdate(bool bInUserTriggered)
 	if (OwningLandscape)
 	{
 		uint32 ModeMask = 0;
-		if (AffectHeightmap)
+		if (AffectsHeightmap())
 		{
 			ModeMask |= ELandscapeLayerUpdateMode::Update_Heightmap_Editing_NoCollision;
 		}
-		if (AffectWeightmap || AffectVisibilityLayer)
+		if (AffectsWeightmap() || AffectsVisibilityLayer())
 		{
 			ModeMask |= ELandscapeLayerUpdateMode::Update_Weightmap_Editing_NoCollision;
 		}
@@ -116,11 +116,11 @@ void ALandscapeBlueprintBrushBase::PushDeferredLayersContentUpdate()
 		LastRequestLayersContentUpdateFrameNumber + CVarLandscapeBrushPadding.GetValueOnAnyThread() <= GFrameNumber)
 	{
 		uint32 ModeMask = 0;
-		if (AffectHeightmap)
+		if (AffectsHeightmap())
 		{
 			ModeMask |= ELandscapeLayerUpdateMode::Update_Heightmap_All;
 		}
-		if (AffectWeightmap || AffectVisibilityLayer)
+		if (AffectsWeightmap() || AffectsVisibilityLayer())
 		{
 			ModeMask |= ELandscapeLayerUpdateMode::Update_Weightmap_All;
 		}
@@ -167,49 +167,86 @@ void ALandscapeBlueprintBrushBase::SetIsVisible(bool bInIsVisible)
 #endif
 }
 
-void ALandscapeBlueprintBrushBase::SetAffectsHeightmap(bool bInAffectsHeightmap)
+
+PRAGMA_DISABLE_DEPRECATION_WARNINGS
+void ALandscapeBlueprintBrushBase::SetAffectsHeightmap(bool bAffectsHeightmap)
+{
+	SetCanAffectHeightmap(bAffectsHeightmap);
+}
+
+void ALandscapeBlueprintBrushBase::SetAffectsWeightmap(bool bAffectsWeightmap)
+{
+	SetCanAffectWeightmap(bAffectsWeightmap);
+}
+
+void ALandscapeBlueprintBrushBase::SetAffectsVisibilityLayer(bool bInAffectsVisibilityLayer)
+{
+	SetCanAffectVisibilityLayer(bInAffectsVisibilityLayer);
+}
+
+bool ALandscapeBlueprintBrushBase::IsAffectingWeightmapLayer(const FName& InLayerName) const
+{
+	return AffectsWeightmapLayer(InLayerName);
+}
+PRAGMA_ENABLE_DEPRECATION_WARNINGS
+
+void ALandscapeBlueprintBrushBase::SetCanAffectHeightmap(bool bInCanAffectHeightmap)
 {
 #if WITH_EDITORONLY_DATA
-	Modify();
-	AffectHeightmap = bInAffectsHeightmap;
-	if (OwningLandscape)
+	if (bInCanAffectHeightmap != AffectHeightmap)
 	{
-		OwningLandscape->OnBlueprintBrushChanged();
+		Modify();
+		AffectHeightmap = bInCanAffectHeightmap;
+		if (OwningLandscape)
+		{
+			OwningLandscape->OnBlueprintBrushChanged();
+		}
 	}
-#endif
+#endif // WITH_EDITORONLY_DATA
 }
 
-void ALandscapeBlueprintBrushBase::SetAffectsWeightmap(bool bInAffectsWeightmap)
+void ALandscapeBlueprintBrushBase::SetCanAffectWeightmap(bool bInCanAffectWeightmap)
 {
 #if WITH_EDITORONLY_DATA
-	Modify();
-	AffectWeightmap = bInAffectsWeightmap;
-	if (OwningLandscape)
+	if (bInCanAffectWeightmap != AffectWeightmap)
 	{
-		OwningLandscape->OnBlueprintBrushChanged();
+		Modify();
+		AffectWeightmap = bInCanAffectWeightmap;
+		if (OwningLandscape)
+		{
+			OwningLandscape->OnBlueprintBrushChanged();
+		}
 	}
-#endif
+#endif // WITH_EDITORONLY_DATA
 }
 
-void ALandscapeBlueprintBrushBase::SetAffectsVisibilityLayer(bool bInAffectsVisibilityLayer)
+void ALandscapeBlueprintBrushBase::SetCanAffectVisibilityLayer(bool bInCanAffectVisibilityLayer)
 {
 #if WITH_EDITORONLY_DATA
-	Modify();
-	AffectVisibilityLayer = bInAffectsVisibilityLayer;
-	if (OwningLandscape)
+	if (bInCanAffectVisibilityLayer != AffectVisibilityLayer)
 	{
-		OwningLandscape->OnBlueprintBrushChanged();
+		Modify();
+		AffectVisibilityLayer = bInCanAffectVisibilityLayer;
+		if (OwningLandscape)
+		{
+			OwningLandscape->OnBlueprintBrushChanged();
+		}
 	}
-#endif
+#endif // WITH_EDITORONLY_DATA
 }
 
-bool ALandscapeBlueprintBrushBase::IsAffectingWeightmapLayer(const FName& InLayerName) const
+bool ALandscapeBlueprintBrushBase::AffectsWeightmapLayer(const FName& InLayerName) const
 {
 #if WITH_EDITORONLY_DATA
+	if (!CanAffectWeightmap())
+	{
+		return false;
+	}
+
 	return AffectedWeightmapLayers.Contains(InLayerName);
-#else
+#else // WITH_EDITORONLY_DATA
 	return false;
-#endif
+#endif // !WITH_EDITORONLY_DATA
 }
 
 void ALandscapeBlueprintBrushBase::PostEditMove(bool bFinished)
diff --git a/Engine/Source/Runtime/Landscape/Private/LandscapeEditLayers.cpp b/Engine/Source/Runtime/Landscape/Private/LandscapeEditLayers.cpp
--- a/Engine/Source/Runtime/Landscape/Private/LandscapeEditLayers.cpp
+++ b/Engine/Source/Runtime/Landscape/Private/LandscapeEditLayers.cpp
@@ -3682,7 +3682,7 @@ bool ALandscape::PrepareLayersBrushResources(ERHIFeatureLevel::Type InFeatureLev
 		{
 			if (ALandscapeBlueprintBrushBase* LandscapeBrush = Brush.GetBrush())
 			{
-				if (LandscapeBrush->IsAffectingWeightmap() || LandscapeBrush->IsAffectingHeightmap())
+				if (LandscapeBrush->AffectsWeightmap() || LandscapeBrush->AffectsHeightmap() || LandscapeBrush->AffectsVisibilityLayer())
 				{
 					LandscapeBrush->GetRenderDependencies(Dependencies);
 				}
@@ -7122,11 +7122,11 @@ int32 ALandscape::PerformLayersWeightmapsGlobalMerge(FUpdateLayersContentContext
 							FLandscapeLayerBrush& Brush = Layer.Brushes[i];
 							TOptional<int32> LayerInfoSettingsAllocatedIndex;
 
-							if (Brush.IsAffectingWeightmapLayer(InfoLayerSettings.GetLayerName()) && !LayerInfoObjects.Contains(InfoLayerSettings.LayerInfoObj))
+							if (Brush.AffectsWeightmapLayer(InfoLayerSettings.GetLayerName()) && !LayerInfoObjects.Contains(InfoLayerSettings.LayerInfoObj))
 							{
 								LayerInfoSettingsAllocatedIndex = LayerInfoSettingsIndex + 1; // due to visibility layer that is at 0
 							}
-							else if (Brush.IsAffectingVisibilityLayer() && UE::Landscape::IsVisibilityLayer(InfoLayerSettings.LayerInfoObj) && !LayerInfoObjects.Contains(InfoLayerSettings.LayerInfoObj))
+							else if (Brush.AffectsVisibilityLayer() && UE::Landscape::IsVisibilityLayer(InfoLayerSettings.LayerInfoObj) && !LayerInfoObjects.Contains(InfoLayerSettings.LayerInfoObj))
 							{
 								LayerInfoSettingsAllocatedIndex = GetVisibilityLayerAllocationIndex();
 							}
@@ -9989,28 +9989,28 @@ void FLandscapeLayerBrush::SetOwner(ALandscape* InOwner)
 #endif
 }
 
-bool FLandscapeLayerBrush::IsAffectingHeightmap() const
+bool FLandscapeLayerBrush::AffectsHeightmap() const
 {
 #if WITH_EDITORONLY_DATA
-	return BlueprintBrush && BlueprintBrush->IsVisible() && BlueprintBrush->IsAffectingHeightmap();
+	return BlueprintBrush && BlueprintBrush->IsVisible() && BlueprintBrush->AffectsHeightmap();
 #else
 	return false;
 #endif
 }
 
-bool FLandscapeLayerBrush::IsAffectingWeightmapLayer(const FName& InWeightmapLayerName) const
+bool FLandscapeLayerBrush::AffectsWeightmapLayer(const FName& InWeightmapLayerName) const
 {
 #if WITH_EDITORONLY_DATA
-	return BlueprintBrush && BlueprintBrush->IsVisible() && BlueprintBrush->IsAffectingWeightmap() && BlueprintBrush->IsAffectingWeightmapLayer(InWeightmapLayerName);
+	return BlueprintBrush && BlueprintBrush->IsVisible() && BlueprintBrush->AffectsWeightmap() && BlueprintBrush->AffectsWeightmapLayer(InWeightmapLayerName);
 #else
 	return false;
 #endif
 }
 
-bool FLandscapeLayerBrush::IsAffectingVisibilityLayer() const
+bool FLandscapeLayerBrush::AffectsVisibilityLayer() const
 {
 #if WITH_EDITORONLY_DATA
-	return BlueprintBrush && BlueprintBrush->IsVisible() && BlueprintBrush->IsAffectingVisibilityLayer();
+	return BlueprintBrush && BlueprintBrush->IsVisible() && BlueprintBrush->AffectsVisibilityLayer();
 #else
 	return false;
 #endif
@@ -10028,9 +10028,9 @@ UTextureRenderTarget2D* FLandscapeLayerBrush::RenderLayer(const FIntRect& InLand
 {
 #if WITH_EDITORONLY_DATA
 	TRACE_CPUPROFILER_EVENT_SCOPE(LandscapeLayers_LayerBrushRender);
-	if (((InParameters.LayerType == ELandscapeToolTargetType::Heightmap) && !IsAffectingHeightmap()) ||
-		((InParameters.LayerType == ELandscapeToolTargetType::Weightmap) && !IsAffectingWeightmapLayer(InParameters.WeightmapLayerName)) ||
-		((InParameters.LayerType == ELandscapeToolTargetType::Visibility) && !IsAffectingVisibilityLayer()))
+	if (((InParameters.LayerType == ELandscapeToolTargetType::Heightmap) && !AffectsHeightmap()) ||
+		((InParameters.LayerType == ELandscapeToolTargetType::Weightmap) && !AffectsWeightmapLayer(InParameters.WeightmapLayerName)) ||
+		((InParameters.LayerType == ELandscapeToolTargetType::Visibility) && !AffectsVisibilityLayer()))
 	{
 		return nullptr;
 	}
diff --git a/Engine/Source/Runtime/Landscape/Public/LandscapeBlueprintBrushBase.h b/Engine/Source/Runtime/Landscape/Public/LandscapeBlueprintBrushBase.h
--- a/Engine/Source/Runtime/Landscape/Public/LandscapeBlueprintBrushBase.h
+++ b/Engine/Source/Runtime/Landscape/Public/LandscapeBlueprintBrushBase.h
@@ -52,13 +52,13 @@ protected:
 	UPROPERTY(Category = "Settings", EditAnywhere, BlueprintReadWrite)
 	bool UpdateOnPropertyChange;
 
-	UPROPERTY(Category = "Settings", EditAnywhere, BlueprintReadWrite)
+	UPROPERTY(Category = "Settings", EditAnywhere, BlueprintReadWrite, Setter = "SetCanAffectHeightmap")
 	bool AffectHeightmap;
 
-	UPROPERTY(Category = "Settings", EditAnywhere, BlueprintReadWrite)
+	UPROPERTY(Category = "Settings", EditAnywhere, BlueprintReadWrite, Setter = "SetCanAffectWeightmap")
 	bool AffectWeightmap;
 
-	UPROPERTY(Category = "Settings", EditAnywhere, BlueprintReadWrite)
+	UPROPERTY(Category = "Settings", EditAnywhere, BlueprintReadWrite, Setter="SetCanAffectVisibilityLayer")
 	bool AffectVisibilityLayer;
 
 	UPROPERTY(Category = "Settings", EditAnywhere, BlueprintReadWrite)
@@ -93,6 +93,10 @@ public:
 	UFUNCTION(BlueprintImplementableEvent, CallInEditor)
 	void GetBlueprintRenderDependencies(TArray<UObject*>& OutStreamableAssets);
 
+	void SetCanAffectHeightmap(bool bInCanAffectHeightmap);
+	void SetCanAffectWeightmap(bool bInCanAffectWeightmap);
+	void SetCanAffectVisibilityLayer(bool bInCanAffectVisibilityLayer);
+
 #if WITH_EDITOR
 	virtual void CheckForErrors() override;
 
@@ -101,16 +105,34 @@ public:
 	virtual void SetOwningLandscape(class ALandscape* InOwningLandscape);
 	class ALandscape* GetOwningLandscape() const;
 
+	UE_DEPRECATED(5.3, "Renamed CanAffectHeightmap")
 	bool IsAffectingHeightmap() const { return AffectHeightmap; }
+	UE_DEPRECATED(5.3, "Renamed CanAffectWeightmap")
 	bool IsAffectingWeightmap() const { return AffectWeightmap; }
-	bool IsAffectingVisibilityLayer() const { return AffectVisibilityLayer; }
+	UE_DEPRECATED(5.3, "Renamed AffectsVisibilityLayer")
 	virtual bool IsAffectingWeightmapLayer(const FName& InLayerName) const;
+	UE_DEPRECATED(5.3, "Renamed CanAffectVisibilityLayer")
+	bool IsAffectingVisibilityLayer() const { return AffectVisibilityLayer; }
+
+	bool CanAffectHeightmap() const { return AffectHeightmap; }
+	bool CanAffectWeightmap() const { return AffectWeightmap; }
+	bool CanAffectVisibilityLayer() const { return AffectVisibilityLayer; }
+
+	virtual bool AffectsHeightmap() const { return CanAffectHeightmap(); }
+	virtual bool AffectsWeightmap() const { return CanAffectWeightmap(); }
+	virtual bool AffectsWeightmapLayer(const FName& InLayerName) const;
+	virtual bool AffectsVisibilityLayer() const { return CanAffectVisibilityLayer(); }
+
 	bool IsVisible() const { return bIsVisible; }
 	bool IsLayerUpdatePending() const;
 
 	void SetIsVisible(bool bInIsVisible);
+
+	UE_DEPRECATED(5.3, "Renamed SetCanAffectHeightmap")
 	void SetAffectsHeightmap(bool bInAffectsHeightmap);
+	UE_DEPRECATED(5.3, "Renamed SetCanAffectWeightmap")
 	void SetAffectsWeightmap(bool bInAffectsWeightmap);
+	UE_DEPRECATED(5.3, "Renamed SetCanAffectVisibilityLayer")
 	void SetAffectsVisibilityLayer(bool bInAffectsVisibilityLayer);
 
 	virtual bool ShouldTickIfViewportsOnly() const override;
