Landscape : Change weight painting strength

Add a gamma function to the paint tools strength to make more of the slider range useful for painting. Setting: PaintStrengthGamma, defaults to 2.2 (1.0 should replicate old behavior)
Added option to disable the painting influence map which effectlively slows down the initial painting and makes it look like it accelerates. Setting: bDisablePaintingStartupSlowdown, defaults to true (false to return to old behavior)

diff --git a/Engine/Source/Editor/LandscapeEditor/Private/LandscapeEdModePaintTools.cpp b/Engine/Source/Editor/LandscapeEditor/Private/LandscapeEdModePaintTools.cpp
--- a/Engine/Source/Editor/LandscapeEditor/Private/LandscapeEdModePaintTools.cpp
+++ b/Engine/Source/Editor/LandscapeEditor/Private/LandscapeEdModePaintTools.cpp
@@ -14,6 +14,7 @@
 #include "LandscapeEdit.h"
 #include "LandscapeDataAccess.h"
 #include "LandscapeEdModeTools.h"
+#include "LandscapeSettings.h"
 #include "Landscape.h"
 #include "Logging/TokenizedMessage.h"
 #include "Logging/MessageLog.h"
@@ -120,10 +121,20 @@ public:
 		: Super(InEdMode, InViewportClient, InTarget)
 		, Cache(InTarget)
 	{
+		const ULandscapeSettings* Settings = GetDefault<ULandscapeSettings>();
+		PaintStrengthGamma = Settings->PaintStrengthGamma;
+		bDisablePaintingStartupSlowdown = Settings->bDisablePaintingStartupSlowdown;
+	}
+
+	float GetStrength(const ULandscapeEditorObject* UISettings, bool bInIsHeightMap) const
+	{
+		return bInIsHeightMap ? UISettings->GetCurrentToolStrength() : FMath::Pow(UISettings->GetCurrentToolStrength(), PaintStrengthGamma);
 	}
 
 protected:
 	typename ToolTarget::CacheClass Cache;
+	float PaintStrengthGamma = 1.0f;
+	bool bDisablePaintingStartupSlowdown = false;
 };
 
 // 
@@ -133,6 +144,7 @@ class FLandscapeToolStrokePaint : public FLandscapeToolStrokePaintBase<FWeightma
 {
 	using Super = FLandscapeToolStrokePaintBase<FWeightmapToolTarget>;
 	using ToolTarget = FWeightmapToolTarget;
+	using ValueType = ToolTarget::CacheClass::DataType;
 
 	TMap<FIntPoint, float> TotalInfluenceMap; // amount of time and weight the brush has spent on each vertex.
 
@@ -219,12 +231,12 @@ public:
 		this->Cache.CacheData(X1, Y1, X2, Y2, bCacheOriginalData);
 
 		// The data we'll be writing to
-		TArray<ToolTarget::CacheClass::DataType> Data;
+		TArray<ValueType> Data;
 		this->Cache.GetCachedData(X1, Y1, X2, Y2, Data);
 
 		// The source data we use for editing. 
-		TArray<ToolTarget::CacheClass::DataType>* SourceDataArrayPtr = &Data;
-		TArray<ToolTarget::CacheClass::DataType> OriginalData;
+		TArray<ValueType>* SourceDataArrayPtr = &Data;
+		TArray<ValueType> OriginalData;
 
 		if (!bUseWeightTargetValue)
 		{
@@ -236,25 +248,31 @@ public:
 
 			for (int32 Y = Y1; Y < Y2; Y++)
 			{
-				auto* DataScanline = Data.GetData() + (Y - Y1) * (X2 - X1 + 1) + (0 - X1);
-				auto* OriginalDataScanline = OriginalData.GetData() + (Y - Y1) * (X2 - X1 + 1) + (0 - X1);
+				ValueType* DataScanline = Data.GetData() + (Y - Y1) * (X2 - X1 + 1) + (0 - X1);
+				ValueType* OriginalDataScanline = OriginalData.GetData() + (Y - Y1) * (X2 - X1 + 1) + (0 - X1);
 				for (int32 X = X1; X < X2; X++)
 				{
-					float VertexInfluence = TotalInfluenceMap.FindRef(FIntPoint(X, Y));
-
-					auto& CurrentValue = DataScanline[X];
-					auto& SourceValue = OriginalDataScanline[X];
-
-					SourceValue = FMath::Lerp(SourceValue, CurrentValue, FMath::Min<float>(VertexInfluence * 0.05f, 1.0f));
+					ValueType& SourceValue = OriginalDataScanline[X];
+					ValueType& CurrentValue = DataScanline[X];
+					if (bDisablePaintingStartupSlowdown)
+					{
+						SourceValue = CurrentValue;
+					}
+					else
+					{
+						float VertexInfluence = TotalInfluenceMap.FindRef(FIntPoint(X, Y));
+						SourceValue = FMath::Lerp(SourceValue, CurrentValue, FMath::Min<float>(VertexInfluence * 0.05f, 1.0f));
+					}
+					
 				}
 			}
 		}
 
 		// Adjust strength based on brush size and drawscale, so strength 1 = one hemisphere
 		const float AdjustedStrength = ToolTarget::StrengthMultiplier(this->LandscapeInfo, UISettings->GetCurrentToolBrushRadius());
-		FWeightmapToolTarget::CacheClass::DataType DestValue = FWeightmapToolTarget::CacheClass::ClampValue(static_cast<int32>(255.0f * UISettings->WeightTargetValue));
-
-		float PaintStrength = UISettings->GetCurrentToolStrength() * Pressure * AdjustedStrength;
+		float PaintStrength = this->GetStrength(UISettings, /* bInIsHeightMap = */ false) * Pressure * AdjustedStrength;
+		
+		ValueType DestValue = FWeightmapToolTarget::CacheClass::ClampValue(static_cast<int32>(255.0f * UISettings->WeightTargetValue));
 
 		// TODO: make paint tool framerate independent like the sculpt tool
 		// const float DeltaTime = FMath::Min<float>(FApp::GetDeltaTime(), 0.1f); // Under 10 fps slow down paint speed
@@ -274,8 +292,8 @@ public:
 		for (int32 Y = BrushInfo.GetBounds().Min.Y; Y < BrushInfo.GetBounds().Max.Y; Y++)
 		{
 			const float* BrushScanline = BrushInfo.GetDataPtr(FIntPoint(0, Y));
-			auto* DataScanline = Data.GetData() + (Y - Y1) * (X2 - X1 + 1) + (0 - X1);
-			auto* SourceDataScanline = SourceDataArrayPtr->GetData() + (Y - Y1) * (X2 - X1 + 1) + (0 - X1);
+			ValueType* DataScanline = Data.GetData() + (Y - Y1) * (X2 - X1 + 1) + (0 - X1);
+			ValueType* SourceDataScanline = SourceDataArrayPtr->GetData() + (Y - Y1) * (X2 - X1 + 1) + (0 - X1);
 
 			for (int32 X = BrushInfo.GetBounds().Min.X; X < BrushInfo.GetBounds().Max.X; X++)
 			{
@@ -287,8 +305,8 @@ public:
 				TotalInfluenceMap.Add(Key, VertexInfluence + BrushValue);
 
 				float PaintAmount = BrushValue * PaintStrength;
-				auto& CurrentValue = DataScanline[X];
-				const auto& SourceValue = SourceDataScanline[X];
+				ValueType& CurrentValue = DataScanline[X];
+				const ValueType& SourceValue = SourceDataScanline[X];
 
 				if (bUseWeightTargetValue)
 				{
@@ -296,13 +314,14 @@ public:
 				}
 				else
 				{
+					const int32 IntPaintAmount = FMath::RoundToInt(PaintAmount);
 					if (bInvert)
 					{
-						CurrentValue = ToolTarget::CacheClass::ClampValue(FMath::Min<int32>(SourceValue - FMath::RoundToInt(PaintAmount), CurrentValue));
+						CurrentValue = ToolTarget::CacheClass::ClampValue(FMath::Min<int32>(SourceValue - IntPaintAmount, CurrentValue));
 					}
 					else
 					{
-						CurrentValue = ToolTarget::CacheClass::ClampValue(FMath::Max<int32>(SourceValue + FMath::RoundToInt(PaintAmount), CurrentValue));
+						CurrentValue = ToolTarget::CacheClass::ClampValue(FMath::Max<int32>(SourceValue + IntPaintAmount, CurrentValue));
 					}
 				}
 			}
@@ -345,7 +364,8 @@ class FLandscapeToolStrokeErase : public FLandscapeToolStrokePaintBase<FHeightma
 {
 	using Super = FLandscapeToolStrokePaintBase<FHeightmapToolTarget>;
 	using ToolTarget = FHeightmapToolTarget;
-	const typename ToolTarget::CacheClass::DataType FlattenHeight;
+	using ValueType = typename ToolTarget::CacheClass::DataType;
+	const ValueType FlattenHeight;
 
 public:
 	// Heightmap sculpt tool will continuously sculpt in the same location, weightmap paint tool doesn't
@@ -384,14 +404,14 @@ public:
 		
 		this->Cache.CacheData(X1, Y1, X2, Y2);
 
-		TArray<typename ToolTarget::CacheClass::DataType> Data;
+		TArray<ValueType> Data;
 		this->Cache.GetCachedData(X1, Y1, X2, Y2, Data);
 
 		// Apply the brush
 		for (int32 Y = BrushInfo.GetBounds().Min.Y; Y < BrushInfo.GetBounds().Max.Y; Y++)
 		{
 			const float* BrushScanline = BrushInfo.GetDataPtr(FIntPoint(0, Y));
-			auto* DataScanline = Data.GetData() + (Y - Y1) * (X2 - X1 + 1) + (0 - X1);
+			ValueType* DataScanline = Data.GetData() + (Y - Y1) * (X2 - X1 + 1) + (0 - X1);
 
 			for (int32 X = BrushInfo.GetBounds().Min.X; X < BrushInfo.GetBounds().Max.X; X++)
 			{
@@ -404,11 +424,11 @@ public:
 					int32 Delta = DataScanline[X] - FlattenHeight;
 					if (Delta > 0)
 					{
-						DataScanline[X] = static_cast<ToolTarget::CacheClass::DataType>(FMath::FloorToInt(FMath::Lerp((float)DataScanline[X], (float)FlattenHeight, Strength)));
+						DataScanline[X] = static_cast<ValueType>(FMath::FloorToInt(FMath::Lerp((float)DataScanline[X], (float)FlattenHeight, Strength)));
 					}
 					else
 					{
-						DataScanline[X] = static_cast<ToolTarget::CacheClass::DataType>(FMath::CeilToInt(FMath::Lerp((float)DataScanline[X], (float)FlattenHeight, Strength)));
+						DataScanline[X] = static_cast<ValueType>(FMath::CeilToInt(FMath::Lerp((float)DataScanline[X], (float)FlattenHeight, Strength)));
 					}
 				}
 			}
@@ -424,6 +444,7 @@ class FLandscapeToolStrokeSculpt : public FLandscapeToolStrokePaintBase<FHeightm
 {
 	using Super = FLandscapeToolStrokePaintBase<FHeightmapToolTarget>;
 	using ToolTarget = FHeightmapToolTarget;
+	using ValueType = ToolTarget::CacheClass::DataType;
 
 public:
 	// Heightmap sculpt tool will continuously sculpt in the same location, weightmap paint tool doesn't
@@ -471,11 +492,11 @@ public:
 		bool bUseClayBrush = UISettings->bUseClayBrush;
 
 		// The data we'll be writing to
-		TArray<ToolTarget::CacheClass::DataType> Data;
+		TArray<ValueType> Data;
 		this->Cache.GetCachedData(X1, Y1, X2, Y2, Data);
 
 		// The source data we use for editing. 
-		TArray<ToolTarget::CacheClass::DataType>* SourceDataArrayPtr = &Data;
+		TArray<ValueType>* SourceDataArrayPtr = &Data;
 
 		FMatrix ToWorld = ToolTarget::ToWorldMatrix(this->LandscapeInfo);
 		FMatrix FromWorld = ToolTarget::FromWorldMatrix(this->LandscapeInfo);
@@ -510,10 +531,10 @@ public:
 
 			for (int32 Y = Y1; Y < Y2; Y++)
 			{
-				auto* SourceDataScanline_0 = SourceDataArrayPtr->GetData() + (Y - Y1) * (X2 - X1 + 1) + (0 - X1);
-				auto* SourceDataScanline_1 = SourceDataArrayPtr->GetData() + (Y - Y1) * (X2 - X1 + 1) + (0 - X1);
-				auto* NormalsScanline_0 = Normals.GetData() + (Y - Y1) * (X2 - X1 + 1) + (0 - X1);
-				auto* NormalsScanline_1 = Normals.GetData() + (Y + 1 - Y1) * (X2 - X1 + 1) + (0 - X1);
+				ValueType* SourceDataScanline_0 = SourceDataArrayPtr->GetData() + (Y - Y1) * (X2 - X1 + 1) + (0 - X1);
+				ValueType* SourceDataScanline_1 = SourceDataArrayPtr->GetData() + (Y - Y1) * (X2 - X1 + 1) + (0 - X1);
+				FVector* NormalsScanline_0 = Normals.GetData() + (Y - Y1) * (X2 - X1 + 1) + (0 - X1);
+				FVector* NormalsScanline_1 = Normals.GetData() + (Y + 1 - Y1) * (X2 - X1 + 1) + (0 - X1);
 				for (int32 X = X1; X < X2; X++)
 				{
 					FVector Vert00 = ToWorld.TransformPosition(FVector((float)X + 0.0f, (float)Y + 0.0f, SourceDataScanline_0[X + 0]));
@@ -533,7 +554,7 @@ public:
 			}
 			for (int32 Y = Y1; Y <= Y2; Y++)
 			{
-				auto* NormalsScanline = Normals.GetData() + (Y - Y1) * (X2 - X1 + 1) + (0 - X1);
+				FVector* NormalsScanline = Normals.GetData() + (Y - Y1) * (X2 - X1 + 1) + (0 - X1);
 				for (int32 X = X1; X <= X2; X++)
 				{
 					NormalsScanline[X] = NormalsScanline[X].GetSafeNormal();
@@ -547,8 +568,8 @@ public:
 			for (int32 Y = BrushInfo.GetBounds().Min.Y; Y < BrushInfo.GetBounds().Max.Y; Y++)
 			{
 				const float* BrushScanline = BrushInfo.GetDataPtr(FIntPoint(0, Y));
-				auto* SourceDataScanline = SourceDataArrayPtr->GetData() + (Y - Y1) * (X2 - X1 + 1) + (0 - X1);
-				auto* NormalsScanline = Normals.GetData() + (Y - Y1) * (X2 - X1 + 1) + (0 - X1);
+				ValueType* SourceDataScanline = SourceDataArrayPtr->GetData() + (Y - Y1) * (X2 - X1 + 1) + (0 - X1);
+				FVector* NormalsScanline = Normals.GetData() + (Y - Y1) * (X2 - X1 + 1) + (0 - X1);
 
 				for (int32 X = BrushInfo.GetBounds().Min.X; X < BrushInfo.GetBounds().Max.X; X++)
 				{
@@ -594,8 +615,8 @@ public:
 			for (int32 Y = BrushInfo.GetBounds().Min.Y; Y < BrushInfo.GetBounds().Max.Y; Y++)
 			{
 				const float* BrushScanline = BrushInfo.GetDataPtr(FIntPoint(0, Y));
-				auto* DataScanline = Data.GetData() + (Y - Y1) * (X2 - X1 + 1) + (0 - X1);
-				auto* SourceDataScanline = SourceDataArrayPtr->GetData() + (Y - Y1) * (X2 - X1 + 1) + (0 - X1);
+				ValueType* DataScanline = Data.GetData() + (Y - Y1) * (X2 - X1 + 1) + (0 - X1);
+				ValueType* SourceDataScanline = SourceDataArrayPtr->GetData() + (Y - Y1) * (X2 - X1 + 1) + (0 - X1);
 
 				for (int32 X = BrushInfo.GetBounds().Min.X; X < BrushInfo.GetBounds().Max.X; X++)
 				{
@@ -603,8 +624,8 @@ public:
 					const float BrushValue = BrushScanline[X];
 
 					float SculptAmount = BrushValue * SculptStrength;
-					auto& CurrentValue = DataScanline[X];
-					const auto& SourceValue = SourceDataScanline[X];
+					ValueType& CurrentValue = DataScanline[X];
+					const ValueType& SourceValue = SourceDataScanline[X];
 
 					if (bUseClayBrush)
 					{
@@ -684,7 +705,7 @@ template<class ToolTarget>
 class FLandscapeToolStrokeSmooth : public FLandscapeToolStrokePaintBase<ToolTarget>
 {
 	using Super = FLandscapeToolStrokePaintBase<ToolTarget>;
-
+	using ValueType = typename ToolTarget::CacheClass::DataType;
 	bool bTargetIsHeightmap;
 	FLandscapeLayerDataCache<ToolTarget> LayerDataCache;
 public:
@@ -732,17 +753,17 @@ public:
 			Y2 += 1;
 		}
 
-		TArray<typename ToolTarget::CacheClass::DataType> Data;
+		TArray<ValueType> Data;
 		LayerDataCache.Initialize(this->LandscapeInfo, bCombinedLayerOperation);
 		LayerDataCache.Read(X1, Y1, X2, Y2, Data);
-		const TArray<typename ToolTarget::CacheClass::DataType> ReadData { Data };
+		const TArray<ValueType> ReadData { Data };
 		
-		const float ToolStrength = FMath::Clamp<float>(UISettings->GetCurrentToolStrength() * Pressure, 0.0f, 1.0f);
+		const float ToolStrength = FMath::Clamp<float>(this->GetStrength(UISettings, bTargetIsHeightmap) * Pressure, 0.0f, 1.0f);
 
 		// Apply the brush
 		if (UISettings->bDetailSmooth)
 		{
-			LowPassFilter<typename ToolTarget::CacheClass::DataType>(X1, Y1, X2, Y2, BrushInfo, Data, UISettings->DetailScale, ToolStrength);
+			LowPassFilter<ValueType>(X1, Y1, X2, Y2, BrushInfo, Data, UISettings->DetailScale, ToolStrength);
 		}
 		else
 		{
@@ -751,7 +772,7 @@ public:
 			for (int32 Y = BrushInfo.GetBounds().Min.Y; Y < BrushInfo.GetBounds().Max.Y; Y++)
 			{
 				const float* BrushScanline = BrushInfo.GetDataPtr(FIntPoint(0, Y));
-				auto* DataScanline = Data.GetData() + (Y - Y1) * (X2 - X1 + 1) + (0 - X1);
+				ValueType* DataScanline = Data.GetData() + (Y - Y1) * (X2 - X1 + 1) + (0 - X1);
 
 				for (int32 X = BrushInfo.GetBounds().Min.X; X < BrushInfo.GetBounds().Max.X; X++)
 				{
@@ -775,7 +796,7 @@ public:
 						{
 							const float* SampleBrushScanline = BrushInfo.GetDataPtr(FIntPoint(0, SampleY));
 							const float* SampleBrushScanline2 = BrushInfo.GetDataPtr(FIntPoint(0, Y + (Y - SampleY)));
-							const auto* SampleDataScanline = ReadData.GetData() + (SampleY - Y1) * (X2 - X1 + 1) + (0 - X1);
+							const ValueType* SampleDataScanline = ReadData.GetData() + (SampleY - Y1) * (X2 - X1 + 1) + (0 - X1);
 
 							for (int32 SampleX = SampleX1; SampleX <= SampleX2; SampleX++)
 							{
@@ -795,7 +816,7 @@ public:
 
 						FilterValue /= FilterSamplingNumber;
 
-						DataScanline[X] = FMath::Lerp(DataScanline[X], (typename ToolTarget::CacheClass::DataType)FilterValue, BrushValue);
+						DataScanline[X] = FMath::Lerp(DataScanline[X], static_cast<ValueType>(FilterValue), BrushValue);
 					}
 				}
 			}
@@ -834,7 +855,9 @@ public:
 template<class ToolTarget>
 class FLandscapeToolStrokeFlatten : public FLandscapeToolStrokePaintBase<ToolTarget>
 {
-	typename ToolTarget::CacheClass::DataType FlattenValue;
+	using Super = FLandscapeToolStrokePaintBase<ToolTarget>;
+	using ValueType = typename ToolTarget::CacheClass::DataType;
+	ValueType FlattenValue;
 
 	FVector FlattenNormal;
 	float FlattenPlaneDist;
@@ -853,7 +876,7 @@ public:
 		{
 			FTransform LocalToWorld = InTarget.LandscapeInfo->GetLandscapeProxy()->ActorToWorld();
 			float Height = static_cast<float>((InEdMode->UISettings->FlattenTarget - LocalToWorld.GetTranslation().Z) / LocalToWorld.GetScale3D().Z);
-			FlattenValue = static_cast<typename ToolTarget::CacheClass::DataType>(LandscapeDataAccess::GetTexHeight(Height));
+			FlattenValue = static_cast<ValueType>(LandscapeDataAccess::GetTexHeight(Height));
 			bInitializedFlattenValue = true;
 		}
 	}
@@ -892,7 +915,7 @@ public:
 				typename ToolTarget::CacheClass::AccessorClass RuntimeDataAccessor(this->Target);
 				RuntimeDataAccessor.SetEditLayer(FGuid());
 
-				typename ToolTarget::CacheClass::DataType P00, P10, P01, P11;
+				ValueType P00, P10, P01, P11;
 				RuntimeDataAccessor.GetDataFast(FlattenHeightX, FlattenHeightY, FlattenHeightX, FlattenHeightY, &P00);
 				RuntimeDataAccessor.GetDataFast(FlattenHeightX + 1, FlattenHeightY, FlattenHeightX + 1, FlattenHeightY, &P10);
 				RuntimeDataAccessor.GetDataFast(FlattenHeightX, FlattenHeightY + 1, FlattenHeightX, FlattenHeightY + 1, &P01);
@@ -910,7 +933,7 @@ public:
 				this->Cache.CacheData(FlattenHeightX, FlattenHeightY, FlattenHeightX + 1, FlattenHeightY + 1);
 				InterpolatedValue = this->Cache.GetValue(FlattenX, FlattenY);
 			}
-			FlattenValue = static_cast<typename ToolTarget::CacheClass::DataType>(InterpolatedValue);
+			FlattenValue = static_cast<ValueType>(InterpolatedValue);
 
 			if (bUseSlopeFlatten && bTargetIsHeightmap)
 			{
@@ -959,15 +982,17 @@ public:
 			Y2 += 1;
 		}
 
-		TArray<typename ToolTarget::CacheClass::DataType> Data;
+		TArray<ValueType> Data;
 		LayerDataCache.Initialize(this->LandscapeInfo, bCombinedLayerOperation);
 		LayerDataCache.Read(X1, Y1, X2, Y2, Data);
 
+		const float PaintStrength = this->GetStrength(UISettings, bTargetIsHeightmap);
+
 		// Apply the brush
 		for (int32 Y = BrushInfo.GetBounds().Min.Y; Y < BrushInfo.GetBounds().Max.Y; Y++)
 		{
 			const float* BrushScanline = BrushInfo.GetDataPtr(FIntPoint(0, Y));
-			auto* DataScanline = Data.GetData() + (Y - Y1) * (X2 - X1 + 1) + (0 - X1);
+			ValueType* DataScanline = Data.GetData() + (Y - Y1) * (X2 - X1 + 1) + (0 - X1);
 
 			for (int32 X = BrushInfo.GetBounds().Min.X; X < BrushInfo.GetBounds().Max.X; X++)
 			{
@@ -975,7 +1000,7 @@ public:
 
 				if (BrushValue > 0.0f)
 				{
-					float Strength = FMath::Clamp<float>(BrushValue * UISettings->GetCurrentToolStrength() * Pressure, 0.0f, 1.0f);
+					float Strength = FMath::Clamp<float>(BrushValue * PaintStrength * Pressure, 0.0f, 1.0f);
 
 					if (!(bUseSlopeFlatten && bTargetIsHeightmap))
 					{
@@ -1010,7 +1035,7 @@ public:
 
 								float FinalHeight = FMath::Lerp(CurrentHeight, WorldHeight, Strength);
 								FinalHeight = (FinalHeight - TranslateZ) / ScaleZ;
-								DataScanline[X] = static_cast<typename ToolTarget::CacheClass::DataType>(LandscapeDataAccess::GetTexHeight(FinalHeight));
+								DataScanline[X] = static_cast<ValueType>(LandscapeDataAccess::GetTexHeight(FinalHeight));
 							}
 							break;
 						case ELandscapeToolFlattenMode::Interval:
@@ -1032,64 +1057,64 @@ public:
 							
 								//back to local space of landscape object
 								TargetHeight = (TargetHeight - TranslateZ) / ScaleZ;
-								DataScanline[X] = static_cast<typename ToolTarget::CacheClass::DataType>(LandscapeDataAccess::GetTexHeight(TargetHeight));
+								DataScanline[X] = static_cast<ValueType>(LandscapeDataAccess::GetTexHeight(TargetHeight));
 
 							}
 							break;
 						case ELandscapeToolFlattenMode::Raise:
 							if (Delta < 0)
 							{
-								DataScanline[X] = static_cast<typename ToolTarget::CacheClass::DataType>(FMath::CeilToInt(FMath::Lerp((float)DataScanline[X], (float)FlattenValue, Strength)));
+								DataScanline[X] = static_cast<ValueType>(FMath::CeilToInt(FMath::Lerp((float)DataScanline[X], (float)FlattenValue, Strength)));
 							}
 							break;
 						case ELandscapeToolFlattenMode::Lower:
 							if (Delta > 0)
 							{
-								DataScanline[X] = static_cast<typename ToolTarget::CacheClass::DataType>(FMath::FloorToInt(FMath::Lerp((float)DataScanline[X], (float)FlattenValue, Strength)));
+								DataScanline[X] = static_cast<ValueType>(FMath::FloorToInt(FMath::Lerp((float)DataScanline[X], (float)FlattenValue, Strength)));
 							}
 							break;
 						default:
 						case ELandscapeToolFlattenMode::Both:
 							if (Delta > 0)
 							{
-								DataScanline[X] = static_cast<typename ToolTarget::CacheClass::DataType>(FMath::FloorToInt(FMath::Lerp((float)DataScanline[X], (float)FlattenValue, Strength)));
+								DataScanline[X] = static_cast<ValueType>(FMath::FloorToInt(FMath::Lerp((float)DataScanline[X], (float)FlattenValue, Strength)));
 							}
 							else
 							{
-								DataScanline[X] = static_cast<typename ToolTarget::CacheClass::DataType>(FMath::CeilToInt(FMath::Lerp((float)DataScanline[X], (float)FlattenValue, Strength)));
+								DataScanline[X] = static_cast<ValueType>(FMath::CeilToInt(FMath::Lerp((float)DataScanline[X], (float)FlattenValue, Strength)));
 							}
 							break;
 						}
 					}
 					else
 					{
-						typename ToolTarget::CacheClass::DataType DestValue = static_cast<typename ToolTarget::CacheClass::DataType>(-(FlattenNormal.X * X + FlattenNormal.Y * Y + FlattenPlaneDist) / FlattenNormal.Z);
+						ValueType DestValue = static_cast<ValueType>(-(FlattenNormal.X * X + FlattenNormal.Y * Y + FlattenPlaneDist) / FlattenNormal.Z);
 						//float PlaneDist = FlattenNormal | FVector(X, Y, HeightData(HeightDataIndex)) + FlattenPlaneDist;
 						float PlaneDist = static_cast<float>(DataScanline[X] - DestValue);
-						DestValue = static_cast<typename ToolTarget::CacheClass::DataType>(DataScanline[X] - PlaneDist * Strength);
+						DestValue = static_cast<ValueType>(DataScanline[X] - PlaneDist * Strength);
 						switch (UISettings->FlattenMode)
 						{
 						case ELandscapeToolFlattenMode::Raise:
 							if (PlaneDist < 0)
 							{
-								DataScanline[X] = static_cast<typename ToolTarget::CacheClass::DataType>(FMath::CeilToInt(FMath::Lerp((float)DataScanline[X], (float)DestValue, Strength)));
+								DataScanline[X] = static_cast<ValueType>(FMath::CeilToInt(FMath::Lerp((float)DataScanline[X], (float)DestValue, Strength)));
 							}
 							break;
 						case ELandscapeToolFlattenMode::Lower:
 							if (PlaneDist > 0)
 							{
-								DataScanline[X] = static_cast<typename ToolTarget::CacheClass::DataType>(FMath::FloorToInt(FMath::Lerp((float)DataScanline[X], (float)DestValue, Strength)));
+								DataScanline[X] = static_cast<ValueType>(FMath::FloorToInt(FMath::Lerp((float)DataScanline[X], (float)DestValue, Strength)));
 							}
 							break;
 						default:
 						case ELandscapeToolFlattenMode::Both:
 							if (PlaneDist > 0)
 							{
-								DataScanline[X] = static_cast<typename ToolTarget::CacheClass::DataType>(FMath::FloorToInt(FMath::Lerp((float)DataScanline[X], (float)DestValue, Strength)));
+								DataScanline[X] = static_cast<ValueType>(FMath::FloorToInt(FMath::Lerp((float)DataScanline[X], (float)DestValue, Strength)));
 							}
 							else
 							{
-								DataScanline[X] = static_cast<typename ToolTarget::CacheClass::DataType>(FMath::CeilToInt(FMath::Lerp((float)DataScanline[X], (float)DestValue, Strength)));
+								DataScanline[X] = static_cast<ValueType>(FMath::CeilToInt(FMath::Lerp((float)DataScanline[X], (float)DestValue, Strength)));
 							}
 							break;
 						}
@@ -1252,6 +1277,7 @@ template<class ToolTarget>
 class FLandscapeToolStrokeNoise : public FLandscapeToolStrokePaintBase<ToolTarget>
 {
 	using Super = FLandscapeToolStrokePaintBase<ToolTarget>;
+	using ValueType = typename ToolTarget::CacheClass::DataType;
 
 public:
 	FLandscapeToolStrokeNoise(FEdModeLandscape* InEdMode, FEditorViewportClient* InViewportClient, const FLandscapeToolTarget& InTarget)
@@ -1289,7 +1315,7 @@ public:
 		}
 
 		this->Cache.CacheData(X1, Y1, X2, Y2);
-		TArray<typename ToolTarget::CacheClass::DataType> Data;
+		TArray<ValueType> Data;
 		this->Cache.GetCachedData(X1, Y1, X2, Y2, Data);
 
 		float BrushSizeAdjust = 1.0f;
@@ -1302,11 +1328,13 @@ public:
 		CA_SUPPRESS(6326);
 		bool bUseWeightTargetValue = UISettings->bUseWeightTargetValue && ToolTarget::TargetType == ELandscapeToolTargetType::Weightmap;
 
+		const float PaintStrength = this->GetStrength(UISettings, /* bInIsHeighthMap = */ false);
+
 		// Apply the brush
 		for (int32 Y = BrushInfo.GetBounds().Min.Y; Y < BrushInfo.GetBounds().Max.Y; Y++)
 		{
 			const float* BrushScanline = BrushInfo.GetDataPtr(FIntPoint(0, Y));
-			auto* DataScanline = Data.GetData() + (Y - Y1) * (X2 - X1 + 1) + (0 - X1);
+			ValueType* DataScanline = Data.GetData() + (Y - Y1) * (X2 - X1 + 1) + (0 - X1);
 
 			for (int32 X = BrushInfo.GetBounds().Min.X; X < BrushInfo.GetBounds().Max.X; X++)
 			{
@@ -1339,10 +1367,10 @@ public:
 					}
 					else
 					{
-						float TotalStrength = BrushValue * UISettings->GetCurrentToolStrength() * Pressure * ToolTarget::StrengthMultiplier(this->LandscapeInfo, UISettings->GetCurrentToolBrushRadius());
+						float TotalStrength = BrushValue * PaintStrength * Pressure * ToolTarget::StrengthMultiplier(this->LandscapeInfo, UISettings->GetCurrentToolBrushRadius());
 						FNoiseParameter NoiseParam(0, UISettings->NoiseScale, TotalStrength * BrushSizeAdjust);
 						float PaintAmount = NoiseModeConversion(UISettings->NoiseMode, NoiseParam.NoiseAmount, NoiseParam.Sample(X, Y));
-						DataScanline[X] = static_cast<typename ToolTarget::CacheClass::DataType>(ToolTarget::CacheClass::ClampValue(static_cast<int32>(OriginalValue + PaintAmount)));
+						DataScanline[X] = static_cast<ValueType>(ToolTarget::CacheClass::ClampValue(static_cast<int32>(OriginalValue + PaintAmount)));
 					}
 				}
 			}
@@ -1381,18 +1409,18 @@ public:
 //
 void FEdModeLandscape::InitializeTool_Paint()
 {
-	auto Tool_Sculpt = MakeUnique<FLandscapeToolSculpt>(this);
+	TUniquePtr<FLandscapeToolSculpt> Tool_Sculpt = MakeUnique<FLandscapeToolSculpt>(this);
 	Tool_Sculpt->ValidBrushes.Add("BrushSet_Circle");
 	Tool_Sculpt->ValidBrushes.Add("BrushSet_Alpha");
 	Tool_Sculpt->ValidBrushes.Add("BrushSet_Pattern");
 	Tool_Sculpt->ValidBrushes.Add("BrushSet_Component");
 	LandscapeTools.Add(MoveTemp(Tool_Sculpt));
 
-	auto Tool_Erase = MakeUnique<FLandscapeToolErase>(this);
+	TUniquePtr<FLandscapeToolErase> Tool_Erase = MakeUnique<FLandscapeToolErase>(this);
 	Tool_Erase->ValidBrushes.Add("BrushSet_Circle");
 	LandscapeTools.Add(MoveTemp(Tool_Erase));
 
-	auto Tool_Paint = MakeUnique<FLandscapeToolPaint>(this);
+	TUniquePtr<FLandscapeToolPaint> Tool_Paint = MakeUnique<FLandscapeToolPaint>(this);
 	Tool_Paint->ValidBrushes.Add("BrushSet_Circle");
 	Tool_Paint->ValidBrushes.Add("BrushSet_Alpha");
 	Tool_Paint->ValidBrushes.Add("BrushSet_Pattern");
@@ -1402,13 +1430,13 @@ void FEdModeLandscape::InitializeTool_Paint()
 
 void FEdModeLandscape::InitializeTool_Smooth()
 {
-	auto Tool_Smooth_Heightmap = MakeUnique<FLandscapeToolSmooth<FHeightmapToolTarget>>(this);
+	TUniquePtr<FLandscapeToolSmooth<FHeightmapToolTarget>> Tool_Smooth_Heightmap = MakeUnique<FLandscapeToolSmooth<FHeightmapToolTarget>>(this);
 	Tool_Smooth_Heightmap->ValidBrushes.Add("BrushSet_Circle");
 	Tool_Smooth_Heightmap->ValidBrushes.Add("BrushSet_Alpha");
 	Tool_Smooth_Heightmap->ValidBrushes.Add("BrushSet_Pattern");
 	LandscapeTools.Add(MoveTemp(Tool_Smooth_Heightmap));
 
-	auto Tool_Smooth_Weightmap = MakeUnique<FLandscapeToolSmooth<FWeightmapToolTarget>>(this);
+	TUniquePtr<FLandscapeToolSmooth<FWeightmapToolTarget>> Tool_Smooth_Weightmap = MakeUnique<FLandscapeToolSmooth<FWeightmapToolTarget>>(this);
 	Tool_Smooth_Weightmap->ValidBrushes.Add("BrushSet_Circle");
 	Tool_Smooth_Weightmap->ValidBrushes.Add("BrushSet_Alpha");
 	Tool_Smooth_Weightmap->ValidBrushes.Add("BrushSet_Pattern");
@@ -1417,13 +1445,13 @@ void FEdModeLandscape::InitializeTool_Smooth()
 
 void FEdModeLandscape::InitializeTool_Flatten()
 {
-	auto Tool_Flatten_Heightmap = MakeUnique<FLandscapeToolFlatten<FHeightmapToolTarget>>(this);
+	TUniquePtr<FLandscapeToolFlatten<FHeightmapToolTarget>> Tool_Flatten_Heightmap = MakeUnique<FLandscapeToolFlatten<FHeightmapToolTarget>>(this);
 	Tool_Flatten_Heightmap->ValidBrushes.Add("BrushSet_Circle");
 	Tool_Flatten_Heightmap->ValidBrushes.Add("BrushSet_Alpha");
 	Tool_Flatten_Heightmap->ValidBrushes.Add("BrushSet_Pattern");
 	LandscapeTools.Add(MoveTemp(Tool_Flatten_Heightmap));
 
-	auto Tool_Flatten_Weightmap = MakeUnique<FLandscapeToolFlatten<FWeightmapToolTarget>>(this);
+	TUniquePtr<FLandscapeToolFlatten<FWeightmapToolTarget>> Tool_Flatten_Weightmap = MakeUnique<FLandscapeToolFlatten<FWeightmapToolTarget>>(this);
 	Tool_Flatten_Weightmap->ValidBrushes.Add("BrushSet_Circle");
 	Tool_Flatten_Weightmap->ValidBrushes.Add("BrushSet_Alpha");
 	Tool_Flatten_Weightmap->ValidBrushes.Add("BrushSet_Pattern");
@@ -1432,13 +1460,13 @@ void FEdModeLandscape::InitializeTool_Flatten()
 
 void FEdModeLandscape::InitializeTool_Noise()
 {
-	auto Tool_Noise_Heightmap = MakeUnique<FLandscapeToolNoise<FHeightmapToolTarget>>(this);
+	TUniquePtr<FLandscapeToolNoise<FHeightmapToolTarget>> Tool_Noise_Heightmap = MakeUnique<FLandscapeToolNoise<FHeightmapToolTarget>>(this);
 	Tool_Noise_Heightmap->ValidBrushes.Add("BrushSet_Circle");
 	Tool_Noise_Heightmap->ValidBrushes.Add("BrushSet_Alpha");
 	Tool_Noise_Heightmap->ValidBrushes.Add("BrushSet_Pattern");
 	LandscapeTools.Add(MoveTemp(Tool_Noise_Heightmap));
 
-	auto Tool_Noise_Weightmap = MakeUnique<FLandscapeToolNoise<FWeightmapToolTarget>>(this);
+	TUniquePtr<FLandscapeToolNoise<FWeightmapToolTarget>> Tool_Noise_Weightmap = MakeUnique<FLandscapeToolNoise<FWeightmapToolTarget>>(this);
 	Tool_Noise_Weightmap->ValidBrushes.Add("BrushSet_Circle");
 	Tool_Noise_Weightmap->ValidBrushes.Add("BrushSet_Alpha");
 	Tool_Noise_Weightmap->ValidBrushes.Add("BrushSet_Pattern");
diff --git a/Engine/Source/Runtime/Landscape/Public/LandscapeSettings.h b/Engine/Source/Runtime/Landscape/Public/LandscapeSettings.h
--- a/Engine/Source/Runtime/Landscape/Public/LandscapeSettings.h
+++ b/Engine/Source/Runtime/Landscape/Public/LandscapeSettings.h
@@ -54,6 +54,12 @@ public:
 
 	UPROPERTY(config, EditAnywhere, Category = "Configuration", meta = (UIMin = "1", UIMax = "1024", ClampMin = "1", ClampMax = "1024", ToolTip = "Maximum Size of Import Image Cache in MB"))
 	uint32 MaxImageImportCacheSizeMegaBytes = 256;
+
+	UPROPERTY(config, EditAnywhere, Category = "Configuration", meta = (UIMin = "0.0", UIMax = "10.0", ClampMin = "0.0", ClampMax = "10.0", ToolTip = "Exponent for the Paint Tool Strength"))
+	float PaintStrengthGamma = 2.2f;
+
+	UPROPERTY(config, EditAnywhere, Category = "Configuration", meta = (ToolTip = "Disable Painting Startup Slowdown"))
+	bool bDisablePaintingStartupSlowdown = true;
 	
 	UPROPERTY(Config, Category = "Configuration", EditAnywhere)
 	ELandscapeDirtyingMode LandscapeDirtyingMode;
