Landscape : Fixed Region import with 2x2 subsections

diff --git a/Engine/Source/Editor/LandscapeEditor/Private/LandscapeEditorDetailCustomization_NewLandscape.cpp b/Engine/Source/Editor/LandscapeEditor/Private/LandscapeEditorDetailCustomization_NewLandscape.cpp
--- a/Engine/Source/Editor/LandscapeEditor/Private/LandscapeEditorDetailCustomization_NewLandscape.cpp
+++ b/Engine/Source/Editor/LandscapeEditor/Private/LandscapeEditorDetailCustomization_NewLandscape.cpp
@@ -1191,11 +1191,11 @@ FReply FLandscapeEditorDetailCustomization_NewLandscape::OnCreateButtonClicked()
 			TRACE_CPUPROFILER_EVENT_SCOPE(AddComponentsToRegion);
 						
 			// Create a LocationVolume around the region 
-			int32 RegionSizeXTexels = QuadsPerSection * UISettings->WorldPartitionRegionSize;
-			int32 RegionSizeYTexels = QuadsPerSection * UISettings->WorldPartitionRegionSize;
+			int32 RegionSizeXTexels = QuadsPerSection * UISettings->WorldPartitionRegionSize * UISettings->NewLandscape_SectionsPerComponent;
+			int32 RegionSizeYTexels = QuadsPerSection * UISettings->WorldPartitionRegionSize * UISettings->NewLandscape_SectionsPerComponent;
 
-			double RegionSizeX = QuadsPerSection * LandscapeProxy->GetActorScale3D().X * UISettings->WorldPartitionRegionSize;
-			double RegionSizeY = QuadsPerSection * LandscapeProxy->GetActorScale3D().Y * UISettings->WorldPartitionRegionSize;
+			double RegionSizeX = RegionSizeXTexels * LandscapeProxy->GetActorScale3D().X;
+			double RegionSizeY = RegionSizeYTexels * LandscapeProxy->GetActorScale3D().Y;
 			ALocationVolume* RegionVolume = LandscapeRegionUtils::CreateLandscapeRegionVolume(World, LandscapeProxy, RegionCoordinate, RegionSizeX);
 			RegionVolumes.Add(RegionVolume);
 			
