Landscape : Import fixes

diff --git a/Engine/Source/Editor/LandscapeEditor/Private/LandscapeEditorObject.cpp b/Engine/Source/Editor/LandscapeEditor/Private/LandscapeEditorObject.cpp
--- a/Engine/Source/Editor/LandscapeEditor/Private/LandscapeEditorObject.cpp
+++ b/Engine/Source/Editor/LandscapeEditor/Private/LandscapeEditorObject.cpp
@@ -710,7 +710,7 @@ void ULandscapeEditorObject::RefreshLayerImport(FLandscapeImportLayer& ImportLay
 			ImportLayer.ErrorMessage = FileInfo.ErrorMessage;
 			if (FileInfo.ResultCode == ELandscapeImportResult::Success)
 			{
-				if (FileInfo.PossibleResolutions[0] != FLandscapeFileResolution(ImportLandscape_Height, ImportLandscape_Width))
+				if (FileInfo.PossibleResolutions[0] != FLandscapeFileResolution(ImportLandscape_Width, ImportLandscape_Height) && bHeightmapSelected)
 				{
 					ImportLayer.ImportResult = ELandscapeImportResult::Error;
 					ImportLayer.ErrorMessage = NSLOCTEXT("LandscapeEditor.ImportLandscape", "Import_WeightHeightResolutionMismatch", "Weightmap import resolution isn't same as Heightmap resolution.");
diff --git a/Engine/Source/Editor/LandscapeEditor/Private/LandscapeImageFileCache.h b/Engine/Source/Editor/LandscapeEditor/Private/LandscapeImageFileCache.h
--- a/Engine/Source/Editor/LandscapeEditor/Private/LandscapeImageFileCache.h
+++ b/Engine/Source/Editor/LandscapeEditor/Private/LandscapeImageFileCache.h
@@ -49,6 +49,12 @@ public:
 		ILandscapeEditorModule& LandscapeEditorModule = FModuleManager::GetModuleChecked<ILandscapeEditorModule>("LandscapeEditor");
 		const ILandscapeFileFormat<T>* FileFormat = LandscapeEditorModule.GetFormatByExtension<T>(*FPaths::GetExtension(InImageFilename, true));
 
+		if (!FileFormat)
+		{
+			Result.ResultCode = ELandscapeImportResult::Error;
+			return Result;
+		}
+
 		const FLandscapeFileInfo FileInfo = FileFormat->Validate(InImageFilename);
 
 		if (FileInfo.ResultCode == ELandscapeImportResult::Success && FileInfo.PossibleResolutions.Num() > 0)
