Landscape : Import Fixes

* Fixed resolution warning when only importing weight maps
* Only use region import when Region exist.

diff --git a/Engine/Source/Editor/LandscapeEditor/Private/LandscapeEditorDetailCustomization_ImportExport.cpp b/Engine/Source/Editor/LandscapeEditor/Private/LandscapeEditorDetailCustomization_ImportExport.cpp
--- a/Engine/Source/Editor/LandscapeEditor/Private/LandscapeEditorDetailCustomization_ImportExport.cpp
+++ b/Engine/Source/Editor/LandscapeEditor/Private/LandscapeEditorDetailCustomization_ImportExport.cpp
@@ -551,6 +551,9 @@ FReply FLandscapeEditorDetailCustomization_ImportExport::OnImportExportButtonCli
 	{
 		FIntRect LandscapeExtent;
 		LandscapeInfo->GetLandscapeExtent(LandscapeExtent);
+
+		LandscapeExtent.Max.X += 1;
+		LandscapeExtent.Max.Y += 1;
 	
 		if (IsImporting())
 		{
@@ -580,7 +583,7 @@ FReply FLandscapeEditorDetailCustomization_ImportExport::OnImportExportButtonCli
 			check(LandscapeEdMode->UISettings->ImportLandscape_HeightmapImportResult != ELandscapeImportResult::Error);
 			const bool bIsWorldPartition = LandscapeEdMode->GetWorld()->GetSubsystem<ULandscapeSubsystem>()->IsGridBased();
 
-			if (bIsWorldPartition && LandscapeEdMode->UISettings->ImportExportMode != ELandscapeImportExportMode::LoadedOnly) // If in world partition import region by region
+			if (bIsWorldPartition && LandscapeEdMode->UISettings->ImportExportMode != ELandscapeImportExportMode::LoadedOnly && LandscapeRegionUtils::NumLandscapeRegions(LandscapeInfo) > 0) // If in world partition import region by region
 			{
 				TArray<ALocationVolume*> LandscapeRegions;
 				TArray<AActor*> Children;
diff --git a/Engine/Source/Editor/LandscapeEditor/Private/LandscapeRegionUtils.cpp b/Engine/Source/Editor/LandscapeEditor/Private/LandscapeRegionUtils.cpp
--- a/Engine/Source/Editor/LandscapeEditor/Private/LandscapeRegionUtils.cpp
+++ b/Engine/Source/Editor/LandscapeEditor/Private/LandscapeRegionUtils.cpp
@@ -146,4 +146,18 @@ void ForEachRegion_LoadProcessUnload(ULandscapeInfo* InLandscapeInfo, const FInt
 	}
 }
 
+int32 NumLandscapeRegions(ULandscapeInfo* InLandscapeInfo)
+{
+	int32 NumRegions = 0;
+	TArray<AActor*> Children;
+	InLandscapeInfo->LandscapeActor->GetAttachedActors(Children);
+	TArray<ALocationVolume*> LandscapeRegions;
+
+	for (AActor* Child : Children)
+	{
+		NumRegions += Child->IsA<ALocationVolume>() ? 1 : 0;
+	}
+
+	return NumRegions;
+}
 } // LandscapeEditor
\ No newline at end of file
diff --git a/Engine/Source/Editor/LandscapeEditor/Private/LandscapeRegionUtils.h b/Engine/Source/Editor/LandscapeEditor/Private/LandscapeRegionUtils.h
--- a/Engine/Source/Editor/LandscapeEditor/Private/LandscapeRegionUtils.h
+++ b/Engine/Source/Editor/LandscapeEditor/Private/LandscapeRegionUtils.h
@@ -24,4 +24,7 @@ namespace LandscapeRegionUtils
 
 	// Load, Process (call RegionFn) and Unload the Region
 	void ForEachRegion_LoadProcessUnload(ULandscapeInfo* LandscapeInfo, const FIntRect& Domain, UWorld* World, TFunctionRef<bool(const FBox&, const TArray<ALandscapeProxy*>)> RegionFn);
+
+	// Number of landscape regions in this landscape.
+	int32 NumLandscapeRegions(ULandscapeInfo* InLandscapeInfo);
 }
