Landscape  : Fix Landscape Export Save Dialog

* The open dialog was incorrectly used for selecting fils to export

diff --git a/Engine/Source/Editor/LandscapeEditor/Private/LandscapeEditorDetailCustomization_ImportExport.cpp b/Engine/Source/Editor/LandscapeEditor/Private/LandscapeEditorDetailCustomization_ImportExport.cpp
--- a/Engine/Source/Editor/LandscapeEditor/Private/LandscapeEditorDetailCustomization_ImportExport.cpp
+++ b/Engine/Source/Editor/LandscapeEditor/Private/LandscapeEditorDetailCustomization_ImportExport.cpp
@@ -460,7 +460,7 @@ FReply FLandscapeEditorDetailCustomization_ImportExport::OnBrowseFilenameButtonC
 	const FString DialogType = bIsImporting ? LandscapeEditorModule.GetHeightmapImportDialogTypeString() : LandscapeEditorModule.GetHeightmapExportDialogTypeString();
 	const FString DialogTitle = bIsImporting ? LOCTEXT("ImportHeightmap", "Import Heightmap").ToString() : LOCTEXT("ExportHeightmap", "Export Heightmap").ToString();
 
-	TOptional<FString> OptionalFilename = LandscapeEditorUtils::GetImportExportFilename(DialogTitle, LandscapeEdMode->UISettings->LastImportPath, DialogType);
+	TOptional<FString> OptionalFilename = LandscapeEditorUtils::GetImportExportFilename(DialogTitle, LandscapeEdMode->UISettings->LastImportPath, DialogType, bIsImporting);
 	if (OptionalFilename.IsSet())
 	{
 		const FString& Filename = OptionalFilename.GetValue();
diff --git a/Engine/Source/Editor/LandscapeEditor/Private/LandscapeEditorDetailCustomization_ImportLayers.cpp b/Engine/Source/Editor/LandscapeEditor/Private/LandscapeEditorDetailCustomization_ImportLayers.cpp
--- a/Engine/Source/Editor/LandscapeEditor/Private/LandscapeEditorDetailCustomization_ImportLayers.cpp
+++ b/Engine/Source/Editor/LandscapeEditor/Private/LandscapeEditorDetailCustomization_ImportLayers.cpp
@@ -236,7 +236,7 @@ FReply FLandscapeEditorStructCustomization_FLandscapeImportLayer::OnLayerFilenam
 	const FString DialogTypeString = bIsImporting  ? LandscapeEditorModule.GetWeightmapImportDialogTypeString() : LandscapeEditorModule.GetWeightmapExportDialogTypeString();
 	const FString DialogTitle = bIsImporting ? LOCTEXT("ImportLayer", "Import Layer").ToString() : LOCTEXT("ExportLayer", "Export Layer").ToString();
 
-	TOptional<FString> OptionalExportImportPath = LandscapeEditorUtils::GetImportExportFilename(DialogTitle, LandscapeEdMode->UISettings->LastImportPath, DialogTypeString);
+	TOptional<FString> OptionalExportImportPath = LandscapeEditorUtils::GetImportExportFilename(DialogTitle, LandscapeEdMode->UISettings->LastImportPath, DialogTypeString, bIsImporting);
 	if (OptionalExportImportPath.IsSet())
 	{
 		const FString& Filename = OptionalExportImportPath.GetValue();
diff --git a/Engine/Source/Editor/LandscapeEditor/Private/LandscapeEditorDetailCustomization_NewLandscape.cpp b/Engine/Source/Editor/LandscapeEditor/Private/LandscapeEditorDetailCustomization_NewLandscape.cpp
--- a/Engine/Source/Editor/LandscapeEditor/Private/LandscapeEditorDetailCustomization_NewLandscape.cpp
+++ b/Engine/Source/Editor/LandscapeEditor/Private/LandscapeEditorDetailCustomization_NewLandscape.cpp
@@ -1404,7 +1404,7 @@ FReply FLandscapeEditorDetailCustomization_NewLandscape::OnImportHeightmapFilena
 	ILandscapeEditorModule& LandscapeEditorModule = FModuleManager::GetModuleChecked<ILandscapeEditorModule>("LandscapeEditor");
 	const FString FileTypes = LandscapeEditorModule.GetHeightmapImportDialogTypeString();
 
-	TOptional<FString> OptionalFilename = LandscapeEditorUtils::GetImportExportFilename(NSLOCTEXT("UnrealEd", "Import", "Import").ToString(), LandscapeEdMode->UISettings->LastImportPath, FileTypes);
+	TOptional<FString> OptionalFilename = LandscapeEditorUtils::GetImportExportFilename(NSLOCTEXT("UnrealEd", "Import", "Import").ToString(), LandscapeEdMode->UISettings->LastImportPath, FileTypes, /* bInImporting = */ true);
 	if (OptionalFilename.IsSet())
 	{
 		const FString& Filename = OptionalFilename.GetValue();
diff --git a/Engine/Source/Editor/LandscapeEditor/Private/LandscapeEditorUtils.cpp b/Engine/Source/Editor/LandscapeEditor/Private/LandscapeEditorUtils.cpp
--- a/Engine/Source/Editor/LandscapeEditor/Private/LandscapeEditorUtils.cpp
+++ b/Engine/Source/Editor/LandscapeEditor/Private/LandscapeEditorUtils.cpp
@@ -27,7 +27,7 @@ namespace LandscapeEditorUtils
 	}
 
 
-	TOptional<FString> GetImportExportFilename(const FString& InDialogTitle, const FString& InStartPath, const FString& InDialogTypeString)
+	TOptional<FString> GetImportExportFilename(const FString& InDialogTitle, const FString& InStartPath, const FString& InDialogTypeString, bool bInImporting)
 	{
 		FString Filename;
 		IDesktopPlatform* DesktopPlatform = FDesktopPlatformModule::Get();
@@ -41,20 +41,36 @@ namespace LandscapeEditorUtils
 		ILandscapeEditorModule& LandscapeEditorModule = FModuleManager::GetModuleChecked<ILandscapeEditorModule>("LandscapeEditor");
 		FEdModeLandscape* LandscapeEdMode = static_cast<FEdModeLandscape*>(GLevelEditorModeTools().GetActiveMode(FBuiltinEditorModes::EM_Landscape));
 
-		const bool bSuccess = DesktopPlatform->OpenFileDialog(
-			FSlateApplication::Get().FindBestParentWindowHandleForDialogs(nullptr),
-			InDialogTitle,
-			InStartPath,
-			TEXT(""),
-			InDialogTypeString,
-			EFileDialogFlags::None,
-			Filenames);
+		bool bSuccess;
+		
+		if (bInImporting)
+		{
+			bSuccess = DesktopPlatform->OpenFileDialog(
+				FSlateApplication::Get().FindBestParentWindowHandleForDialogs(nullptr),
+				InDialogTitle,
+				InStartPath,
+				TEXT(""),
+				InDialogTypeString,
+				EFileDialogFlags::None,
+				Filenames);
+		}
+		else
+		{
+			bSuccess = DesktopPlatform->SaveFileDialog(
+				FSlateApplication::Get().FindBestParentWindowHandleForDialogs(nullptr),
+				InDialogTitle,
+				InStartPath,
+				TEXT(""),
+				InDialogTypeString,
+				EFileDialogFlags::None,
+				Filenames);
+		}
 
 		if (bSuccess)
 		{
 			FString TiledFileNamePattern;
 
-			if (FLandscapeTiledImage::CheckTiledNamePath(Filenames[0], TiledFileNamePattern) && FMessageDialog::Open(EAppMsgType::YesNo, FText::FromString(FString::Format(TEXT("Use '{0}' Tiled Image?"), { TiledFileNamePattern }))) == EAppReturnType::Yes)
+			if (bInImporting && FLandscapeTiledImage::CheckTiledNamePath(Filenames[0], TiledFileNamePattern) && FMessageDialog::Open(EAppMsgType::YesNo, FText::FromString(FString::Format(TEXT("Use '{0}' Tiled Image?"), { TiledFileNamePattern }))) == EAppReturnType::Yes)
 			{
 				Filename = TiledFileNamePattern;
 			}
diff --git a/Engine/Source/Editor/LandscapeEditor/Public/LandscapeEditorUtils.h b/Engine/Source/Editor/LandscapeEditor/Public/LandscapeEditorUtils.h
--- a/Engine/Source/Editor/LandscapeEditor/Public/LandscapeEditorUtils.h
+++ b/Engine/Source/Editor/LandscapeEditor/Public/LandscapeEditorUtils.h
@@ -16,7 +16,7 @@ namespace LandscapeEditorUtils
 	bool LANDSCAPEEDITOR_API SetWeightmapData(ALandscapeProxy* Landscape, ULandscapeLayerInfoObject* LayerObject, const TArray<uint8>& Data);
 
 	int32 GetMaxSizeInComponents();
-	TOptional<FString> GetImportExportFilename(const FString& InDialogTitle, const FString& InStartPath, const FString& InDialogTypeString);
+	TOptional<FString> GetImportExportFilename(const FString& InDialogTitle, const FString& InStartPath, const FString& InDialogTypeString, bool bInImporting);
 
 	template<typename T>
 	void SaveObjects(TArrayView<T*> InObjects)
