Fix potential overflow security issue importing landscape PNG files

diff --git a/Engine/Source/Editor/LandscapeEditor/Private/LandscapeFileFormatPng.cpp b/Engine/Source/Editor/LandscapeEditor/Private/LandscapeFileFormatPng.cpp
--- a/Engine/Source/Editor/LandscapeEditor/Private/LandscapeFileFormatPng.cpp
+++ b/Engine/Source/Editor/LandscapeEditor/Private/LandscapeFileFormatPng.cpp
@@ -44,11 +44,20 @@ FLandscapeFileInfo FLandscapeHeightmapFileFormat_Png::Validate(const TCHAR* Heig
 		IImageWrapperModule& ImageWrapperModule = FModuleManager::LoadModuleChecked<IImageWrapperModule>("ImageWrapper");
 		TSharedPtr<IImageWrapper> ImageWrapper = ImageWrapperModule.CreateImageWrapper(EImageFormat::PNG);
 
-		if (!ImageWrapper->SetCompressed(ImportData.GetData(), ImportData.Num()) || ImageWrapper->GetWidth() <= 0 || ImageWrapper->GetHeight() <= 0)
+		if (!ImageWrapper->SetCompressed(ImportData.GetData(), ImportData.Num()) 
+			|| ImageWrapper->GetWidth() <= 0 
+			|| ImageWrapper->GetHeight() <= 0)
 		{
 			Result.ResultCode = ELandscapeImportResult::Error;
 			Result.ErrorMessage = LOCTEXT("Import_HeightmapFileCorruptPng", "The heightmap file cannot be read (corrupt png?)");
 		}
+		else if ((ImageWrapper->GetWidth() > MAX_int32)
+			|| (ImageWrapper->GetHeight() > MAX_int32)
+			|| (ImageWrapper->GetWidth() > MAX_int64 / ImageWrapper->GetHeight()))	// the total pixel count should fit in an int64 to avoid overflow issues
+		{
+			Result.ResultCode = ELandscapeImportResult::Error;
+			Result.ErrorMessage = LOCTEXT("Import_HeightmapFileTooLargePng", "The heightmap file is too large to load");
+		}
 		else
 		{
 			if (ImageWrapper->GetFormat() != ERGBFormat::Gray)
@@ -62,8 +71,8 @@ FLandscapeFileInfo FLandscapeHeightmapFileFormat_Png::Validate(const TCHAR* Heig
 				Result.ErrorMessage = LOCTEXT("Import_HeightmapFileLowBitDepth", "The heightmap file appears to be an 8-bit png, 16-bit is preferred. The import *can* continue, but the result may be lower quality than desired.");
 			}
 			FLandscapeFileResolution ImportResolution;
-			ImportResolution.Width = ImageWrapper->GetWidth();
-			ImportResolution.Height = ImageWrapper->GetHeight();
+			ImportResolution.Width = static_cast<uint32>(ImageWrapper->GetWidth());
+			ImportResolution.Height = static_cast<uint32>(ImageWrapper->GetHeight());
 			Result.PossibleResolutions.Add(ImportResolution);
 		}
 	}
@@ -90,16 +99,24 @@ FLandscapeImportData<uint16> FLandscapeHeightmapFileFormat_Png::Import(const TCH
 		IImageWrapperModule& ImageWrapperModule = FModuleManager::LoadModuleChecked<IImageWrapperModule>("ImageWrapper");
 		TSharedPtr<IImageWrapper> ImageWrapper = ImageWrapperModule.CreateImageWrapper(EImageFormat::PNG);
 
+		const int64 Width = ExpectedResolution.Width;
+		const int64 Height = ExpectedResolution.Height;
+
 		if (!ImageWrapper->SetCompressed(TempData.GetData(), TempData.Num()))
 		{
 			Result.ResultCode = ELandscapeImportResult::Error;
 			Result.ErrorMessage = LOCTEXT("Import_HeightmapFileCorruptPng", "The heightmap file cannot be read (corrupt png?)");
 		}
-		else if (ImageWrapper->GetWidth() != ExpectedResolution.Width || ImageWrapper->GetHeight() != ExpectedResolution.Height)
+		else if (ImageWrapper->GetWidth() != Width || ImageWrapper->GetHeight() != Height)
 		{
 			Result.ResultCode = ELandscapeImportResult::Error;
 			Result.ErrorMessage = LOCTEXT("Import_HeightmapResolutionMismatch", "The heightmap file's resolution does not match the requested resolution");
 		}
+		else if (Width > MAX_int64 / Height) // total pixel count must fit in an int64
+		{
+			Result.ResultCode = ELandscapeImportResult::Error;
+			Result.ErrorMessage = LOCTEXT("Import_HeightmapFileTooLargePng", "The heightmap file is too large to load");
+		}
 		else
 		{
 			if (ImageWrapper->GetFormat() != ERGBFormat::Gray)
@@ -113,32 +130,37 @@ FLandscapeImportData<uint16> FLandscapeHeightmapFileFormat_Png::Import(const TCH
 				Result.ErrorMessage = LOCTEXT("Import_HeightmapFileLowBitDepth", "The heightmap file appears to be an 8-bit png, 16-bit is preferred. The import *can* continue, but the result may be lower quality than desired.");
 			}
 
+			const int64 TotalPixels = Width * Height;
+
 			TArray64<uint8> RawData;
 			if (ImageWrapper->GetBitDepth() <= 8)
 			{
-				if (!ImageWrapper->GetRaw(ERGBFormat::Gray, 8, RawData))
+				if (!ImageWrapper->GetRaw(ERGBFormat::Gray, 8, RawData)
+					|| (RawData.Num() != TotalPixels))
 				{
 					Result.ResultCode = ELandscapeImportResult::Error;
 					Result.ErrorMessage = LOCTEXT("Import_HeightmapFileCorruptPng", "The heightmap file cannot be read (corrupt png?)");
 				}
 				else
 				{
-					Result.Data.Empty(ExpectedResolution.Width * ExpectedResolution.Height);
+					Result.Data.Empty(TotalPixels);
 					Algo::Transform(RawData, Result.Data, [](uint8 Value) { return static_cast<uint16>(Value * 0x101); }); // Expand to 16-bit
 				}
 			}
 			else
 			{
-				if (!ImageWrapper->GetRaw(ERGBFormat::Gray, 16, RawData))
+				if (!ImageWrapper->GetRaw(ERGBFormat::Gray, 16, RawData)
+					|| (RawData.Num()/2 != TotalPixels))
 				{
 					Result.ResultCode = ELandscapeImportResult::Error;
 					Result.ErrorMessage = LOCTEXT("Import_HeightmapFileCorruptPng", "The heightmap file cannot be read (corrupt png?)");
 				}
 				else
 				{
-					Result.Data.Empty(ExpectedResolution.Width * ExpectedResolution.Height);
-					Result.Data.AddUninitialized(ExpectedResolution.Width * ExpectedResolution.Height);
-					FMemory::Memcpy(Result.Data.GetData(), RawData.GetData(), ExpectedResolution.Width * ExpectedResolution.Height * 2);
+					Result.Data.Empty(TotalPixels);
+					Result.Data.AddUninitialized(TotalPixels);
+					const uint64 TotalBytes = static_cast<uint64>(TotalPixels) * 2;
+					FMemory::Memcpy(Result.Data.GetData(), RawData.GetData(), TotalBytes);
 				}
 			}
 		}
diff --git a/Engine/Source/Editor/LandscapeEditor/Private/LandscapeImportHelper.cpp b/Engine/Source/Editor/LandscapeEditor/Private/LandscapeImportHelper.cpp
--- a/Engine/Source/Editor/LandscapeEditor/Private/LandscapeImportHelper.cpp
+++ b/Engine/Source/Editor/LandscapeEditor/Private/LandscapeImportHelper.cpp
@@ -81,19 +81,34 @@ ELandscapeImportResult GetImportDataInternal(const FLandscapeImportDescriptor& I
 		return ELandscapeImportResult::Error;
 	}
 	
-	int32 TotalWidth = ImportDescriptor.ImportResolutions[DescriptorIndex].Width;
-	int32 TotalHeight = ImportDescriptor.ImportResolutions[DescriptorIndex].Height;
+	int64 TotalWidth = ImportDescriptor.ImportResolutions[DescriptorIndex].Width;	// convert from uint32
+	int64 TotalHeight = ImportDescriptor.ImportResolutions[DescriptorIndex].Height;
+
+	if (TotalWidth <= 0 || TotalHeight <= 0)
+	{
+		OutMessage = LOCTEXT("Import_InvalidImportResolution", "Import Resolution is not valid");
+		return ELandscapeImportResult::Error;
+	}
+
+	if (TotalWidth > MAX_int64 / TotalHeight)	// Total Pixels should fit in an int64
+	{
+		OutMessage = LOCTEXT("Import_ImageTooLarge", "Landscape image is too large");
+		return ELandscapeImportResult::Error;
+	}
+
+	int64 TotalPixels = TotalWidth * TotalHeight;
 
 	OutData.Reset();
-	OutData.SetNumZeroed(TotalWidth * TotalHeight);
+	OutData.SetNumZeroed(TotalPixels);
+
 	// Initialize All to default value so that non-covered regions have data
 	TArray<T> StrideData;
 	StrideData.SetNumUninitialized(TotalWidth);
-	for (int32 X = 0; X < TotalWidth; ++X)
+	for (int64 X = 0; X < TotalWidth; ++X)
 	{
 		StrideData[X] = DefaultValue;
 	}
-	for (int32 Y = 0; Y < TotalHeight; ++Y)
+	for (int64 Y = 0; Y < TotalHeight; ++Y)
 	{
 		FMemory::Memcpy(&OutData[Y * TotalWidth], StrideData.GetData(), sizeof(T) * TotalWidth);
 	}
@@ -104,8 +119,8 @@ ELandscapeImportResult GetImportDataInternal(const FLandscapeImportDescriptor& I
 	const ILandscapeFileFormat<T>* FileFormat = LandscapeEditorModule.GetFormatByExtension<T>(*FPaths::GetExtension(ImportDescriptor.FileDescriptors[0].FilePath, true));
 	check(FileFormat);
 
-	int32 FileWidth = ImportDescriptor.FileResolutions[DescriptorIndex].Width;
-	int32 FileHeight = ImportDescriptor.FileResolutions[DescriptorIndex].Height;
+	int64 FileWidth = ImportDescriptor.FileResolutions[DescriptorIndex].Width;	// convert from uint32
+	int64 FileHeight = ImportDescriptor.FileResolutions[DescriptorIndex].Height;
 	
 	for (const FLandscapeImportFileDescriptor& FileDescriptor : ImportDescriptor.FileDescriptors)
 	{
@@ -117,12 +132,13 @@ ELandscapeImportResult GetImportDataInternal(const FLandscapeImportDescriptor& I
 			break;
 		}
 		
-		int32 StartX = FileDescriptor.Coord.X * FileWidth;
-		int32 StartY = FileDescriptor.Coord.Y * FileHeight;
+		// NOTE: this assumes the same file resolution for all descriptors..
+		int64 StartX = FileDescriptor.Coord.X * FileWidth;
+		int64 StartY = FileDescriptor.Coord.Y * FileHeight;
 		
-		for (int32 Y = 0; Y < FileHeight; ++Y)
+		for (int64 Y = 0; Y < FileHeight; ++Y)
 		{		
-			int32 DestY = StartY + Y;
+			int64 DestY = StartY + Y;
 			FMemory::Memcpy(&OutData[DestY * TotalWidth + StartX], &ImportData.Data[Y * FileWidth], FileWidth * sizeof(T));
 		}
 	}
diff --git a/Engine/Source/Runtime/ImageWrapper/Private/Formats/PngImageWrapper.cpp b/Engine/Source/Runtime/ImageWrapper/Private/Formats/PngImageWrapper.cpp
--- a/Engine/Source/Runtime/ImageWrapper/Private/Formats/PngImageWrapper.cpp
+++ b/Engine/Source/Runtime/ImageWrapper/Private/Formats/PngImageWrapper.cpp
@@ -426,9 +426,14 @@ void FPngImageWrapper::UncompressPNGData(const ERGBFormat InFormat, const int32
 			// Calculate Pixel Depth
 			const uint64 PixelChannels = (InFormat == ERGBFormat::Gray) ? 1 : 4;
 			const uint64 BytesPerPixel = (InBitDepth * PixelChannels) / 8;
-			const uint64 BytesPerRow = BytesPerPixel * Width;
-			RawData.Empty(Height * BytesPerRow);
-			RawData.AddUninitialized(Height * BytesPerRow);
+			check(BytesPerPixel > 0);
+
+			const int64 BytesPerRow = BytesPerPixel * Width;
+			check(Height <= MAX_int64 / BytesPerRow);	// check for overflow on multiplication
+			const int64 TotalBytes = Height * BytesPerRow;
+
+			RawData.Empty(TotalBytes);
+			RawData.AddUninitialized(TotalBytes);
 
 			png_set_read_fn(png_ptr, this, FPngImageWrapper::user_read_compressed);
 
