Landscape : fix import of non u16 greyscale images to landscape height

diff --git a/Engine/Source/Editor/LandscapeEditor/Private/LandscapeEditorObject.cpp b/Engine/Source/Editor/LandscapeEditor/Private/LandscapeEditorObject.cpp
--- a/Engine/Source/Editor/LandscapeEditor/Private/LandscapeEditorObject.cpp
+++ b/Engine/Source/Editor/LandscapeEditor/Private/LandscapeEditorObject.cpp
@@ -657,12 +657,20 @@ void ULandscapeEditorObject::RefreshImports()
 	{
 		FLandscapeTiledImage TiledImage;
 		
-		if (FLandscapeFileInfo FileInfo = TiledImage.Load(*ImportLandscape_HeightmapFilename); FileInfo.ResultCode == ELandscapeImportResult::Success)
+		FLandscapeFileInfo FileInfo = TiledImage.Load(*ImportLandscape_HeightmapFilename);
+
+		if (FileInfo.PossibleResolutions.Num() > 0)
 		{
 			ImportLandscape_Width = TiledImage.GetResolution().X;
 			ImportLandscape_Height = TiledImage.GetResolution().Y;
+		}
+
+		if ( FileInfo.ResultCode != ELandscapeImportResult::Error)
+		{
 			ChooseBestComponentSizeForImport();
 			ImportLandscapeData();
+			ImportLandscape_HeightmapImportResult = FileInfo.ResultCode;
+			ImportLandscape_HeightmapErrorMessage = FileInfo.ErrorMessage;
 		}
 	}
 
@@ -736,7 +744,7 @@ void ULandscapeEditorObject::ImportLandscapeData()
 	FLandscapeTiledImage TiledImage;
 	FLandscapeFileInfo FileInfo = TiledImage.Load(*ImportLandscape_HeightmapFilename);
 
-	if (FileInfo.ResultCode != ELandscapeImportResult::Success)
+	if (FileInfo.ResultCode == ELandscapeImportResult::Error)
 	{
 		ImportLandscape_Data.Empty();
 	}
diff --git a/Engine/Source/Editor/LandscapeEditor/Private/LandscapeImageFileCache.h b/Engine/Source/Editor/LandscapeEditor/Private/LandscapeImageFileCache.h
--- a/Engine/Source/Editor/LandscapeEditor/Private/LandscapeImageFileCache.h
+++ b/Engine/Source/Editor/LandscapeEditor/Private/LandscapeImageFileCache.h
@@ -22,6 +22,8 @@ struct FLandscapeImageDataRef
 {
 	TSharedPtr<TArray<uint8>> Data;
 	FIntPoint Resolution;
+	ELandscapeImportResult Result;
+	FText ErrorMessage;
 };
 
 class FLandscapeImageFileCache
@@ -42,6 +44,8 @@ public:
 			CacheEntry->UsageCount++;
 			OutImageData = CacheEntry->ImageData;
 			Result.PossibleResolutions.Add(FLandscapeFileResolution(OutImageData.Resolution.X, OutImageData.Resolution.Y));
+			Result.ResultCode = CacheEntry->ImageData.Result;
+			Result.ErrorMessage = CacheEntry->ImageData.ErrorMessage;
 			return Result;
 		}
 
@@ -57,7 +61,7 @@ public:
 
 		const FLandscapeFileInfo FileInfo = FileFormat->Validate(InImageFilename);
 
-		if (FileInfo.ResultCode == ELandscapeImportResult::Success && FileInfo.PossibleResolutions.Num() > 0)
+		if (FileInfo.ResultCode != ELandscapeImportResult::Error && FileInfo.PossibleResolutions.Num() > 0)
 		{
 			FLandscapeFileResolution ExpectedResolution = FileInfo.PossibleResolutions[0];
 			FLandscapeImportData<T> ImportData = FileFormat->Import(InImageFilename, ExpectedResolution);
@@ -67,6 +71,8 @@ public:
 			NewImageData.Data->SetNumUninitialized(BufferSize);
 			FMemory::Memcpy(NewImageData.Data->GetData(), ImportData.Data.GetData(), BufferSize);
 			NewImageData.Resolution = FIntPoint(ExpectedResolution.Width, ExpectedResolution.Height);
+			NewImageData.Result = FileInfo.ResultCode;
+			NewImageData.ErrorMessage = FileInfo.ErrorMessage;
 		}
 		else
 		{
@@ -78,6 +84,9 @@ public:
 		Add(FString(InImageFilename), OutImageData);
 		
 		Result.PossibleResolutions.Add(FLandscapeFileResolution(OutImageData.Resolution.X, OutImageData.Resolution.Y));
+		Result.ResultCode = FileInfo.ResultCode;
+		Result.ErrorMessage = FileInfo.ErrorMessage;
+
 		return Result;
 	}
 
diff --git a/Engine/Source/Editor/LandscapeEditor/Private/LandscapeTiledImage.cpp b/Engine/Source/Editor/LandscapeEditor/Private/LandscapeTiledImage.cpp
--- a/Engine/Source/Editor/LandscapeEditor/Private/LandscapeTiledImage.cpp
+++ b/Engine/Source/Editor/LandscapeEditor/Private/LandscapeTiledImage.cpp
@@ -115,15 +115,23 @@ FLandscapeFileInfo FLandscapeTiledImage::Load(const TCHAR* Filename)
 	{
 		FLandscapeImageDataRef ImageData;
 		FLandscapeImageFileCache& LandscapeImageFileCache = FModuleManager::GetModuleChecked<ILandscapeEditorModule>("LandscapeEditor").GetImageFileCache();
-		if (FLandscapeFileInfo TileResult16 = LandscapeImageFileCache.FindImage<uint16>(*Tile.Value, ImageData); TileResult16.ResultCode != ELandscapeImportResult::Success)
+		FLandscapeFileInfo TileLoadResult;
+		if (TileLoadResult = LandscapeImageFileCache.FindImage<uint16>(*Tile.Value, ImageData); TileLoadResult.ResultCode != ELandscapeImportResult::Success)
 		{
-			if (FLandscapeFileInfo TileResult8 = LandscapeImageFileCache.FindImage<uint8>(*Tile.Value, ImageData);  TileResult8.ResultCode != ELandscapeImportResult::Success)
+			if (TileLoadResult = LandscapeImageFileCache.FindImage<uint8>(*Tile.Value, ImageData);  TileLoadResult.ResultCode != ELandscapeImportResult::Success)
 			{
-				return TileResult16;
+				if (TileLoadResult.ResultCode == ELandscapeImportResult::Warning)
+				{
+					Result.ErrorMessage = TileLoadResult.ErrorMessage;
+					Result.ResultCode = TileLoadResult.ResultCode;
+				}
+				else if (TileLoadResult.ResultCode == ELandscapeImportResult::Error)
+				{
+					return TileLoadResult;
+				}
 			}
 		}
 
-
 		SizeInTiles.X = FMath::Max(SizeInTiles.X, Tile.Key.X + 1);
 		SizeInTiles.Y = FMath::Max(SizeInTiles.Y, Tile.Key.Y + 1);
 
diff --git a/Engine/Source/Editor/LandscapeEditor/Private/LandscapeTiledImage.h b/Engine/Source/Editor/LandscapeEditor/Private/LandscapeTiledImage.h
--- a/Engine/Source/Editor/LandscapeEditor/Private/LandscapeTiledImage.h
+++ b/Engine/Source/Editor/LandscapeEditor/Private/LandscapeTiledImage.h
@@ -69,7 +69,7 @@ public:
 					FLandscapeImageFileCache& LandscapeImageFileCache = FModuleManager::GetModuleChecked<ILandscapeEditorModule>("LandscapeEditor").GetImageFileCache();
 					FLandscapeFileInfo TileResult = LandscapeImageFileCache.FindImage<T>(**TileFilename, ImageData);
 
-					if (TileResult.ResultCode != ELandscapeImportResult::Success)
+					if (TileResult.ResultCode == ELandscapeImportResult::Error)
 					{
 						return TileResult;
 					}
