Landscape  : Fix LansscapeProxies being preventend from being unloaded when in LandscapeMode sculpt or paint mode

diff --git a/Engine/Source/Editor/LandscapeEditor/Private/LandscapeEdModeBrushes.cpp b/Engine/Source/Editor/LandscapeEditor/Private/LandscapeEdModeBrushes.cpp
--- a/Engine/Source/Editor/LandscapeEditor/Private/LandscapeEdModeBrushes.cpp
+++ b/Engine/Source/Editor/LandscapeEditor/Private/LandscapeEdModeBrushes.cpp
@@ -48,13 +48,16 @@ class FLandscapeBrushCircle : public FLandscapeBrush
 {
 	using Super = FLandscapeBrush;
 
-	TSet<TObjectPtr<ULandscapeComponent>> BrushMaterialComponents;
-	TArray<UMaterialInstanceDynamic*> BrushMaterialFreeInstances;
+	// Components which previously we're under the area of the brush.
+	 TSet<TWeakObjectPtr<ULandscapeComponent>> BrushMaterialComponents;
+	
+	// A Cache of previously created and now unused MIDs we can reuse if required.
+	TArray<TWeakObjectPtr<UMaterialInstanceDynamic>> BrushMaterialFreeInstances;
 
 protected:
 	FVector2f LastMousePosition;
 	TObjectPtr<UMaterialInterface> BrushMaterial;
-	TMap<TObjectPtr<ULandscapeComponent>, TObjectPtr<UMaterialInstanceDynamic>> BrushMaterialInstanceMap;
+	TMap<TWeakObjectPtr<ULandscapeComponent>, TWeakObjectPtr<UMaterialInstanceDynamic>> BrushMaterialInstanceMap;
 	bool bCanPaint;
 
 	virtual float CalculateFalloff(float Distance, float Radius, float Falloff) = 0;
@@ -76,12 +79,6 @@ public:
 	{
 		Collector.AddReferencedObject(BrushMaterial);
 
-		// Allow any currently unused material instances to be GC'd
-		BrushMaterialFreeInstances.Empty();
-
-		Collector.AddReferencedObjects(BrushMaterialComponents);
-		Collector.AddReferencedObjects(BrushMaterialInstanceMap);
-
 		// If a user tool removes any components then we will have bad (null) entries in our TSet/TMap, remove them
 		// We can't just call .Remove(nullptr) because the entries were hashed as non-null values so a hash lookup of nullptr won't find them
 		for (auto It = BrushMaterialComponents.CreateIterator(); It; ++It)
@@ -102,16 +99,16 @@ public:
 
 	virtual void LeaveBrush() override
 	{
-		for (ULandscapeComponent* Component : BrushMaterialComponents)
+		for (auto It = BrushMaterialInstanceMap.CreateIterator(); It; ++It)
 		{
-			if (Component)
+			if (ULandscapeComponent* Component = It.Key().Get())
 			{
 				Component->EditToolRenderData.ToolMaterial = nullptr;
 				Component->UpdateEditToolRenderData();
 			}
 		}
 
-		TArray<UMaterialInstanceDynamic*> BrushMaterialInstances;
+		TArray<TWeakObjectPtr<UMaterialInstanceDynamic>> BrushMaterialInstances;
 		ObjectPtrDecay(BrushMaterialInstanceMap).GenerateValueArray(BrushMaterialInstances);
 		BrushMaterialFreeInstances += BrushMaterialInstances;
 		BrushMaterialInstanceMap.Empty();
@@ -201,9 +198,18 @@ public:
 			LandscapeInfo->GetComponentsInRegion(Bounds.Min.X, Bounds.Min.Y, Bounds.Max.X - 1, Bounds.Max.Y - 1, NewComponents);
 			bHasUnloadedComponents = LandscapeInfo->HasUnloadedComponentsInRegion(Bounds.Min.X, Bounds.Min.Y, Bounds.Max.X - 1, Bounds.Max.Y - 1);
 		}
+		
+		TSet<ULandscapeComponent*> PreviousComponents;
+		for (TWeakObjectPtr<ULandscapeComponent> BrushMaterialComponent : BrushMaterialComponents)
+		{
+			if (ULandscapeComponent* Component = BrushMaterialComponent.Get())
+			{
+				PreviousComponents.Add(Component);
+			}
+		}
 
 		// Remove the material from any old components that are no longer in the region
-		TSet<ULandscapeComponent*> RemovedComponents = ObjectPtrDecay(BrushMaterialComponents).Difference(NewComponents);
+		TSet<ULandscapeComponent*> RemovedComponents = PreviousComponents.Difference(NewComponents);
 		for (ULandscapeComponent* RemovedComponent : RemovedComponents)
 		{
 			BrushMaterialFreeInstances.Push(BrushMaterialInstanceMap.FindAndRemoveChecked(RemovedComponent));
@@ -213,15 +219,21 @@ public:
 		}
 
 		// Set brush material for components in new region
-		TSet<ULandscapeComponent*> AddedComponents = NewComponents.Difference(ObjectPtrDecay(BrushMaterialComponents));
+		TSet<ULandscapeComponent*> AddedComponents = NewComponents.Difference(PreviousComponents);
 		for (ULandscapeComponent* AddedComponent : AddedComponents)
 		{
 			UMaterialInstanceDynamic* BrushMaterialInstance = nullptr;
-			if (BrushMaterialFreeInstances.Num() > 0)
+
+			while (BrushMaterialFreeInstances.Num() > 0)
 			{
-				BrushMaterialInstance = BrushMaterialFreeInstances.Pop();
+				if (auto FreeBrushMaterialInstance = BrushMaterialFreeInstances.Pop(); FreeBrushMaterialInstance.IsValid())
+				{
+					BrushMaterialInstance = FreeBrushMaterialInstance.Get();
+					break;
+				}
 			}
-			else
+		
+			if (!BrushMaterialInstance)
 			{
 				BrushMaterialInstance = UMaterialInstanceDynamic::Create(BrushMaterial, nullptr);
 			}
@@ -230,7 +242,12 @@ public:
 			AddedComponent->UpdateEditToolRenderData();
 		}
 
-		BrushMaterialComponents = ObjectPtrWrap(MoveTemp(NewComponents));
+		BrushMaterialComponents.Empty();
+		for (ULandscapeComponent* BrushMaterialComponent : NewComponents)
+		{
+			BrushMaterialComponents.Add(TWeakObjectPtr<ULandscapeComponent>(BrushMaterialComponent));
+		}
+
 
 		// Set params for brush material.
 		FVector WorldLocation = Proxy->LandscapeActorToWorld().TransformPosition(FVector(LastMousePosition.X, LastMousePosition.Y, 0));
@@ -240,23 +257,27 @@ public:
 
 		for (const auto& BrushMaterialInstancePair : BrushMaterialInstanceMap)
 		{
-			ULandscapeComponent* const Component = BrushMaterialInstancePair.Key;
-			UMaterialInstanceDynamic* const MaterialInstance = BrushMaterialInstancePair.Value;
+			
+			ULandscapeComponent* const Component = BrushMaterialInstancePair.Key.Get();
+			UMaterialInstanceDynamic* const MaterialInstance = BrushMaterialInstancePair.Value.Get();
 
-				// Painting can cause the EditToolRenderData to be destructed, so update it if necessary
-			if (!AddedComponents.Contains(Component))
+			if (Component && MaterialInstance)
 			{
-				if (Component->EditToolRenderData.ToolMaterial == nullptr)
+				// Painting can cause the EditToolRenderData to be destructed, so update it if necessary
+				if (!AddedComponents.Contains(Component))
 				{
-					Component->EditToolRenderData.ToolMaterial = MaterialInstance;
-					Component->UpdateEditToolRenderData();
+					if (Component->EditToolRenderData.ToolMaterial == nullptr)
+					{
+						Component->EditToolRenderData.ToolMaterial = MaterialInstance;
+						Component->UpdateEditToolRenderData();
+					}
 				}
-			}
 
-			MaterialInstance->SetScalarParameterValue(FName(TEXT("LocalRadius")), Radius);
-			MaterialInstance->SetScalarParameterValue(FName(TEXT("LocalFalloff")), Falloff);
-			MaterialInstance->SetDoubleVectorParameterValue(FName(TEXT("WorldPosition")), FVector4(WorldLocation.X, WorldLocation.Y, WorldLocation.Z, ScaleXY));
-			MaterialInstance->SetScalarParameterValue("CanPaint", bCanPaint ? 1.0f : 0.0f);
+				MaterialInstance->SetScalarParameterValue(FName(TEXT("LocalRadius")), Radius);
+				MaterialInstance->SetScalarParameterValue(FName(TEXT("LocalFalloff")), Falloff);
+				MaterialInstance->SetDoubleVectorParameterValue(FName(TEXT("WorldPosition")), FVector4(WorldLocation.X, WorldLocation.Y, WorldLocation.Z, ScaleXY));
+				MaterialInstance->SetScalarParameterValue("CanPaint", bCanPaint ? 1.0f : 0.0f);
+			}
 		}
 	}
 
@@ -1292,7 +1313,12 @@ public:
 
 			for (const auto& BrushMaterialInstancePair : BrushMaterialInstanceMap)
 			{
-				UMaterialInstanceDynamic* const MaterialInstance = BrushMaterialInstancePair.Value;
+				UMaterialInstanceDynamic* const MaterialInstance = BrushMaterialInstancePair.Value.Get();
+				if (!MaterialInstance)
+				{
+					continue;
+				}
+				
 				MaterialInstance->SetVectorParameterValue(FName(TEXT("AlphaScaleBias")), AlphaScaleBias);
 				MaterialInstance->SetDoubleVectorParameterValue(FName(TEXT("LandscapeLocation")), LandscapeLocationParam);
 				MaterialInstance->SetVectorParameterValue(FName(TEXT("AlphaTextureMask")), AlphaTextureMask);
@@ -1531,7 +1557,13 @@ public:
 
 			for (const auto& BrushMaterialInstancePair : BrushMaterialInstanceMap)
 			{
-				UMaterialInstanceDynamic* const MaterialInstance = BrushMaterialInstancePair.Value;
+				UMaterialInstanceDynamic* const MaterialInstance = BrushMaterialInstancePair.Value.Get();
+
+				if (!MaterialInstance)
+				{
+					continue;
+				}
+				
 				MaterialInstance->SetVectorParameterValue(FName(TEXT("BrushScaleRot")), BrushScaleRot);
 				MaterialInstance->SetVectorParameterValue(FName(TEXT("AlphaTextureMask")), AlphaTextureMask);
 				MaterialInstance->SetTextureParameterValue(FName(TEXT("AlphaTexture")), EdMode->UISettings->AlphaTexture);
