Landscape : Testing - Added a simple import test.

diff --git a/Engine/Source/Editor/LandscapeEditor/Private/LandscapeEditorDetailCustomization_ImportExport.h b/Engine/Source/Editor/LandscapeEditor/Private/LandscapeEditorDetailCustomization_ImportExport.h
--- a/Engine/Source/Editor/LandscapeEditor/Private/LandscapeEditorDetailCustomization_ImportExport.h
+++ b/Engine/Source/Editor/LandscapeEditor/Private/LandscapeEditorDetailCustomization_ImportExport.h
@@ -18,8 +18,15 @@ class IDetailLayoutBuilder;
  * Slate widgets customizer for the "Import Export" tool
  */
 
+namespace LandscapeTestUtils
+{
+	struct LandscapeTestCommands;
+} //namespace LandscapeTestUtils
+
 class FLandscapeEditorDetailCustomization_ImportExport : public FLandscapeEditorDetailCustomization_Base
 {
+	// So tests have access to private members
+	friend struct LandscapeTestUtils::LandscapeTestCommands;
 public:
 	FLandscapeEditorDetailCustomization_ImportExport()
 	{}
diff --git a/Engine/Source/Editor/LandscapeEditor/Private/Tests/LandscapeAutomationTests.cpp b/Engine/Source/Editor/LandscapeEditor/Private/Tests/LandscapeAutomationTests.cpp
--- a/Engine/Source/Editor/LandscapeEditor/Private/Tests/LandscapeAutomationTests.cpp
+++ b/Engine/Source/Editor/LandscapeEditor/Private/Tests/LandscapeAutomationTests.cpp
@@ -1,6 +1,9 @@
 // Copyright Epic Games, Inc. All Rights Reserved.
 
 #include "CoreMinimal.h"
+#include "Engine/DirectionalLight.h"
+#include "GameFramework/Actor.h"
+#include "Engine/World.h"
 #include "Misc/Paths.h"
 #include "Misc/AutomationTest.h"
 #include "Misc/EngineVersion.h"
@@ -10,11 +13,13 @@
 #include "EditorModeManager.h"
 #include "EditorModes.h"
 #include "LandscapeEditorObject.h"
+#include "LandscapeInfoMap.h"
 
 #include "Tests/AutomationCommon.h"
 #include "Tests/AutomationEditorCommon.h"
 
 #include "LandscapeEditorDetailCustomization_NewLandscape.h"
+#include "LandscapeEditorDetailCustomization_ImportExport.h"
 
 #if WITH_DEV_AUTOMATION_TESTS
 
@@ -44,12 +49,35 @@ namespace LandscapeTestUtils
 
 		return SelectedViewport;
 	}
+
+	struct LandscapeTestCommands
+	{
+		static void Import(const FString& HeightMapFilenamee);
+	};
+
+	void LandscapeTestCommands::Import(const FString& HeightMapFilename)
+	{
+		//Switch to the Landscape tool
+		GLevelEditorModeTools().ActivateMode(FBuiltinEditorModes::EM_Landscape);
+		FEdModeLandscape* LandscapeEdMode = (FEdModeLandscape*)GLevelEditorModeTools().GetActiveMode(FBuiltinEditorModes::EM_Landscape);
+
+		TSharedPtr<FLandscapeEditorDetailCustomization_ImportExport> Customization_ImportExportLandscape = MakeShareable(new FLandscapeEditorDetailCustomization_ImportExport);
+
+		LandscapeEdMode->UISettings->ImportLandscape_HeightmapFilename = HeightMapFilename;
+		LandscapeEdMode->UISettings->bHeightmapSelected = true;
+		LandscapeEdMode->UISettings->ImportExportMode = ELandscapeImportExportMode::All;
+		LandscapeEdMode->UISettings->ImportType = ELandscapeImportTransformType::Resample;
+
+		Customization_ImportExportLandscape->OnOriginResetButtonClicked();
+
+		Customization_ImportExportLandscape->OnImportExportButtonClicked();
+	}
 }
 
 /**
 * Latent command to create a new landscape
 */
-DEFINE_LATENT_AUTOMATION_COMMAND(FCreateLandscapeCommand);
+DEFINE_LATENT_AUTOMATION_COMMAND_TWO_PARAMETER(FCreateLandscapeCommand,int32,NumComponents,int32, QuadsPerComponent);
 bool FCreateLandscapeCommand::Update()
 {
 	//Switch to the Landscape tool
@@ -57,7 +85,10 @@ bool FCreateLandscapeCommand::Update()
 	FEdModeLandscape* LandscapeEdMode = (FEdModeLandscape*)GLevelEditorModeTools().GetActiveMode(FBuiltinEditorModes::EM_Landscape);
 
 	//Modify the "Section size"
-	LandscapeEdMode->UISettings->NewLandscape_QuadsPerSection = 7;
+	LandscapeEdMode->UISettings->NewLandscape_QuadsPerSection = QuadsPerComponent;
+
+	LandscapeEdMode->UISettings->NewLandscape_ComponentCount.X = NumComponents;
+	LandscapeEdMode->UISettings->NewLandscape_ComponentCount.Y = NumComponents;
 	LandscapeEdMode->UISettings->NewLandscape_ClampSize();
 
 	//Create the landscape
@@ -76,6 +107,24 @@ bool FCreateLandscapeCommand::Update()
 	return true;
 }
 
+DEFINE_LATENT_AUTOMATION_COMMAND_ONE_PARAMETER(FImportLandscapeCommand, FString, HeightMapFilename);
+bool FImportLandscapeCommand::Update()
+{
+	LandscapeTestUtils::LandscapeTestCommands::Import(HeightMapFilename);
+	return true;
+}
+
+
+DEFINE_LATENT_AUTOMATION_COMMAND(FAddDirectionalLight);
+bool FAddDirectionalLight::Update()
+{
+	UWorld* World = GEditor->GetEditorWorldContext().World();
+	
+	ADirectionalLight*  DirectionalLight = World->SpawnActor<ADirectionalLight>();
+
+	return true;
+}
+
 /**
 * Latent command to start using the sculpting tool
 */
@@ -121,10 +170,36 @@ bool FEndModifyLandscapeCommand::Update()
 	return true;
 }
 
+DEFINE_LATENT_AUTOMATION_COMMAND(FCheckHeight);
+bool FCheckHeight::Update()
+{
+	FAutomationTestBase* CurrentTest = FAutomationTestFramework::Get().GetCurrentTest();
+	if (!CurrentTest)
+	{
+		return true;
+	}
+
+	UWorld* World = GEditor->GetEditorWorldContext().World();
+	if (!World)
+	{
+		return true;
+	}
+
+	
+	for (const auto& InfoMapPair : ULandscapeInfoMap::GetLandscapeInfoMap(World).Map)
+	{
+		TOptional<float> Height = InfoMapPair.Value->GetLandscapeProxy()->GetHeightAtLocation(FVector(0, 0, 0), EHeightfieldSource::Editor);
+		CurrentTest->TestEqual("Has Height Value at 0,0", Height.IsSet(), true);
+		CurrentTest->TestNearlyEqual("Height Value at 0,0 is 0", Height.GetValue(), 100.0f, 1e-4f);
+		break;
+	}
+	return true;
+}
+
 /**
 * Landscape creation / edit test
 */
-IMPLEMENT_SIMPLE_AUTOMATION_TEST(FLandscapeEditorTest, "System.Promotion.Editor.Landscape Editor", EAutomationTestFlags::EditorContext | EAutomationTestFlags::NonNullRHI | EAutomationTestFlags::EngineFilter);
+IMPLEMENT_SIMPLE_AUTOMATION_TEST(FLandscapeEditorTest, "System.Editor.Landscape CreateAndModify", EAutomationTestFlags::EditorContext | EAutomationTestFlags::NonNullRHI | EAutomationTestFlags::EngineFilter);
 bool FLandscapeEditorTest::RunTest(const FString& Parameters)
 {
 	//New level
@@ -139,7 +214,7 @@ bool FLandscapeEditorTest::RunTest(const FString& Parameters)
 		return false;
 	}
 
-	ADD_LATENT_AUTOMATION_COMMAND(FCreateLandscapeCommand());
+	ADD_LATENT_AUTOMATION_COMMAND(FCreateLandscapeCommand(/* NumComponents = */ 8, /* QuadsPerComponent = */ 7));
 
 	//For some reason the heightmap component takes a few ticks to register with the nav system.  We crash if we try to modify the heightmap before then.
 	ADD_LATENT_AUTOMATION_COMMAND(FWaitLatentCommand(1.0f));
@@ -149,6 +224,38 @@ bool FLandscapeEditorTest::RunTest(const FString& Parameters)
 	return true;
 }
 
+IMPLEMENT_SIMPLE_AUTOMATION_TEST(FLandscapeEditorTestFooBar, "System.Editor.Landscape Import", EAutomationTestFlags::EditorContext | EAutomationTestFlags::NonNullRHI | EAutomationTestFlags::EngineFilter);
+bool FLandscapeEditorTestFooBar::RunTest(const FString& Parameters)
+{
+	//New level
+	UWorld* NewMap = FAutomationEditorCommonUtils::CreateNewMap();
+	if (NewMap)
+	{
+		UE_LOG(LogLandscapeAutomationTests, Display, TEXT("Created an empty level"));
+	}
+	else
+	{
+		UE_LOG(LogLandscapeAutomationTests, Error, TEXT("Failed to create an empty level"));
+		return false;
+	}
+
+	ADD_LATENT_AUTOMATION_COMMAND(FCreateLandscapeCommand(/* NumComponents = */ 8, /* QuadsPerComponent = */ 63));
+
+	//For some reason the heightmap component takes a few ticks to register with the nav system.  We crash if we try to modify the heightmap before then.
+	ADD_LATENT_AUTOMATION_COMMAND(FWaitLatentCommand(1.0f));
+
+	FString HeightMapFilename = FPaths::Combine(FPaths::EngineContentDir(), FString(TEXT("FunctionalTesting\\height-505-flat.png")));
+	ADD_LATENT_AUTOMATION_COMMAND(FImportLandscapeCommand(HeightMapFilename));
+
+	ADD_LATENT_AUTOMATION_COMMAND(FAddDirectionalLight());
+
+	ADD_LATENT_AUTOMATION_COMMAND(FCheckHeight());
+
+	return true;
+}
+
+
+
 #undef LOCTEXT_NAMESPACE
 
 #endif //WITH_DEV_AUTOMATION_TESTS
