Landscape : Correct  automation test locations in hierarchy

#pf none

diff --git a/Engine/Source/Editor/LandscapeEditor/Private/Tests/LandscapeAutomationTests.cpp b/Engine/Source/Editor/LandscapeEditor/Private/Tests/LandscapeAutomationTests.cpp
--- a/Engine/Source/Editor/LandscapeEditor/Private/Tests/LandscapeAutomationTests.cpp
+++ b/Engine/Source/Editor/LandscapeEditor/Private/Tests/LandscapeAutomationTests.cpp
@@ -199,7 +199,7 @@ bool FCheckHeight::Update()
 /**
 * Landscape creation / edit test
 */
-IMPLEMENT_SIMPLE_AUTOMATION_TEST(FLandscapeEditorTest, "System.Editor.Landscape CreateAndModify", EAutomationTestFlags::EditorContext | EAutomationTestFlags::NonNullRHI | EAutomationTestFlags::EngineFilter);
+IMPLEMENT_SIMPLE_AUTOMATION_TEST(FLandscapeEditorTest, "Editor.Landscape CreateAndModify", EAutomationTestFlags::EditorContext | EAutomationTestFlags::NonNullRHI | EAutomationTestFlags::EngineFilter);
 bool FLandscapeEditorTest::RunTest(const FString& Parameters)
 {
 	//New level
@@ -224,7 +224,7 @@ bool FLandscapeEditorTest::RunTest(const FString& Parameters)
 	return true;
 }
 
-IMPLEMENT_SIMPLE_AUTOMATION_TEST(FLandscapeEditorTestFooBar, "System.Editor.Landscape Import", EAutomationTestFlags::EditorContext | EAutomationTestFlags::NonNullRHI | EAutomationTestFlags::EngineFilter);
+IMPLEMENT_SIMPLE_AUTOMATION_TEST(FLandscapeEditorTestFooBar, "Editor.Landscape Import", EAutomationTestFlags::EditorContext | EAutomationTestFlags::NonNullRHI | EAutomationTestFlags::EngineFilter);
 bool FLandscapeEditorTestFooBar::RunTest(const FString& Parameters)
 {
 	//New level
