Resize operation was not preserving all relevant landscape settings

Now uses CopySharedProperties to copy the bulk of the settings over.
Advantage is that this should automatically pick up any new settings properties we add in the future.

I added additional copies to handle the remaining settings.  In the future we might want to consider converting these to be shared properties as well.  Or if not shared, at least flag them as settings properties, so we can include them in the automatic copy here.

diff --git a/Engine/Source/Editor/LandscapeEditor/Private/LandscapeEdMode.cpp b/Engine/Source/Editor/LandscapeEditor/Private/LandscapeEdMode.cpp
--- a/Engine/Source/Editor/LandscapeEditor/Private/LandscapeEdMode.cpp
+++ b/Engine/Source/Editor/LandscapeEditor/Private/LandscapeEdMode.cpp
@@ -3676,33 +3676,38 @@ ALandscape* FEdModeLandscape::ChangeComponentSetting(int32 NumComponentsX, int32
 			SpawnParams.OverrideLevel = OldLandscape->GetLevel();
 			NewLandscape = OldLandscape->GetWorld()->SpawnActor<ALandscape>(Location, OldLandscape->GetActorRotation(), SpawnParams);
 			NewLandscape->bCanHaveLayersContent = OldLandscape->bCanHaveLayersContent;
+			NewLandscape->bAreNewLandscapeActorsSpatiallyLoaded = OldLandscape->bAreNewLandscapeActorsSpatiallyLoaded;
+			NewLandscape->bIncludeGridSizeInNameForLandscapeActors = OldLandscape->bIncludeGridSizeInNameForLandscapeActors;
+			NewLandscape->bUseGeneratedLandscapeSplineMeshesActors = OldLandscape->bUseGeneratedLandscapeSplineMeshesActors;
+
 			const FVector OldScale = OldLandscape->GetActorScale();
 			NewLandscape->SetActorRelativeScale3D(FVector(OldScale.X * LandscapeScaleFactor, OldScale.Y * LandscapeScaleFactor, OldScale.Z));
 
-			NewLandscape->LandscapeMaterial = OldLandscape->LandscapeMaterial;
-			NewLandscape->SetPerLODOverrideMaterials(OldLandscape->GetPerLODOverrideMaterials());
-			NewLandscape->CollisionMipLevel = OldLandscape->CollisionMipLevel;
-			NewLandscape->MaxLODLevel = OldLandscape->MaxLODLevel;
-			NewLandscape->LODDistanceFactor_DEPRECATED = OldLandscape->LODDistanceFactor_DEPRECATED;
-			NewLandscape->LODFalloff_DEPRECATED = OldLandscape->LODFalloff_DEPRECATED;
-			NewLandscape->ComponentScreenSizeToUseSubSections = OldLandscape->ComponentScreenSizeToUseSubSections;
-			NewLandscape->LODDistributionSetting = OldLandscape->LODDistributionSetting;
-			NewLandscape->LOD0DistributionSetting = OldLandscape->LOD0DistributionSetting;
+			// Copy all of the shared property settings over
+			NewLandscape->CopySharedProperties(OldLandscape);
+
+			// Double check things were copied
+			check(NewLandscape->GetPerLODOverrideMaterials().Num() == OldLandscape->GetPerLODOverrideMaterials().Num());
+			check(NewLandscape->RuntimeVirtualTextures.Num() == OldLandscape->RuntimeVirtualTextures.Num());
+			check(NewLandscape->LandscapeMaterial == OldLandscape->LandscapeMaterial);
+
+			// LandscapeGuid is stomped by CopySharedProperties, but original guid is not -- fix the mismatch or it will complain during Import
+ 			NewLandscape->SetLandscapeGuid(FGuid(), /* bValidateGuid= */ false);
+			
+			// Copy settings that are not copied by CopySharedProperties
 			NewLandscape->ExportLOD = OldLandscape->ExportLOD;
 			NewLandscape->StaticLightingLOD = OldLandscape->StaticLightingLOD;
-			NewLandscape->NegativeZBoundsExtension = OldLandscape->NegativeZBoundsExtension;
-			NewLandscape->PositiveZBoundsExtension = OldLandscape->PositiveZBoundsExtension;
-			NewLandscape->DefaultPhysMaterial = OldLandscape->DefaultPhysMaterial;
 			NewLandscape->StreamingDistanceMultiplier = OldLandscape->StreamingDistanceMultiplier;
-			NewLandscape->LandscapeHoleMaterial = OldLandscape->LandscapeHoleMaterial;
 			NewLandscape->StaticLightingResolution = OldLandscape->StaticLightingResolution;
-			NewLandscape->bCastStaticShadow = OldLandscape->bCastStaticShadow;
-			NewLandscape->bCastShadowAsTwoSided = OldLandscape->bCastShadowAsTwoSided;
-			NewLandscape->LightingChannels = OldLandscape->LightingChannels;
-			NewLandscape->bRenderCustomDepth = OldLandscape->bRenderCustomDepth;
-			NewLandscape->CustomDepthStencilWriteMask = OldLandscape->CustomDepthStencilWriteMask;
-			NewLandscape->CustomDepthStencilValue = OldLandscape->CustomDepthStencilValue;
-			NewLandscape->LightmassSettings = OldLandscape->LightmassSettings;
+			NewLandscape->ShadowCacheInvalidationBehavior = OldLandscape->ShadowCacheInvalidationBehavior;
+			NewLandscape->bUseMaterialPositionOffsetInStaticLighting = OldLandscape->bUseMaterialPositionOffsetInStaticLighting;
+			NewLandscape->bUseDynamicMaterialInstance = OldLandscape->bUseDynamicMaterialInstance;
+			NewLandscape->bGenerateOverlapEvents = OldLandscape->bGenerateOverlapEvents;
+			NewLandscape->bBakeMaterialPositionOffsetIntoCollision = OldLandscape->bBakeMaterialPositionOffsetIntoCollision;
+			NewLandscape->bFillCollisionUnderLandscapeForNavmesh = OldLandscape->bFillCollisionUnderLandscapeForNavmesh;
+			NewLandscape->NavigationGeometryGatheringMode = OldLandscape->NavigationGeometryGatheringMode;
+			NewLandscape->bUseLandscapeForCullingInvisibleHLODVertices = OldLandscape->bUseLandscapeForCullingInvisibleHLODVertices;
+
 			NewLandscape->BodyInstance.SetCollisionProfileName(OldLandscape->BodyInstance.GetCollisionProfileName());
 			if (NewLandscape->BodyInstance.DoesUseCollisionProfile() == false)
 			{
diff --git a/Engine/Source/Runtime/Landscape/Classes/LandscapeProxy.h b/Engine/Source/Runtime/Landscape/Classes/LandscapeProxy.h
--- a/Engine/Source/Runtime/Landscape/Classes/LandscapeProxy.h
+++ b/Engine/Source/Runtime/Landscape/Classes/LandscapeProxy.h
@@ -916,10 +916,10 @@ public:
 	 * landscape to exist simultaneously.  If you want the original (uninstanced) value, use GetOriginalLandscapeGuid().
 	 */
 	virtual FGuid GetLandscapeGuid() const override { return LandscapeGuid; }
-	void SetLandscapeGuid(const FGuid& Guid)
+	void SetLandscapeGuid(const FGuid& Guid, bool bValidateGuid = true)
 	{
 		// we probably shouldn't be setting the landscape guid on instanced landscapes
-		check((OriginalLandscapeGuid == LandscapeGuid) || !OriginalLandscapeGuid.IsValid());
+		check(!bValidateGuid || (OriginalLandscapeGuid == LandscapeGuid) || !OriginalLandscapeGuid.IsValid());
 		LandscapeGuid = Guid;
 		OriginalLandscapeGuid = Guid;
 	}
