Landscape : Don't create a Transaction when creating a landscape with regions

diff --git a/Engine/Source/Editor/LandscapeEditor/Private/LandscapeEditorDetailCustomization_NewLandscape.cpp b/Engine/Source/Editor/LandscapeEditor/Private/LandscapeEditorDetailCustomization_NewLandscape.cpp
--- a/Engine/Source/Editor/LandscapeEditor/Private/LandscapeEditorDetailCustomization_NewLandscape.cpp
+++ b/Engine/Source/Editor/LandscapeEditor/Private/LandscapeEditorDetailCustomization_NewLandscape.cpp
@@ -1070,7 +1070,12 @@ FReply FLandscapeEditorDetailCustomization_NewLandscape::OnCreateButtonClicked()
 	// ComputeHeightData will also modify/expand material layers data, which is why we create MaterialLayerDataPerLayers after calling ComputeHeightData
 	MaterialLayerDataPerLayers.Add(FGuid(), MoveTemp(MaterialImportLayers));
 
-	FScopedTransaction Transaction(bIsNewLandscape ? LOCTEXT("Undo_CreateLandscape", "Creating New Landscape") : LOCTEXT("Undo_CreateAndImportLandscape", "Creating and Importing New Landscape"));
+	TUniquePtr<FScopedTransaction> Transaction;
+	// If we're going to use regions then we're going to save and won't be able to undo this operation
+	if (!bNeedsLandscapeRegions)
+	{
+		Transaction = MakeUnique<FScopedTransaction>(bIsNewLandscape ? LOCTEXT("Undo_CreateLandscape", "Creating New Landscape") : LOCTEXT("Undo_CreateAndImportLandscape", "Creating and Importing New Landscape"));
+	}
 
 	const FVector Offset = FTransform(UISettings->NewLandscape_Rotation, FVector::ZeroVector, UISettings->NewLandscape_Scale).TransformVector(FVector(-UISettings->NewLandscape_ComponentCount.X * QuadsPerComponent / 2, -UISettings->NewLandscape_ComponentCount.Y * QuadsPerComponent / 2, 0));
 	
