Fix for crash exiting landscape editor import/export panel:

Though I only see crashes reported for ImportExport.cpp, I fixed a few similar issues in other detail panels.  The majority of the other detail panels already follow this convention of dynamically checking for null.

diff --git a/Engine/Source/Editor/LandscapeEditor/Private/LandscapeEditorDetailCustomization_CopyPaste.cpp b/Engine/Source/Editor/LandscapeEditor/Private/LandscapeEditorDetailCustomization_CopyPaste.cpp
--- a/Engine/Source/Editor/LandscapeEditor/Private/LandscapeEditorDetailCustomization_CopyPaste.cpp
+++ b/Engine/Source/Editor/LandscapeEditor/Private/LandscapeEditorDetailCustomization_CopyPaste.cpp
@@ -663,7 +663,10 @@ void FLandscapeEditorStructCustomization_FGizmoImportLayer::CustomizeChildren(TS
 FReply FLandscapeEditorStructCustomization_FGizmoImportLayer::OnGizmoImportLayerFilenameButtonClicked(TSharedRef<IPropertyHandle> PropertyHandle_LayerFilename)
 {
 	FEdModeLandscape* LandscapeEdMode = GetEditorMode();
-	check(LandscapeEdMode != NULL);
+	if (!LandscapeEdMode)
+	{
+		return FReply::Handled();
+	}
 
 	// Prompt the user for the Filenames
 	IDesktopPlatform* DesktopPlatform = FDesktopPlatformModule::Get();
diff --git a/Engine/Source/Editor/LandscapeEditor/Private/LandscapeEditorDetailCustomization_ImportExport.cpp b/Engine/Source/Editor/LandscapeEditor/Private/LandscapeEditorDetailCustomization_ImportExport.cpp
--- a/Engine/Source/Editor/LandscapeEditor/Private/LandscapeEditorDetailCustomization_ImportExport.cpp
+++ b/Engine/Source/Editor/LandscapeEditor/Private/LandscapeEditorDetailCustomization_ImportExport.cpp
@@ -300,24 +300,20 @@ END_SLATE_FUNCTION_BUILD_OPTIMIZATION
 bool FLandscapeEditorDetailCustomization_ImportExport::GetExportSingleFileIsEnabled()
 {
 	FEdModeLandscape* LandscapeEdMode = GetEditorMode();
-	check(LandscapeEdMode);
 	// The single file option only makes sense for landscapes where multiple file import is an option : 
-	return !LandscapeEdMode->UseSingleFileImport();
+	return LandscapeEdMode && !LandscapeEdMode->UseSingleFileImport();
 }
 
 ECheckBoxState FLandscapeEditorDetailCustomization_ImportExport::GetExportSingleFileCheckState()
 {
 	FEdModeLandscape* LandscapeEdMode = GetEditorMode();
-	check(LandscapeEdMode);
-	return (LandscapeEdMode->UseSingleFileImport() || (LandscapeEdMode->UISettings &&  LandscapeEdMode->UISettings->bExportSingleFile)) ? ECheckBoxState::Checked : ECheckBoxState::Unchecked;
+	return (LandscapeEdMode && (LandscapeEdMode->UseSingleFileImport() || (LandscapeEdMode->UISettings &&  LandscapeEdMode->UISettings->bExportSingleFile))) ? ECheckBoxState::Checked : ECheckBoxState::Unchecked;
 }
 
 void FLandscapeEditorDetailCustomization_ImportExport::OnExportSingleFileCheckStateChanged(ECheckBoxState NewCheckState)
 {
 	FEdModeLandscape* LandscapeEdMode = GetEditorMode();
-	check(LandscapeEdMode);
-
-	if (LandscapeEdMode->UISettings)
+	if (LandscapeEdMode && LandscapeEdMode->UISettings)
 	{
 		if (NewCheckState != ECheckBoxState::Undetermined)
 		{
@@ -334,8 +330,7 @@ EVisibility FLandscapeEditorDetailCustomization_ImportExport::GetImportExportVis
 bool FLandscapeEditorDetailCustomization_ImportExport::IsHeightmapEnabled()
 {
 	FEdModeLandscape* LandscapeEdMode = GetEditorMode();
-	check(LandscapeEdMode != nullptr);
-	if (LandscapeEdMode->UISettings)
+	if (LandscapeEdMode && LandscapeEdMode->UISettings)
 	{
 		return LandscapeEdMode->UISettings->bHeightmapSelected;
 	}
@@ -351,8 +346,7 @@ ECheckBoxState FLandscapeEditorDetailCustomization_ImportExport::GetHeightmapSel
 void FLandscapeEditorDetailCustomization_ImportExport::OnHeightmapSelectedCheckStateChanged(ECheckBoxState NewCheckedState)
 {
 	FEdModeLandscape* LandscapeEdMode = GetEditorMode();
-	check(LandscapeEdMode != nullptr);
-	if (LandscapeEdMode->UISettings)
+	if (LandscapeEdMode && LandscapeEdMode->UISettings)
 	{
 		if (NewCheckedState != ECheckBoxState::Undetermined)
 		{
@@ -364,8 +358,7 @@ void FLandscapeEditorDetailCustomization_ImportExport::OnHeightmapSelectedCheckS
 ECheckBoxState FLandscapeEditorDetailCustomization_ImportExport::ModeIsChecked(EImportExportMode Value)
 {
 	FEdModeLandscape* LandscapeEdMode = GetEditorMode();
-	check(LandscapeEdMode != nullptr);
-	return LandscapeEdMode->ImportExportMode == Value ? ECheckBoxState::Checked : ECheckBoxState::Unchecked;
+	return (LandscapeEdMode && (LandscapeEdMode->ImportExportMode == Value)) ? ECheckBoxState::Checked : ECheckBoxState::Unchecked;
 }
 
 void FLandscapeEditorDetailCustomization_ImportExport::OnModeChanged(ECheckBoxState NewCheckedState, EImportExportMode Value)
@@ -373,7 +366,11 @@ void FLandscapeEditorDetailCustomization_ImportExport::OnModeChanged(ECheckBoxSt
 	if (NewCheckedState == ECheckBoxState::Checked)
 	{
 		FEdModeLandscape* LandscapeEdMode = GetEditorMode();
-		check(LandscapeEdMode != nullptr);
+		if (!LandscapeEdMode)
+		{
+			return;
+		}
+
 		LandscapeEdMode->ImportExportMode = Value;
 
 		if (Value == EImportExportMode::Import)
@@ -454,7 +451,11 @@ void FLandscapeEditorDetailCustomization_ImportExport::SetFilename(const FText&
 FReply FLandscapeEditorDetailCustomization_ImportExport::OnBrowseFilenameButtonClicked(TSharedRef<IPropertyHandle> PropertyHandle_Filename)
 {
 	FEdModeLandscape* LandscapeEdMode = GetEditorMode();
-	check(LandscapeEdMode != nullptr);
+	if (!LandscapeEdMode)
+	{
+		return FReply::Handled();
+	}
+
 	ILandscapeEditorModule& LandscapeEditorModule = FModuleManager::GetModuleChecked<ILandscapeEditorModule>("LandscapeEditor");
 	const bool bIsImporting = IsImporting();
 	const FString DialogType = bIsImporting ? LandscapeEditorModule.GetHeightmapImportDialogTypeString() : LandscapeEditorModule.GetHeightmapExportDialogTypeString();
@@ -474,17 +475,17 @@ FReply FLandscapeEditorDetailCustomization_ImportExport::OnBrowseFilenameButtonC
 
 void FLandscapeEditorDetailCustomization_ImportExport::OnImportHeightmapFilenameChanged()
 {
-	FEdModeLandscape* LandscapeEdMode = GetEditorMode();
-	check(LandscapeEdMode != nullptr);
-	LandscapeEdMode->UISettings->OnImportHeightmapFilenameChanged();
+	if (FEdModeLandscape* LandscapeEdMode = GetEditorMode())
+	{
+		LandscapeEdMode->UISettings->OnImportHeightmapFilenameChanged();
+	}
 }
 
 void FLandscapeEditorDetailCustomization_ImportExport::FormatFilename(TSharedRef<IPropertyHandle> PropertyHandle_Filename, bool bForExport)
 {
 	FEdModeLandscape* LandscapeEdMode = GetEditorMode();
-	check(LandscapeEdMode != nullptr);
 	// No need to format the file name if single file import/export is enabled
-	if (!LandscapeEdMode->UseSingleFileImport() && (!bForExport || !LandscapeEdMode->UISettings->bExportSingleFile))
+	if (LandscapeEdMode && !LandscapeEdMode->UseSingleFileImport() && (!bForExport || !LandscapeEdMode->UISettings->bExportSingleFile))
 	{
 		// If selected file as the _xN_yM suffix remove it
 		FString	FilePath;
@@ -503,8 +504,7 @@ void FLandscapeEditorDetailCustomization_ImportExport::FormatFilename(TSharedRef
 bool FLandscapeEditorDetailCustomization_ImportExport::IsImporting()
 {
 	FEdModeLandscape* LandscapeEdMode = GetEditorMode();
-	check(LandscapeEdMode != nullptr);
-	return LandscapeEdMode->ImportExportMode == EImportExportMode::Import;
+	return LandscapeEdMode && (LandscapeEdMode->ImportExportMode == EImportExportMode::Import);
 }
 
 FText FLandscapeEditorDetailCustomization_ImportExport::GetImportExportButtonText()
@@ -519,11 +519,14 @@ FText FLandscapeEditorDetailCustomization_ImportExport::GetImportExportButtonTex
 
 FReply FLandscapeEditorDetailCustomization_ImportExport::OnOriginResetButtonClicked()
 {
+	FEdModeLandscape* LandscapeEdMode = GetEditorMode();	
+	if (!LandscapeEdMode)
+	{
+		return FReply::Handled();
+	}
+
 	FScopedTransaction Transaction(LOCTEXT("LandscapeEdResetImportLocation", "Reset Import Location"));
-	
-	FEdModeLandscape* LandscapeEdMode = GetEditorMode();
-	check(LandscapeEdMode);
-	
+		
 	TWeakObjectPtr<ULandscapeInfo> LandscapeInfo = LandscapeEdMode->CurrentToolTarget.LandscapeInfo;
 	TWeakObjectPtr<ALandscapeGizmoActiveActor> CurrentGizmoActor = LandscapeEdMode->CurrentGizmoActor;
 
@@ -542,8 +545,11 @@ FReply FLandscapeEditorDetailCustomization_ImportExport::OnOriginResetButtonClic
 
 FReply FLandscapeEditorDetailCustomization_ImportExport::OnImportExportButtonClicked()
 {
-	FEdModeLandscape* LandscapeEdMode = GetEditorMode();
-	check(LandscapeEdMode);
+	FEdModeLandscape* LandscapeEdMode = GetEditorMode();	
+	if (!LandscapeEdMode)
+	{
+		return FReply::Handled();
+	}
 
 	UWorld* World = LandscapeEdMode->GetWorld();
 
@@ -775,7 +781,10 @@ FReply FLandscapeEditorDetailCustomization_ImportExport::OnImportExportButtonCli
 bool FLandscapeEditorDetailCustomization_ImportExport::GetImportExportButtonIsEnabled()
 {
 	FEdModeLandscape* LandscapeEdMode = GetEditorMode();
-	check(LandscapeEdMode);
+	if (!LandscapeEdMode)
+	{
+		return false;
+	}
 
 	bool bHasOneSelection = false;
 	if (LandscapeEdMode->UISettings->bHeightmapSelected)
@@ -831,7 +840,7 @@ EVisibility FLandscapeEditorDetailCustomization_ImportExport::GetImportingVisibi
 EVisibility FLandscapeEditorDetailCustomization_ImportExport::GetImportExportLandscapeErrorVisibility() const
 {
 	FEdModeLandscape* EdMode = GetEditorMode();
-	
+
 	if (EdMode != nullptr)
 	{
 		return EdMode->IsLandscapeResolutionCompliant() ? EVisibility::Hidden : EVisibility::Visible;
@@ -854,16 +863,18 @@ FText FLandscapeEditorDetailCustomization_ImportExport::GetImportExportLandscape
 
 TSharedRef<SWidget> FLandscapeEditorDetailCustomization_ImportExport::GetImportLandscapeResolutionMenu()
 {
-	FMenuBuilder MenuBuilder(true, nullptr);
-
 	FEdModeLandscape* LandscapeEdMode = GetEditorMode();
-	check(LandscapeEdMode != nullptr);
-	for (int32 i = 0; i < LandscapeEdMode->UISettings->HeightmapImportDescriptor.ImportResolutions.Num(); i++)
+
+	FMenuBuilder MenuBuilder(true, nullptr);
+	if (LandscapeEdMode)
 	{
-		FFormatNamedArguments Args;
-		Args.Add(TEXT("Width"), LandscapeEdMode->UISettings->HeightmapImportDescriptor.ImportResolutions[i].Width);
-		Args.Add(TEXT("Height"), LandscapeEdMode->UISettings->HeightmapImportDescriptor.ImportResolutions[i].Height);
-		MenuBuilder.AddMenuEntry(FText::Format(LOCTEXT("ImportResolution_Format", "{Width}\u00D7{Height}"), Args), FText(), FSlateIcon(), FExecuteAction::CreateStatic(&FLandscapeEditorDetailCustomization_ImportExport::OnChangeImportLandscapeResolution, i));
+		for (int32 i = 0; i < LandscapeEdMode->UISettings->HeightmapImportDescriptor.ImportResolutions.Num(); i++)
+		{
+			FFormatNamedArguments Args;
+			Args.Add(TEXT("Width"), LandscapeEdMode->UISettings->HeightmapImportDescriptor.ImportResolutions[i].Width);
+			Args.Add(TEXT("Height"), LandscapeEdMode->UISettings->HeightmapImportDescriptor.ImportResolutions[i].Height);
+			MenuBuilder.AddMenuEntry(FText::Format(LOCTEXT("ImportResolution_Format", "{Width}\u00D7{Height}"), Args), FText(), FSlateIcon(), FExecuteAction::CreateStatic(&FLandscapeEditorDetailCustomization_ImportExport::OnChangeImportLandscapeResolution, i));
+		}
 	}
 	
 	return MenuBuilder.MakeWidget();
@@ -871,15 +882,20 @@ TSharedRef<SWidget> FLandscapeEditorDetailCustomization_ImportExport::GetImportL
 
 void FLandscapeEditorDetailCustomization_ImportExport::OnChangeImportLandscapeResolution(int32 Index)
 {
-	FEdModeLandscape* LandscapeEdMode = GetEditorMode();
-	check(LandscapeEdMode != nullptr);
-	LandscapeEdMode->UISettings->OnChangeImportLandscapeResolution(Index);
+	if (FEdModeLandscape* LandscapeEdMode = GetEditorMode())
+	{
+		LandscapeEdMode->UISettings->OnChangeImportLandscapeResolution(Index);
+	}
 }
 
 FText FLandscapeEditorDetailCustomization_ImportExport::GetImportLandscapeResolution()
 {
 	FEdModeLandscape* LandscapeEdMode = GetEditorMode();
-	check(LandscapeEdMode != nullptr);
+	if (!LandscapeEdMode)
+	{
+		return FText::GetEmpty();
+	}
+
 	const int32	Width = LandscapeEdMode->UISettings->ImportLandscape_Width;
 	const int32	Height = LandscapeEdMode->UISettings->ImportLandscape_Height;
 		
diff --git a/Engine/Source/Editor/LandscapeEditor/Private/LandscapeEditorDetailCustomization_ImportLayers.cpp b/Engine/Source/Editor/LandscapeEditor/Private/LandscapeEditorDetailCustomization_ImportLayers.cpp
--- a/Engine/Source/Editor/LandscapeEditor/Private/LandscapeEditorDetailCustomization_ImportLayers.cpp
+++ b/Engine/Source/Editor/LandscapeEditor/Private/LandscapeEditorDetailCustomization_ImportLayers.cpp
@@ -222,14 +222,19 @@ void FLandscapeEditorStructCustomization_FLandscapeImportLayer::CustomizeChildre
 void FLandscapeEditorStructCustomization_FLandscapeImportLayer::OnImportWeightmapFilenameChanged()
 {
 	FEdModeLandscape* LandscapeEdMode = GetEditorMode();
-	check(LandscapeEdMode != nullptr);
-	LandscapeEdMode->UISettings->OnImportWeightmapFilenameChanged();
+	if (LandscapeEdMode && LandscapeEdMode->UISettings)
+	{
+		LandscapeEdMode->UISettings->OnImportWeightmapFilenameChanged();
+	}
 }
 
 FReply FLandscapeEditorStructCustomization_FLandscapeImportLayer::OnLayerFilenameButtonClicked(TSharedRef<IPropertyHandle> PropertyHandle_LayerFilename)
 {
 	FEdModeLandscape* LandscapeEdMode = GetEditorMode();
-	check(LandscapeEdMode != nullptr);
+	if (!LandscapeEdMode)
+	{
+		return FReply::Handled();
+	}
 
 	ILandscapeEditorModule& LandscapeEditorModule = FModuleManager::GetModuleChecked<ILandscapeEditorModule>("LandscapeEditor");
 	const bool bIsImporting = IsImporting();
@@ -306,7 +311,10 @@ TSharedRef<SWidget> FLandscapeEditorStructCustomization_FLandscapeImportLayer::O
 void FLandscapeEditorStructCustomization_FLandscapeImportLayer::OnImportLayerCreateClicked(TSharedRef<IPropertyHandle> PropertyHandle_LayerInfo, FName LayerName, bool bNoWeightBlend)
 {
 	FEdModeLandscape* LandscapeEdMode = GetEditorMode();
-	check(LandscapeEdMode != nullptr);
+	if (!LandscapeEdMode)
+	{
+		return;
+	}
 	
 	// Hack as we don't have a direct world pointer in the EdMode...
 	ULevel* Level = LandscapeEdMode->CurrentGizmoActor->GetWorld()->GetCurrentLevel();
