Potential fix for rare crash in landscape circle brush

diff --git a/Engine/Source/Editor/LandscapeEditor/Private/LandscapeEdModeBrushes.cpp b/Engine/Source/Editor/LandscapeEditor/Private/LandscapeEdModeBrushes.cpp
--- a/Engine/Source/Editor/LandscapeEditor/Private/LandscapeEdModeBrushes.cpp
+++ b/Engine/Source/Editor/LandscapeEditor/Private/LandscapeEdModeBrushes.cpp
@@ -212,7 +212,11 @@ public:
 		TSet<ULandscapeComponent*> RemovedComponents = PreviousComponents.Difference(NewComponents);
 		for (ULandscapeComponent* RemovedComponent : RemovedComponents)
 		{
-			BrushMaterialFreeInstances.Push(BrushMaterialInstanceMap.FindAndRemoveChecked(RemovedComponent));
+			TWeakObjectPtr<UMaterialInstanceDynamic> RemovedMaterialInstance;
+			if (ensure(BrushMaterialInstanceMap.RemoveAndCopyValue(RemovedComponent, RemovedMaterialInstance)))
+			{
+				BrushMaterialFreeInstances.Push(RemovedMaterialInstance);
+			}
 
 			RemovedComponent->EditToolRenderData.ToolMaterial = nullptr;
 			RemovedComponent->UpdateEditToolRenderData();
