Make all landscape usage of TSet and TMaps with Weak Object Ptr Keys use the proper Funcs

diff --git a/Engine/Source/Editor/LandscapeEditor/Private/LandscapeEdModeBrushes.cpp b/Engine/Source/Editor/LandscapeEditor/Private/LandscapeEdModeBrushes.cpp
--- a/Engine/Source/Editor/LandscapeEditor/Private/LandscapeEdModeBrushes.cpp
+++ b/Engine/Source/Editor/LandscapeEditor/Private/LandscapeEdModeBrushes.cpp
@@ -49,7 +49,7 @@ class FLandscapeBrushCircle : public FLandscapeBrush
 	using Super = FLandscapeBrush;
 
 	// Components which previously we're under the area of the brush.
-	 TSet<TWeakObjectPtr<ULandscapeComponent>> BrushMaterialComponents;
+	TSet<TWeakObjectPtr<ULandscapeComponent>, TWeakObjectPtrSetKeyFuncs<TWeakObjectPtr<ULandscapeComponent>>> BrushMaterialComponents;
 	
 	// A Cache of previously created and now unused MIDs we can reuse if required.
 	TArray<TWeakObjectPtr<UMaterialInstanceDynamic>> BrushMaterialFreeInstances;
@@ -57,7 +57,8 @@ class FLandscapeBrushCircle : public FLandscapeBrush
 protected:
 	FVector2f LastMousePosition;
 	TObjectPtr<UMaterialInterface> BrushMaterial;
-	TMap<TWeakObjectPtr<ULandscapeComponent>, TWeakObjectPtr<UMaterialInstanceDynamic>> BrushMaterialInstanceMap;
+	TMap<TWeakObjectPtr<ULandscapeComponent>, TWeakObjectPtr<UMaterialInstanceDynamic>, FDefaultSetAllocator, TWeakObjectPtrMapKeyFuncs<TWeakObjectPtr<ULandscapeComponent>, TWeakObjectPtr<UMaterialInstanceDynamic>>> BrushMaterialInstanceMap;
+
 	bool bCanPaint;
 
 	virtual float CalculateFalloff(float Distance, float Radius, float Falloff) = 0;
@@ -208,12 +209,22 @@ public:
 			}
 		}
 
+		// Remove invalid keys from BrushMaterialInstanceMap (weak pointers can be invalidated)
+		for (auto It= BrushMaterialInstanceMap.CreateIterator(); It; ++It)
+		{
+			if (!It.Key().IsValid())
+			{
+				BrushMaterialFreeInstances.Push(It.Value());
+				It.RemoveCurrent();
+			}
+		}
+
 		// Remove the material from any old components that are no longer in the region
 		TSet<ULandscapeComponent*> RemovedComponents = PreviousComponents.Difference(NewComponents);
 		for (ULandscapeComponent* RemovedComponent : RemovedComponents)
 		{
 			TWeakObjectPtr<UMaterialInstanceDynamic> RemovedMaterialInstance;
-			if (ensure(BrushMaterialInstanceMap.RemoveAndCopyValue(RemovedComponent, RemovedMaterialInstance)))
+			if (BrushMaterialInstanceMap.RemoveAndCopyValue(RemovedComponent, RemovedMaterialInstance))
 			{
 				BrushMaterialFreeInstances.Push(RemovedMaterialInstance);
 			}
diff --git a/Engine/Source/Runtime/Landscape/Classes/LandscapeInfo.h b/Engine/Source/Runtime/Landscape/Classes/LandscapeInfo.h
--- a/Engine/Source/Runtime/Landscape/Classes/LandscapeInfo.h
+++ b/Engine/Source/Runtime/Landscape/Classes/LandscapeInfo.h
@@ -169,14 +169,14 @@ public:
 
 private:
 #if WITH_EDITORONLY_DATA
-	TSet<TWeakObjectPtr<UObject>> PackagesToDeprecate;
+	TSet<TWeakObjectPtr<UObject>, TWeakObjectPtrSetKeyFuncs<TWeakObjectPtr<UObject>>> PackagesToDeprecate;
 
 	// SORTED list of all actors implementing the spline interface that are registered with this landscape info
 	UPROPERTY()
 	TArray<TScriptInterface<ILandscapeSplineInterface>> SplineActors;
 
 	// Not a property since this shouldn't be modified through transactions (no undo/redo)
-	TSet<TWeakObjectPtr<UPackage>> ModifiedPackages;
+	TSet<TWeakObjectPtr<UPackage>, TWeakObjectPtrSetKeyFuncs<TWeakObjectPtr<UPackage>>> ModifiedPackages;
 
 	bool bDirtyOnlyInMode;
 
diff --git a/Engine/Source/Runtime/Landscape/Private/LandscapeGrass.cpp b/Engine/Source/Runtime/Landscape/Private/LandscapeGrass.cpp
--- a/Engine/Source/Runtime/Landscape/Private/LandscapeGrass.cpp
+++ b/Engine/Source/Runtime/Landscape/Private/LandscapeGrass.cpp
@@ -3073,7 +3073,7 @@ void ALandscapeProxy::GetGrassTypes(const UWorld* World, UMaterialInterface* Lan
 
 static uint32 GGrassExclusionChangeTag = 1;
 static uint32 GFrameNumberLastStaleCheck = 0;
-static TMap<FWeakObjectPtr, FBox> GGrassExclusionBoxes;
+static TMap<FWeakObjectPtr, FBox, FDefaultSetAllocator, TWeakObjectPtrMapKeyFuncs<FWeakObjectPtr, FBox>> GGrassExclusionBoxes;
 
 void ALandscapeProxy::AddExclusionBox(FWeakObjectPtr Owner, const FBox& BoxToRemove)
 {
