Landscape: Fix crash when creating multiple landscapes with regions

diff --git a/Engine/Source/Editor/LandscapeEditor/Private/LandscapeRegionUtils.cpp b/Engine/Source/Editor/LandscapeEditor/Private/LandscapeRegionUtils.cpp
--- a/Engine/Source/Editor/LandscapeEditor/Private/LandscapeRegionUtils.cpp
+++ b/Engine/Source/Editor/LandscapeEditor/Private/LandscapeRegionUtils.cpp
@@ -31,6 +31,7 @@ ALocationVolume* CreateLandscapeRegionVolume(UWorld* InWorld, ALandscapeProxy* I
 
 	const FString Label = FString::Printf(TEXT("LandscapeRegion_%i_%i"), InRegionCoordinate.X, InRegionCoordinate.Y);
 	SpawnParameters.Name = FName(*Label);
+	SpawnParameters.NameMode = FActorSpawnParameters::ESpawnActorNameMode::Requested;
 
 	ALocationVolume* LocationVolume = InWorld->SpawnActor<ALocationVolume>(Location, Rotation, SpawnParameters);
 	LocationVolume->SetActorLabel(Label);
