Add an option for landscape splines to always auto-rotate "forward" (i.e. don't reverse directions).  This only applies when there are exactly two spline segments connected to a control point; otherwise it falls back on the old behavior.

Original review in UE5/main branch: https://p4-swarm.epicgames.net/reviews/26543334

diff --git a/Engine/Source/Editor/LandscapeEditor/Private/LandscapeEdMode.h b/Engine/Source/Editor/LandscapeEditor/Private/LandscapeEdMode.h
--- a/Engine/Source/Editor/LandscapeEditor/Private/LandscapeEdMode.h
+++ b/Engine/Source/Editor/LandscapeEditor/Private/LandscapeEdMode.h
@@ -313,6 +313,8 @@ public:
 	virtual void UpdateSplineMeshLevels();
 	void SetbUseAutoRotateOnJoin(bool InbAutoRotateOnJoin);
 	bool GetbUseAutoRotateOnJoin();
+	void SetbAlwaysRotateForward(bool InbAlwaysRotateForward);
+	bool GetbAlwaysRotateForward();
 
 	// UI callbacks for ramp tool
 	void ApplyRampTool();
diff --git a/Engine/Source/Editor/LandscapeEditor/Private/LandscapeEdModeSplineTools.cpp b/Engine/Source/Editor/LandscapeEditor/Private/LandscapeEdModeSplineTools.cpp
--- a/Engine/Source/Editor/LandscapeEditor/Private/LandscapeEdModeSplineTools.cpp
+++ b/Engine/Source/Editor/LandscapeEditor/Private/LandscapeEdModeSplineTools.cpp
@@ -57,6 +57,7 @@ public:
 		, DraggingTangent_End(false)
 		, bMovingControlPoint(false)
 		, bAutoRotateOnJoin(true)
+		, bAlwaysRotateForward(false)
 		, bAutoChangeConnectionsOnMove(true)
 		, bDeleteLooseEnds(false)
 		, bCopyMeshToNewControlPoint(false)
@@ -372,13 +373,13 @@ public:
 		bool bUpdatedEnd = false;
 		if (bAutoRotateStart)
 		{
-			Start->AutoCalcRotation();
+			Start->AutoCalcRotation(bAlwaysRotateForward);
 			Start->UpdateSplinePoints();
 			bUpdatedStart = true;
 		}
 		if (bAutoRotateEnd)
 		{
-			End->AutoCalcRotation();
+			End->AutoCalcRotation(bAlwaysRotateForward);
 			End->UpdateSplinePoints();
 			bUpdatedEnd = true;
 		}
@@ -442,14 +443,14 @@ public:
 		bool bUpdatedEnd = false;
 		if (bAutoRotateStart)
 		{
-			Start->AutoCalcRotation();
+			Start->AutoCalcRotation(bAlwaysRotateForward);
 			Start->UpdateSplinePoints();
 			bUpdatedStart = true;
 		}
 
 		if (bAutoRotateEnd)
 		{
-			End->AutoCalcRotation();
+			End->AutoCalcRotation(bAlwaysRotateForward);
 			End->UpdateSplinePoints();
 			bUpdatedEnd = true;
 		}
@@ -1204,15 +1205,15 @@ public:
 
 				for (ULandscapeSplineControlPoint* ControlPoint : SelectedSplineControlPoints)
 				{
-					ControlPoint->AutoCalcRotation();
+					ControlPoint->AutoCalcRotation(bAlwaysRotateForward);
 					ControlPoint->UpdateSplinePoints();
 				}
 
 				for (ULandscapeSplineSegment* Segment : SelectedSplineSegments)
 				{
-					Segment->Connections[0].ControlPoint->AutoCalcRotation();
+					Segment->Connections[0].ControlPoint->AutoCalcRotation(bAlwaysRotateForward);
 					Segment->Connections[0].ControlPoint->UpdateSplinePoints();
-					Segment->Connections[1].ControlPoint->AutoCalcRotation();
+					Segment->Connections[1].ControlPoint->AutoCalcRotation(bAlwaysRotateForward);
 					Segment->Connections[1].ControlPoint->UpdateSplinePoints();
 				}
 
@@ -1515,7 +1516,7 @@ public:
 					{
 						if (bDuplicatingControlPoint && bAutoRotateOnJoin)
 						{
-							ControlPoint->AutoCalcRotation();
+							ControlPoint->AutoCalcRotation(bAlwaysRotateForward);
 						}
 
 						ControlPoint->UpdateSplinePoints(true);
@@ -2326,6 +2327,7 @@ protected:
 	uint32 bMovingControlPoint : 1;
 
 	uint32 bAutoRotateOnJoin : 1;
+	uint32 bAlwaysRotateForward : 1;
 	uint32 bAutoChangeConnectionsOnMove : 1;
 	uint32 bDeleteLooseEnds : 1;
 	uint32 bCopyMeshToNewControlPoint : 1;
@@ -2510,6 +2512,23 @@ bool FEdModeLandscape::GetbUseAutoRotateOnJoin()
 	return true; // default value
 }
 
+void FEdModeLandscape::SetbAlwaysRotateForward(bool InbAlwaysRotateForward)
+{
+	if (SplinesTool /*&& SplinesTool == CurrentTool*/)
+	{
+		SplinesTool->bAlwaysRotateForward = InbAlwaysRotateForward;
+	}
+}
+
+bool FEdModeLandscape::GetbAlwaysRotateForward()
+{
+	if (SplinesTool /*&& SplinesTool == CurrentTool*/)
+	{
+		return SplinesTool->bAlwaysRotateForward;
+	}
+	return false; // default value
+}
+
 void FEdModeLandscape::InitializeTool_Splines()
 {
 	auto Tool_Splines = MakeUnique<FLandscapeToolSplines>(this);
diff --git a/Engine/Source/Editor/LandscapeEditor/Private/LandscapeEditorDetailCustomization_MiscTools.cpp b/Engine/Source/Editor/LandscapeEditor/Private/LandscapeEditorDetailCustomization_MiscTools.cpp
--- a/Engine/Source/Editor/LandscapeEditor/Private/LandscapeEditorDetailCustomization_MiscTools.cpp
+++ b/Engine/Source/Editor/LandscapeEditor/Private/LandscapeEditorDetailCustomization_MiscTools.cpp
@@ -185,6 +185,21 @@ void FLandscapeEditorDetailCustomization_MiscTools::CustomizeDetails(IDetailLayo
 				]
 			]
 		];
+		ToolsCategory.AddCustomRow(LOCTEXT("Spline.bAlwaysForward.Selected", "Auto-Rotate Always Forward"))
+		[
+			SNew(SHorizontalBox)
+			+ SHorizontalBox::Slot()
+			.Padding(0, 6, 0, 0)
+			[
+				SNew(SCheckBox)
+				.OnCheckStateChanged(this, &FLandscapeEditorDetailCustomization_MiscTools::OnbAlwaysRotateForwardChanged)
+				.IsChecked(this, &FLandscapeEditorDetailCustomization_MiscTools::GetbAlwaysRotateForward)
+				.Content()
+				[
+					SNew(STextBlock).Text(LOCTEXT("Spline.bAlwaysForward.Selected", "Auto-Rotate Always Forward"))
+				]
+			]
+		];
 	}
 
 
@@ -416,6 +431,25 @@ ECheckBoxState FLandscapeEditorDetailCustomization_MiscTools::GetbUseAutoRotateC
 	return ECheckBoxState::Unchecked;
 }
 
+void FLandscapeEditorDetailCustomization_MiscTools::OnbAlwaysRotateForwardChanged(ECheckBoxState NewState)
+{
+	FEdModeLandscape* LandscapeEdMode = GetEditorMode();
+	if (LandscapeEdMode)
+	{
+		LandscapeEdMode->SetbAlwaysRotateForward(NewState == ECheckBoxState::Checked);
+	}
+}
+
+ECheckBoxState FLandscapeEditorDetailCustomization_MiscTools::GetbAlwaysRotateForward() const
+{
+	FEdModeLandscape* LandscapeEdMode = GetEditorMode();
+	if (LandscapeEdMode)
+	{
+		return LandscapeEdMode->GetbAlwaysRotateForward() ? ECheckBoxState::Checked : ECheckBoxState::Unchecked;
+	}
+	return ECheckBoxState::Unchecked;
+}
+
 FReply FLandscapeEditorDetailCustomization_MiscTools::OnApplyRampButtonClicked()
 {
 	FEdModeLandscape* LandscapeEdMode = GetEditorMode();
diff --git a/Engine/Source/Editor/LandscapeEditor/Private/LandscapeEditorDetailCustomization_MiscTools.h b/Engine/Source/Editor/LandscapeEditor/Private/LandscapeEditorDetailCustomization_MiscTools.h
--- a/Engine/Source/Editor/LandscapeEditor/Private/LandscapeEditorDetailCustomization_MiscTools.h
+++ b/Engine/Source/Editor/LandscapeEditor/Private/LandscapeEditorDetailCustomization_MiscTools.h
@@ -39,6 +39,8 @@ protected:
 	static FReply OnSelectAllSegmentsButtonClicked();
 	void OnbUseAutoRotateControlPointChanged(ECheckBoxState NewState);
 	ECheckBoxState GetbUseAutoRotateControlPoint() const;
+	void OnbAlwaysRotateForwardChanged(ECheckBoxState NewState);
+	ECheckBoxState GetbAlwaysRotateForward() const;
 
 	// Ramp Tool
 	static FReply OnApplyRampButtonClicked();
diff --git a/Engine/Source/Runtime/Landscape/Classes/LandscapeSplineControlPoint.h b/Engine/Source/Runtime/Landscape/Classes/LandscapeSplineControlPoint.h
--- a/Engine/Source/Runtime/Landscape/Classes/LandscapeSplineControlPoint.h
+++ b/Engine/Source/Runtime/Landscape/Classes/LandscapeSplineControlPoint.h
@@ -252,7 +252,7 @@ public:
 	virtual void SetSplineSelected(bool bInSelected);
 
 	/** Calculates rotation from connected segments */
-	virtual void AutoCalcRotation();
+	virtual void AutoCalcRotation(bool bAlwaysRotateForward);
 
 	/**  */
 	virtual void AutoFlipTangents();
diff --git a/Engine/Source/Runtime/Landscape/Private/LandscapeSplines.cpp b/Engine/Source/Runtime/Landscape/Private/LandscapeSplines.cpp
--- a/Engine/Source/Runtime/Landscape/Private/LandscapeSplines.cpp
+++ b/Engine/Source/Runtime/Landscape/Private/LandscapeSplines.cpp
@@ -1977,12 +1977,62 @@ void ULandscapeSplineControlPoint::SetSplineSelected(bool bInSelected)
 	}
 }
 
-void ULandscapeSplineControlPoint::AutoCalcRotation()
+void ULandscapeSplineControlPoint::AutoCalcRotation(bool bAlwaysRotateForward)
 {
 	Modify();
 
+	// always rotate forward only applies when there is exactly 2 segments connected
+	if (bAlwaysRotateForward && ConnectedSegments.Num() == 2)
+	{
+		const FLandscapeSplineSegmentConnection& Near0 = ConnectedSegments[0].GetNearConnection();
+		const FLandscapeSplineSegmentConnection& Far0 = ConnectedSegments[0].GetFarConnection();
+
+		// we modify the rotation so it is halfway between the two connected segments
+		// Get the start and end location/rotation of this connection
+		FVector StartLocation0; FRotator StartRotation0;
+		this->GetConnectionLocationAndRotation(Near0.SocketName, StartLocation0, StartRotation0);
+		FVector EndLocation0; FRotator EndRotation0;
+		Far0.ControlPoint->GetConnectionLocationAndRotation(Far0.SocketName, EndLocation0, EndRotation0);
+		const FVector DesiredDirection0 = (EndLocation0 - StartLocation0);
+
+		const FLandscapeSplineSegmentConnection& Near1 = ConnectedSegments[1].GetNearConnection();
+		const FLandscapeSplineSegmentConnection& Far1 = ConnectedSegments[1].GetFarConnection();
+		FVector StartLocation1; FRotator StartRotation1;
+		this->GetConnectionLocationAndRotation(Near1.SocketName, StartLocation1, StartRotation1);
+		FVector EndLocation1; FRotator EndRotation1;
+		Far1.ControlPoint->GetConnectionLocationAndRotation(Far1.SocketName, EndLocation1, EndRotation1);
+		const FVector DesiredDirection1 = (EndLocation1 - StartLocation1);
+
+		// compute orientations for incoming and outgoing segments from their direction
+		FRotator Rot0 = DesiredDirection0.Rotation();
+		FRotator Rot1 = (-DesiredDirection1).Rotation();
+
+		// determine the midpoint orientation via slerp, then convert back to rotator representation
+		FRotator DesiredRot = FQuat::Slerp(Rot0.Quaternion(), Rot1.Quaternion(), 0.5).Rotator();
+
+		// enforce pitch as the average of the incoming and outgoing orientations
+		DesiredRot.Pitch = (Rot0.Pitch + Rot1.Pitch) * 0.5;
+		
+		// Remove socket local rotation to calculate the Rotation setting
+		FVector StartLocalLocation; FRotator StartLocalRotation;
+		this->GetConnectionLocalLocationAndRotation(Near0.SocketName, StartLocalLocation, StartLocalRotation);
+		FQuat SocketLocalRotation = StartLocalRotation.Quaternion();
+		if (FMath::Sign(Near0.TangentLen) < 0)	// flip 180 if the tangent is inverted
+		{
+			SocketLocalRotation = SocketLocalRotation * FRotator(0, 180, 0).Quaternion();
+		}
+
+		FQuat LocalQuat = DesiredRot.Quaternion() * SocketLocalRotation.Inverse();
+		Rotation = LocalQuat.Rotator().GetNormalized();
+
+		AutoFlipTangents();
+
+		return;
+	}
+
 	FRotator Delta = FRotator::ZeroRotator;
 
+	// check all connections to this control point
 	for (const FLandscapeSplineConnection& Connection : ConnectedSegments)
 	{
 		// Get the start and end location/rotation of this connection
