Verify that creating a Water Body Actor creates a new Landscape Edit Layer called Water

diff --git a/Engine/Source/Editor/LandscapeEditor/LandscapeEditor.Build.cs b/Engine/Source/Editor/LandscapeEditor/LandscapeEditor.Build.cs
--- a/Engine/Source/Editor/LandscapeEditor/LandscapeEditor.Build.cs
+++ b/Engine/Source/Editor/LandscapeEditor/LandscapeEditor.Build.cs
@@ -31,7 +31,8 @@ public class LandscapeEditor : ModuleRules
 				"ToolWidgets",
 				"SourceControl",
 				"DirectoryWatcher",
-				"DeveloperSettings"
+				"DeveloperSettings",
+				"PlacementMode"
 			}
 			);
 
diff --git a/Engine/Source/Editor/LandscapeEditor/Private/Tests/Water/WaterInLandscapeEditLayerTests.cpp b/Engine/Source/Editor/LandscapeEditor/Private/Tests/Water/WaterInLandscapeEditLayerTests.cpp
new file mode 100644
--- /dev/null
+++ b/Engine/Source/Editor/LandscapeEditor/Private/Tests/Water/WaterInLandscapeEditLayerTests.cpp
@@ -0,0 +1,200 @@
+// Copyright Epic Games, Inc. All Rights Reserved.
+
+#include "CoreMinimal.h"
+#include "GameFramework/Actor.h"
+#include "Engine/World.h"
+#include "Misc/AutomationTest.h"
+#include "EditorModeManager.h"
+#include "EditorModes.h"
+#include "LandscapeEditorObject.h"
+#include "LandscapeInfoMap.h"
+#include "LandscapeEditorDetailCustomization_NewLandscape.h"
+#include "LandscapeEditorDetailCustomization_ImportExport.h"
+#include "IPlacementModeModule.h"
+#include "UObject/SoftObjectPtr.h"
+
+#if WITH_DEV_AUTOMATION_TESTS
+
+class FWaterInLandscapeEditLayerTestBase : public FAutomationTestBase
+{
+public:
+	FWaterInLandscapeEditLayerTestBase(const FString& InName, const bool bInComplexTask)
+		: FAutomationTestBase(InName, bInComplexTask)
+	{}
+	bool RunTestImpl(const FString& Parameters);
+private:
+	void Setup(const FString& Parameters);
+	bool PerformTest(const FString& Parameters);
+	void Teardown();
+
+	void CreateLandscape();
+	void SpawnWaterBodyActor(const FString& WaterBodyName);
+private:
+	ALandscape* Landscape { nullptr };
+	ALandscapeBlueprintBrushBase* WaterBrush { nullptr };
+	AActor* SpawnedWaterActor { nullptr };
+	UWorld* World { nullptr };
+	FEdModeLandscape* LandscapeEdMode { nullptr };
+protected:
+	const TArray<FString> WaterLandscapeNamesGetter
+	{
+		TEXT("RiverWaterInLandscapeEditLayer"),
+		TEXT("LakeWaterInLandscapeEditLayer"),
+		TEXT("OceanWaterInLandscapeEditLayer"),
+		TEXT("IslandWaterInLandscapeEditLayer")
+	};
+	const TArray<FString> WaterLandscapeCmdsGetter
+	{
+		TEXT("Water Body River"),
+		TEXT("Water Body Lake"),
+		TEXT("Water Body Ocean"),
+		TEXT("Water Body Island")
+	};
+};
+
+bool FWaterInLandscapeEditLayerTestBase::RunTestImpl(const FString& Parameters)
+{
+	Setup(Parameters);
+	const bool bSuccess = PerformTest(Parameters);
+	Teardown();
+	return bSuccess;
+}
+
+void FWaterInLandscapeEditLayerTestBase::Setup(const FString& Parameters)
+{
+	// Switching to Landscape Tool
+	GLevelEditorModeTools().ActivateMode(FBuiltinEditorModes::EM_Landscape);
+	LandscapeEdMode = (FEdModeLandscape*)GLevelEditorModeTools().GetActiveMode(FBuiltinEditorModes::EM_Landscape);
+
+	CreateLandscape();
+
+	World = LandscapeEdMode->GetWorld();
+	SpawnWaterBodyActor(Parameters);
+}
+
+bool FWaterInLandscapeEditLayerTestBase::PerformTest(const FString& Parameters)
+{
+	UTEST_NOT_NULL(TEXT("Landscape Editor"), LandscapeEdMode);
+	int32 LayerCount = LandscapeEdMode->GetLayerCount();
+	const FName WaterLayerName = TEXT("Water");
+	bool bWaterLayerFound = false;
+
+	// Going through Landscape's layers, if "Water" is found, tests are successful
+	for (int32 LayerIdx = 0; LayerIdx < LayerCount && !bWaterLayerFound; LayerIdx++)
+	{
+		bWaterLayerFound = (WaterLayerName == LandscapeEdMode->GetLayerName(LayerIdx));
+	}
+
+	TestTrue(TEXT("Water is not found in Landscape Edit Layers"), bWaterLayerFound);
+
+	return true;
+}
+
+void FWaterInLandscapeEditLayerTestBase::Teardown()
+{
+	// Removing Water Body Actor
+	if (SpawnedWaterActor)
+	{
+		World->EditorDestroyActor(SpawnedWaterActor, false);
+	}
+	// Removing Water Brush created for Landscape
+	if (LandscapeEdMode)
+	{
+		Landscape = LandscapeEdMode->GetLandscape();
+	}
+	if (Landscape)
+	{
+		for (ALandscapeBlueprintBrushBase* Brush : LandscapeEdMode->GetBrushList())
+		{
+			if (Brush->GetOwningLandscape() == Landscape)
+			{
+				WaterBrush = Brush;
+			}
+		}
+	}
+	if (WaterBrush)
+	{
+		Landscape->RemoveBrush(WaterBrush);
+		World->EditorDestroyActor(WaterBrush, false);
+	}
+	// Removing landscape
+	if (Landscape)
+	{
+		World->EditorDestroyActor(Landscape, false);
+	}
+}
+
+void FWaterInLandscapeEditLayerTestBase::CreateLandscape()
+{
+	if (TestNotNull(TEXT("Landscape Editor"), LandscapeEdMode))
+	{
+		// Setting up Section Size for the landscape
+		LandscapeEdMode->UISettings->NewLandscape_QuadsPerSection = 7;
+		LandscapeEdMode->UISettings->NewLandscape_ComponentCount.X = 8;
+		LandscapeEdMode->UISettings->NewLandscape_ComponentCount.Y = 8;
+		LandscapeEdMode->UISettings->NewLandscape_ClampSize();
+
+		// Creating Landscape
+		TSharedPtr<FLandscapeEditorDetailCustomization_NewLandscape> Customization_NewLandscape = MakeShareable(new FLandscapeEditorDetailCustomization_NewLandscape);
+		Customization_NewLandscape->OnCreateButtonClicked();
+	}
+}
+
+void FWaterInLandscapeEditLayerTestBase::SpawnWaterBodyActor(const FString& WaterBodyName)
+{
+	// Getting Water Body Actor's Class through Content Menu
+	const IPlacementModeModule& PlacementModeModule = IPlacementModeModule::Get();
+	// Water Body Actors are only accessible through All Classes
+	const FName PlacementModeCategoryHandle = TEXT("AllClasses");
+
+	TArray<TSharedPtr<FPlaceableItem>> OutItems;
+
+	// We need to manually load All Classes to search for a Water Body Actor
+	PlacementModeModule.Get().RegenerateItemsForCategory(PlacementModeCategoryHandle);
+
+	PlacementModeModule.GetItemsForCategory(PlacementModeCategoryHandle, OutItems);
+
+	// Searching for the Water Body Actor, so we can get its Object Path String
+	TSharedPtr<FPlaceableItem>* Item = OutItems.FindByPredicate([&WaterBodyName](const TSharedPtr<FPlaceableItem>& Item) {
+		return Item->DisplayName.ToString().Contains(WaterBodyName);
+		});
+
+	if (Item && World)
+	{
+		// Loading Water Body Actor class through Object Path String
+		const FString WaterBodyPathString = Item->Get()->AssetData.GetObjectPathString();
+		FSoftObjectPath ObjectPath(WaterBodyPathString);
+		UClass* ActorClass = Cast<UClass>(ObjectPath.TryLoad());
+
+		// Spawning Water Body Actor using the dynamically casted ActorClass
+		FActorSpawnParameters SpawnParams;
+		SpawnParams.SpawnCollisionHandlingOverride = ESpawnActorCollisionHandlingMethod::AdjustIfPossibleButAlwaysSpawn;
+
+		SpawnedWaterActor = World->SpawnActor<AActor>(ActorClass, FVector::ZeroVector, FRotator::ZeroRotator, SpawnParams);
+	}
+}
+
+// TestRail: 27425806
+// TestRail: 36679402
+// TestRail: 36679403
+// TestRail: 36679404
+// Verify that creating a Water Body Actor (River, Lake, Ocean, Island) creates a new Landscape Edit Layer called Water
+IMPLEMENT_CUSTOM_COMPLEX_AUTOMATION_TEST(FWaterInLandscapeEditLayerTest, FWaterInLandscapeEditLayerTestBase, "Editor.Landscape.Water", EAutomationTestFlags::EditorContext | EAutomationTestFlags::EngineFilter);
+bool FWaterInLandscapeEditLayerTest::RunTest(const FString& Parameters)
+{
+	return RunTestImpl(Parameters);
+}
+
+void FWaterInLandscapeEditLayerTest::GetTests(TArray<FString>& OutBeautifiedNames, TArray<FString>& OutTestCommands) const
+{
+	for (const FString& Test : WaterLandscapeNamesGetter)
+	{
+		OutBeautifiedNames.Add(Test);
+	}
+	for (const FString& TestCmd : WaterLandscapeCmdsGetter)
+	{
+		OutTestCommands.Add(TestCmd);
+	}
+}
+
+#endif
\ No newline at end of file
