Landscape : Fixed Landscape Import of .raw files

diff --git a/Engine/Source/Editor/LandscapeEditor/Private/LandscapeFileFormatRaw.cpp b/Engine/Source/Editor/LandscapeEditor/Private/LandscapeFileFormatRaw.cpp
--- a/Engine/Source/Editor/LandscapeEditor/Private/LandscapeFileFormatRaw.cpp
+++ b/Engine/Source/Editor/LandscapeEditor/Private/LandscapeFileFormatRaw.cpp
@@ -44,6 +44,8 @@ bool GetRawResolution(const TCHAR* InFilename, FLandscapeFileResolution& OutReso
 	FString Extension = FPaths::GetExtension(InFilename);
 	FString RawMetadataFilename = FPaths::SetExtension(InFilename, ".json");
 	
+	OutBitsPerPixel = 0;
+
 	if (FPaths::FileExists(RawMetadataFilename))
 	{
 		FString JsonStr;
@@ -128,12 +130,20 @@ FLandscapeFileInfo FLandscapeHeightmapFileFormat_Raw::Validate(const TCHAR* Heig
 			{
 				Result.PossibleResolutions = { Resolution };
 			}
-			
 		}
 		else
 		{
-			Result.ResultCode = ELandscapeImportResult::Error;
-			Result.ErrorMessage = LOCTEXT("Import_HeightmapFileInvalidSize", "The heightmap file has an invalid size (possibly not 16-bit?)");
+			if (BitsPerPixel == 0)
+			{
+				Result.ResultCode = ELandscapeImportResult::Error;
+				Result.ErrorMessage = LOCTEXT("Import_RawFileInvalidExtension", "The file bit depth unknown bit depth use .r16 (height) or .r8 (weight)");
+			}
+			else
+			{
+				Result.ResultCode = ELandscapeImportResult::Error;
+				Result.ErrorMessage = LOCTEXT("Import_HeightmapFileInvalidSize", "The heightmap file has an invalid size");
+			}
+			
 		}
 	}
 
@@ -223,8 +233,16 @@ FLandscapeFileInfo FLandscapeWeightmapFileFormat_Raw::Validate(const TCHAR* Weig
 		}
 		else
 		{
-			Result.ResultCode = ELandscapeImportResult::Error;
-			Result.ErrorMessage = LOCTEXT("Import_WeightmapFileInvalidSize", "The layer file has an invalid size");
+			if (BitsPerPixel == 0)
+			{
+				Result.ResultCode = ELandscapeImportResult::Error;
+				Result.ErrorMessage = LOCTEXT("Import_RawFileInvalidExtension", "The file bit depth unknown bit depth use .r16 (height) or .r8 (weight)");
+			}
+			else
+			{
+				Result.ResultCode = ELandscapeImportResult::Error;
+				Result.ErrorMessage = LOCTEXT("Import_WeightmapFileInvalidSize", "The layer file has an invalid size");
+			}			
 		}
 	}
 
