Landscape: fixed loctext with same key having different strings.

diff --git a/Engine/Source/Editor/LandscapeEditor/Private/LandscapeFileFormatRaw.cpp b/Engine/Source/Editor/LandscapeEditor/Private/LandscapeFileFormatRaw.cpp
--- a/Engine/Source/Editor/LandscapeEditor/Private/LandscapeFileFormatRaw.cpp
+++ b/Engine/Source/Editor/LandscapeEditor/Private/LandscapeFileFormatRaw.cpp
@@ -141,7 +141,7 @@ FLandscapeFileInfo FLandscapeHeightmapFileFormat_Raw::Validate(const TCHAR* Heig
 			else
 			{
 				Result.ResultCode = ELandscapeImportResult::Error;
-				Result.ErrorMessage = LOCTEXT("Import_HeightmapFileInvalidSize", "The heightmap file has an invalid size");
+				Result.ErrorMessage = LOCTEXT("Import_HeightmapFileInvalidSize", "The heightmap file has an invalid size (possibly not 16-bit?)");
 			}
 			
 		}
