In FLandscapeBrushCircle, moved cleanup of stale weak object ptrs entries for both BrushMaterialInstanceMap and BrushMaterialComponents from AddReferencedObjects to Tick.

The issue was caused by the cleanup of BrushMaterialInstanceMap which could remove entries from the map if the value (i.e. the material instance) got stale, while keeping the corresponding entry in the BrushMaterialComponents (since the landscape component itself was still perfectly fine), which led to an assert/crash in Tick, when we were trying to remove entries from the map, assuming there was one for each component in BrushMaterialComponents.

Now we first cleanup BrushMaterialInstanceMap from stale keys (i.e. null landscape component weak object ptrs) and then also cleanup stale entries in BrushMaterialFreeInstances. BrushMaterialComponents doesn't need to be cleaned up because of the way the code already works, at the end of the Tick it only contains valid components.

diff --git a/Engine/Source/Editor/LandscapeEditor/Private/LandscapeEdModeBrushes.cpp b/Engine/Source/Editor/LandscapeEditor/Private/LandscapeEdModeBrushes.cpp
--- a/Engine/Source/Editor/LandscapeEditor/Private/LandscapeEdModeBrushes.cpp
+++ b/Engine/Source/Editor/LandscapeEditor/Private/LandscapeEdModeBrushes.cpp
@@ -79,23 +79,6 @@ public:
 	virtual void AddReferencedObjects(FReferenceCollector& Collector) override
 	{
 		Collector.AddReferencedObject(BrushMaterial);
-
-		// If a user tool removes any components then we will have bad (null) entries in our TSet/TMap, remove them
-		// We can't just call .Remove(nullptr) because the entries were hashed as non-null values so a hash lookup of nullptr won't find them
-		for (auto It = BrushMaterialComponents.CreateIterator(); It; ++It)
-		{
-			if (*It == nullptr)
-			{
-				It.RemoveCurrent();
-			}
-		}
-		for (auto It = BrushMaterialInstanceMap.CreateIterator(); It; ++It)
-		{
-			if (It->Key == nullptr || It->Value == nullptr)
-			{
-				It.RemoveCurrent();
-			}
-		}
 	}
 
 	virtual void LeaveBrush() override
@@ -209,7 +192,8 @@ public:
 			}
 		}
 
-		// Remove invalid keys from BrushMaterialInstanceMap (weak pointers can be invalidated)
+		// We're dealing with weak object ptrs. Now is a good time to remove the null entries (weak pointers can be invalidated) from our containers if any : 
+		// First, remove all components (the map's keys) that might have become stale (the map's value are also weak object ptrs to materials but they will be handled right after) :
 		for (auto It= BrushMaterialInstanceMap.CreateIterator(); It; ++It)
 		{
 			if (!It.Key().IsValid())
@@ -219,13 +203,24 @@ public:
 			}
 		}
 
+		// Now remove all material instances that might have become stale :
+		for (auto It = BrushMaterialFreeInstances.CreateIterator(); It; ++It)
+		{
+			if (!It->IsValid())
+			{
+				It.RemoveCurrent();
+			}
+		}
+
 		// Remove the material from any old components that are no longer in the region
 		TSet<ULandscapeComponent*> RemovedComponents = PreviousComponents.Difference(NewComponents);
 		for (ULandscapeComponent* RemovedComponent : RemovedComponents)
 		{
 			TWeakObjectPtr<UMaterialInstanceDynamic> RemovedMaterialInstance;
-			if (BrushMaterialInstanceMap.RemoveAndCopyValue(RemovedComponent, RemovedMaterialInstance))
+			if (ensure(BrushMaterialInstanceMap.RemoveAndCopyValue(RemovedComponent, RemovedMaterialInstance)))
 			{
+				// We should have removed stale material instances already :
+				check(RemovedMaterialInstance.IsValid());
 				BrushMaterialFreeInstances.Push(RemovedMaterialInstance);
 			}
 
