Landscape : Better handling of 8 bit images for heights and 16 bit images for weights.

Allows for reading 8bit images for heights and 16bit images for weights.

diff --git a/Engine/Source/Editor/LandscapeEditor/Private/LandscapeImageFileCache.h b/Engine/Source/Editor/LandscapeEditor/Private/LandscapeImageFileCache.h
--- a/Engine/Source/Editor/LandscapeEditor/Private/LandscapeImageFileCache.h
+++ b/Engine/Source/Editor/LandscapeEditor/Private/LandscapeImageFileCache.h
@@ -24,6 +24,7 @@ struct FLandscapeImageDataRef
 	FIntPoint Resolution;
 	ELandscapeImportResult Result;
 	FText ErrorMessage;
+	int32 BytesPerPixel;
 };
 
 class FLandscapeImageFileCache
@@ -73,6 +74,7 @@ public:
 			NewImageData.Resolution = FIntPoint(ExpectedResolution.Width, ExpectedResolution.Height);
 			NewImageData.Result = FileInfo.ResultCode;
 			NewImageData.ErrorMessage = FileInfo.ErrorMessage;
+			NewImageData.BytesPerPixel = BufferSize / (ExpectedResolution.Width * ExpectedResolution.Height);
 		}
 		else
 		{
diff --git a/Engine/Source/Editor/LandscapeEditor/Private/LandscapeTiledImage.h b/Engine/Source/Editor/LandscapeEditor/Private/LandscapeTiledImage.h
--- a/Engine/Source/Editor/LandscapeEditor/Private/LandscapeTiledImage.h
+++ b/Engine/Source/Editor/LandscapeEditor/Private/LandscapeTiledImage.h
@@ -74,7 +74,7 @@ public:
 						return TileResult;
 					}
 
-					const T* TileData = (const T* ) ImageData.Data->GetData();
+					const uint8* TileData = ImageData.Data->GetData();
 
 					FIntRect TileRegion = Region;
 					TileRegion -= TileMin;
@@ -89,7 +89,21 @@ public:
 							
 							const int32 SrcY = bFlipYAxis ? TileResolution.Y - TileCoord.Y - 1 : TileCoord.Y;
 
-							OutImageData[OutCoord.Y * Region.Width() + OutCoord.X] = TileData[SrcY * TileResolution.X + TileCoord.X];
+							T& Dest = OutImageData[OutCoord.Y * Region.Width() + OutCoord.X];
+							const uint8* Src = &TileData[ImageData.BytesPerPixel * (SrcY * TileResolution.X + TileCoord.X)];
+
+							if (sizeof(T) == ImageData.BytesPerPixel)
+							{
+								Dest = *(reinterpret_cast<const T*>(Src));
+							}
+							else if (ImageData.BytesPerPixel == 2 && sizeof(T) == 1)
+							{
+								Dest = *(reinterpret_cast<const uint16*>(Src)) >> 8;
+							}
+							else
+							{
+								Dest = *Src;
+							}
 						}
 					}
 				}
