Landscape: Fix duplicate localization key

diff --git a/Engine/Source/Editor/LandscapeEditor/Private/LandscapeFileFormatRaw.cpp b/Engine/Source/Editor/LandscapeEditor/Private/LandscapeFileFormatRaw.cpp
--- a/Engine/Source/Editor/LandscapeEditor/Private/LandscapeFileFormatRaw.cpp
+++ b/Engine/Source/Editor/LandscapeEditor/Private/LandscapeFileFormatRaw.cpp
@@ -113,7 +113,7 @@ FLandscapeFileInfo FLandscapeHeightmapFileFormat_Raw::Validate(const TCHAR* Heig
 	else if (ImportFileSize == 0 || ImportFileSize % 2 != 0)
 	{
 		Result.ResultCode = ELandscapeImportResult::Error;
-		Result.ErrorMessage = LOCTEXT("Import_HeightmapFileInvalidSize", "The heightmap file has an invalid size (possibly not 16-bit?)");
+		Result.ErrorMessage = LOCTEXT("Import_HeightmapFileInvalidSizeNot16bit", "The heightmap file has an invalid size (possibly not 16-bit?)");
 	}
 	else
 	{
