Updates Landscape editor icons (BP Brush affects Heightmap/Weightmap/VisibilityLayer).
Changes buttons to toggle buttons for landscape editor BP brushes affecting layers features.
Updates landscape editor target list colors, padding, text color to reflect UE5 style.

diff --git a/Engine/Source/Editor/EditorStyle/Private/SlateEditorStyle.cpp b/Engine/Source/Editor/EditorStyle/Private/SlateEditorStyle.cpp
--- a/Engine/Source/Editor/EditorStyle/Private/SlateEditorStyle.cpp
+++ b/Engine/Source/Editor/EditorStyle/Private/SlateEditorStyle.cpp
@@ -6928,18 +6928,27 @@ void FSlateEditorStyle::FStyle::SetupLandscapeEditorStyle()
 		Set("LandscapeEditor.Brushes.Alpha.UseBChannel", new IMAGE_BRUSH("Icons/icon_Landscape_Brush_Alpha_UseBChannel_20x", Icon20x20));
 		Set("LandscapeEditor.Brushes.Alpha.UseAChannel", new IMAGE_BRUSH("Icons/icon_Landscape_Brush_Alpha_UseAChannel_20x", Icon20x20));
 
+		Set("LandscapeEditor.Brush.AffectsHeightmap", new IMAGE_BRUSH("Icons/icon_Landscape_Affects_Heightmap_16x", Icon16x16));
+		Set("LandscapeEditor.Brush.AffectsHeightmap.Disabled", new IMAGE_BRUSH("Icons/icon_Landscape_Affects_Heightmap_Disabled_16x", Icon16x16));
+
 		Set("LandscapeEditor.Brush.AffectsHeight.Enabled", new IMAGE_BRUSH("Icons/icon_Landscape_Affects_Height_Enabled_16x", Icon16x16));
 		Set("LandscapeEditor.Brush.AffectsHeight.Disabled", new IMAGE_BRUSH("Icons/icon_Landscape_Affects_Height_Disabled_16x", Icon16x16));
+		
+		Set("LandscapeEditor.Brush.AffectsWeightmap", new IMAGE_BRUSH("Icons/icon_Landscape_Affects_Weightmap_16x", Icon16x16));
+		Set("LandscapeEditor.Brush.AffectsWeightmap.Disabled", new IMAGE_BRUSH("Icons/icon_Landscape_Affects_Weightmap_Disabled_16x", Icon16x16));
+		
 		Set("LandscapeEditor.Brush.AffectsWeight.Enabled", new IMAGE_BRUSH("Icons/icon_Landscape_Affects_Weight_Enabled_16x", Icon16x16));
 		Set("LandscapeEditor.Brush.AffectsWeight.Disabled", new IMAGE_BRUSH("Icons/icon_Landscape_Affects_Weight_Disabled_16x", Icon16x16));
+		
+		Set("LandscapeEditor.Brush.AffectsVisibilityLayer", new IMAGE_BRUSH("Icons/icon_Landscape_Affects_VisibilityLayer_16x", Icon16x16));
 		Set("LandscapeEditor.Brush.AffectsVisibilityLayer.Enabled", new IMAGE_BRUSH("Icons/icon_Landscape_Affects_VisibilityLayer_Enabled_16x", Icon16x16));
 		Set("LandscapeEditor.Brush.AffectsVisibilityLayer.Disabled", new IMAGE_BRUSH("Icons/icon_Landscape_Affects_VisibilityLayer_Disabled_16x", Icon16x16));
 
 		// Target List
 		Set("LandscapeEditor.TargetList.RowBackground",        new FSlateNoResource());
-		Set("LandscapeEditor.TargetList.RowBackgroundHovered", new BOX_BRUSH("Common/RoundedSelection_16x", 4.0f / 16.0f, FLinearColor(1.0f, 1.0f, 1.0f, 0.1f)));
-		Set("LandscapeEditor.TargetList.RowSelected",          new BOX_BRUSH("Common/RoundedSelection_16x", 4.0f / 16.0f, SelectionColor_Pressed));
-		Set("LandscapeEditor.TargetList.RowSelectedHovered",   new BOX_BRUSH("Common/RoundedSelection_16x", 4.0f / 16.0f, SelectionColor));
+		Set("LandscapeEditor.TargetList.RowBackgroundHovered", new BOX_BRUSH("Common/RoundedSelection_16x", 4.0f / 16.0f, HighlightColor));
+		Set("LandscapeEditor.TargetList.RowSelected",          new BOX_BRUSH("Common/RoundedSelection_16x", 4.0f / 16.0f, SelectionColor));
+		Set("LandscapeEditor.TargetList.RowSelectedHovered",   new BOX_BRUSH("Common/RoundedSelection_16x", 4.0f / 16.0f, SelectionColor_Inactive));
 
 		Set("LandscapeEditor.Target_Heightmap",  new IMAGE_BRUSH("Icons/icon_Landscape_Target_Heightmap_48x",  Icon48x48));
 		Set("LandscapeEditor.Target_Visibility", new IMAGE_BRUSH("Icons/icon_Landscape_Target_Visibility_48x", Icon48x48));
diff --git a/Engine/Source/Editor/EditorStyle/Private/StarshipStyle.cpp b/Engine/Source/Editor/EditorStyle/Private/StarshipStyle.cpp
--- a/Engine/Source/Editor/EditorStyle/Private/StarshipStyle.cpp
+++ b/Engine/Source/Editor/EditorStyle/Private/StarshipStyle.cpp
@@ -6410,18 +6410,27 @@ void FStarshipEditorStyle::FStyle::SetupLandscapeEditorStyle()
 		Set("LandscapeEditor.Brushes.Alpha.UseBChannel", new IMAGE_BRUSH("Icons/icon_Landscape_Brush_Alpha_UseBChannel_20x", Icon20x20));
 		Set("LandscapeEditor.Brushes.Alpha.UseAChannel", new IMAGE_BRUSH("Icons/icon_Landscape_Brush_Alpha_UseAChannel_20x", Icon20x20));
 
+		Set("LandscapeEditor.Brush.AffectsHeightmap", new IMAGE_BRUSH("Icons/icon_Landscape_Affects_Heightmap_16x", Icon16x16));
+		Set("LandscapeEditor.Brush.AffectsHeightmap.Disabled", new IMAGE_BRUSH("Icons/icon_Landscape_Affects_Heightmap_Disabled_16x", Icon16x16));
+
 		Set("LandscapeEditor.Brush.AffectsHeight.Enabled", new IMAGE_BRUSH("Icons/icon_Landscape_Affects_Height_Enabled_16x", Icon16x16));
 		Set("LandscapeEditor.Brush.AffectsHeight.Disabled", new IMAGE_BRUSH("Icons/icon_Landscape_Affects_Height_Disabled_16x", Icon16x16));
+		
+		Set("LandscapeEditor.Brush.AffectsWeightmap", new IMAGE_BRUSH("Icons/icon_Landscape_Affects_Weightmap_16x", Icon16x16));
+		Set("LandscapeEditor.Brush.AffectsWeightmap.Disabled", new IMAGE_BRUSH("Icons/icon_Landscape_Affects_Weightmap_Disabled_16x", Icon16x16));
+
 		Set("LandscapeEditor.Brush.AffectsWeight.Enabled", new IMAGE_BRUSH("Icons/icon_Landscape_Affects_Weight_Enabled_16x", Icon16x16));
 		Set("LandscapeEditor.Brush.AffectsWeight.Disabled", new IMAGE_BRUSH("Icons/icon_Landscape_Affects_Weight_Disabled_16x", Icon16x16));
+		
+		Set("LandscapeEditor.Brush.AffectsVisibilityLayer", new IMAGE_BRUSH("Icons/icon_Landscape_Affects_VisibilityLayer_16x", Icon16x16));
 		Set("LandscapeEditor.Brush.AffectsVisibilityLayer.Enabled", new IMAGE_BRUSH("Icons/icon_Landscape_Affects_VisibilityLayer_Enabled_16x", Icon16x16));
 		Set("LandscapeEditor.Brush.AffectsVisibilityLayer.Disabled", new IMAGE_BRUSH("Icons/icon_Landscape_Affects_VisibilityLayer_Disabled_16x", Icon16x16));
 
 		// Target List
 		Set("LandscapeEditor.TargetList.RowBackground",        new FSlateNoResource());
-		Set("LandscapeEditor.TargetList.RowBackgroundHovered", new BOX_BRUSH("Common/RoundedSelection_16x", 4.0f / 16.0f, FLinearColor(1.0f, 1.0f, 1.0f, 0.1f)));
-		Set("LandscapeEditor.TargetList.RowSelected",          new BOX_BRUSH("Common/RoundedSelection_16x", 4.0f / 16.0f, SelectionColor_Pressed));
-		Set("LandscapeEditor.TargetList.RowSelectedHovered",   new BOX_BRUSH("Common/RoundedSelection_16x", 4.0f / 16.0f, SelectionColor));
+		Set("LandscapeEditor.TargetList.RowBackgroundHovered", new BOX_BRUSH("Common/RoundedSelection_16x", 4.0f / 16.0f, FStyleColors::SelectHover));
+		Set("LandscapeEditor.TargetList.RowSelected",          new BOX_BRUSH("Common/RoundedSelection_16x", 4.0f / 16.0f, FStyleColors::Select));
+		Set("LandscapeEditor.TargetList.RowSelectedHovered",   new BOX_BRUSH("Common/RoundedSelection_16x", 4.0f / 16.0f, FStyleColors::PrimaryHover));
 
 		Set("LandscapeEditor.Target_Heightmap",  new IMAGE_BRUSH("Icons/icon_Landscape_Target_Heightmap_48x",  Icon48x48));
 		Set("LandscapeEditor.Target_Visibility", new IMAGE_BRUSH("Icons/icon_Landscape_Target_Visibility_48x", Icon48x48));
diff --git a/Engine/Source/Editor/LandscapeEditor/Private/LandscapeEditorDetailCustomization_Layers.cpp b/Engine/Source/Editor/LandscapeEditor/Private/LandscapeEditorDetailCustomization_Layers.cpp
--- a/Engine/Source/Editor/LandscapeEditor/Private/LandscapeEditorDetailCustomization_Layers.cpp
+++ b/Engine/Source/Editor/LandscapeEditor/Private/LandscapeEditorDetailCustomization_Layers.cpp
@@ -175,7 +175,7 @@ TSharedPtr<SWidget> FLandscapeEditorCustomNodeBuilder_Layers::GenerateRow(int32
 {
 	FEdModeLandscape* LandscapeEdMode = GetEditorMode();
 	TSharedPtr<SWidget> RowWidget = SNew(SLandscapeEditorSelectableBorder)
-		.Padding(0.0f)
+		.Padding(FMargin(8.f, 0.f))
 		.VAlign(VAlign_Center)
 		.OnContextMenuOpening(this, &FLandscapeEditorCustomNodeBuilder_Layers::OnLayerContextMenuOpening, InLayerIndex)
 		.OnSelected(this, &FLandscapeEditorCustomNodeBuilder_Layers::OnLayerSelectionChanged, InLayerIndex)
@@ -200,6 +200,7 @@ TSharedPtr<SWidget> FLandscapeEditorCustomNodeBuilder_Layers::GenerateRow(int32
 
 			+SHorizontalBox::Slot()
 			.AutoWidth()
+			.Padding(4, 0)
 			.VAlign(VAlign_Center)
 			[
 				SNew(SButton)
@@ -212,7 +213,7 @@ TSharedPtr<SWidget> FLandscapeEditorCustomNodeBuilder_Layers::GenerateRow(int32
 				.Content()
 				[
 					SNew(SImage)
-			.Image(this, &FLandscapeEditorCustomNodeBuilder_Layers::GetVisibilityBrushForLayer, InLayerIndex)
+					.Image(this, &FLandscapeEditorCustomNodeBuilder_Layers::GetVisibilityBrushForLayer, InLayerIndex)
 				]
 			]
 
@@ -224,6 +225,7 @@ TSharedPtr<SWidget> FLandscapeEditorCustomNodeBuilder_Layers::GenerateRow(int32
 				SAssignNew(InlineTextBlocks[InLayerIndex], SInlineEditableTextBlock)
 				.IsEnabled(this, &FLandscapeEditorCustomNodeBuilder_Layers::IsLayerEditionEnabled, InLayerIndex)
 				.Text(this, &FLandscapeEditorCustomNodeBuilder_Layers::GetLayerDisplayName, InLayerIndex)
+				.ColorAndOpacity(TAttribute<FSlateColor>::Create(TAttribute<FSlateColor>::FGetter::CreateSP(this, &FLandscapeEditorCustomNodeBuilder_Layers::GetLayerTextColor, InLayerIndex)))
 				.ToolTipText(LOCTEXT("LandscapeLayers_tooltip", "Name of the Layer"))
 				.OnVerifyTextChanged(FOnVerifyTextChanged::CreateSP(this, &FLandscapeEditorCustomNodeBuilder_Layers::CanRenameLayerTo, InLayerIndex))
 				.OnEnterEditingMode(this, &FLandscapeEditorCustomNodeBuilder_Layers::OnBeginNameTextEdit)
@@ -248,6 +250,7 @@ TSharedPtr<SWidget> FLandscapeEditorCustomNodeBuilder_Layers::GenerateRow(int32
 					.IsEnabled(this, &FLandscapeEditorCustomNodeBuilder_Layers::IsLayerEditionEnabled, InLayerIndex)
 					.Visibility(this, &FLandscapeEditorCustomNodeBuilder_Layers::GetLayerAlphaVisibility, InLayerIndex)
 					.Text(LOCTEXT("LandscapeLayerAlpha", "Alpha"))
+					.ColorAndOpacity(TAttribute<FSlateColor>::Create(TAttribute<FSlateColor>::FGetter::CreateSP(this, &FLandscapeEditorCustomNodeBuilder_Layers::GetLayerTextColor, InLayerIndex)))
 				]
 				+ SHorizontalBox::Slot()
 				.VAlign(VAlign_Center)
@@ -301,7 +304,7 @@ FText FLandscapeEditorCustomNodeBuilder_Layers::GetLayerDisplayName(int32 InLaye
 	return FText::FromString(TEXT("None"));
 }
 
-bool FLandscapeEditorCustomNodeBuilder_Layers::IsLayerSelected(int32 InLayerIndex)
+bool FLandscapeEditorCustomNodeBuilder_Layers::IsLayerSelected(int32 InLayerIndex) const
 {
 	FEdModeLandscape* LandscapeEdMode = GetEditorMode();
 	if (LandscapeEdMode)
@@ -348,6 +351,10 @@ void FLandscapeEditorCustomNodeBuilder_Layers::SetLayerName(const FText& InText,
 	}
 }
 
+FSlateColor FLandscapeEditorCustomNodeBuilder_Layers::GetLayerTextColor(int32 InLayerIndex) const
+{
+	return IsLayerSelected(InLayerIndex) ? FStyleColors::ForegroundHover : FSlateColor::UseForeground();
+}
 
 void FLandscapeEditorCustomNodeBuilder_Layers::FillAddBrushMenu(FMenuBuilder& MenuBuilder, TArray<ALandscapeBlueprintBrushBase*> Brushes)
 {
diff --git a/Engine/Source/Editor/LandscapeEditor/Private/LandscapeEditorDetailCustomization_Layers.h b/Engine/Source/Editor/LandscapeEditor/Private/LandscapeEditorDetailCustomization_Layers.h
--- a/Engine/Source/Editor/LandscapeEditor/Private/LandscapeEditorDetailCustomization_Layers.h
+++ b/Engine/Source/Editor/LandscapeEditor/Private/LandscapeEditorDetailCustomization_Layers.h
@@ -73,7 +73,7 @@ protected:
 	TOptional<SDragAndDropVerticalBox::EItemDropZone> HandleCanAcceptDrop(const FDragDropEvent& DragDropEvent, SDragAndDropVerticalBox::EItemDropZone DropZone, SVerticalBox::FSlot* Slot);
 	FReply HandleAcceptDrop(FDragDropEvent const& DragDropEvent, SDragAndDropVerticalBox::EItemDropZone DropZone, int32 SlotIndex, SVerticalBox::FSlot* Slot);
 
-	bool IsLayerSelected(int32 LayerIndex);
+	bool IsLayerSelected(int32 LayerIndex) const;
 	void OnLayerSelectionChanged(int32 LayerIndex);
 	TSharedPtr<SWidget> OnLayerContextMenuOpening(int32 InLayerIndex);
 	void CreateLayer();
@@ -90,6 +90,7 @@ protected:
 	void OnEndNameTextEdit();
 	void SetLayerName(const FText& InText, ETextCommit::Type InCommitType, int32 InLayerIndex);
 	FText GetLayerText(int32 InLayerIndex) const;
+	FSlateColor GetLayerTextColor(int32 InLayerIndex) const;
 	void SetLandscapeSplinesReservedLayer(int32 InLayerIndex);
 	FText GetLayerDisplayName(int32 InLayerIndex) const;
 	bool IsLayerEditionEnabled(int32 InLayerIndex) const;
diff --git a/Engine/Source/Editor/LandscapeEditor/Private/LandscapeEditorDetailCustomization_LayersBrushStack.cpp b/Engine/Source/Editor/LandscapeEditor/Private/LandscapeEditorDetailCustomization_LayersBrushStack.cpp
--- a/Engine/Source/Editor/LandscapeEditor/Private/LandscapeEditorDetailCustomization_LayersBrushStack.cpp
+++ b/Engine/Source/Editor/LandscapeEditor/Private/LandscapeEditorDetailCustomization_LayersBrushStack.cpp
@@ -14,6 +14,7 @@
 #include "Framework/Application/SlateApplication.h"
 #include "Widgets/Images/SImage.h"
 #include "Widgets/Input/SButton.h"
+#include "Widgets/Input/SCheckBox.h"
 #include "EditorModeManager.h"
 #include "EditorModes.h"
 #include "DetailLayoutBuilder.h"
@@ -22,6 +23,7 @@
 #include "PropertyCustomizationHelpers.h"
 #include "LandscapeEditorDetailCustomization_Layers.h"
 #include "LandscapeEditTypes.h"
+#include "Styling/StyleColors.h"
 
 #include "ScopedTransaction.h"
 
@@ -125,7 +127,7 @@ TSharedPtr<SWidget> FLandscapeEditorCustomNodeBuilder_LayersBrushStack::Generate
 {
 	FEdModeLandscape* LandscapeEdMode = GetEditorMode();
 	TSharedPtr<SWidget> RowWidget = SNew(SLandscapeEditorSelectableBorder)
-		.Padding(0.0f)
+		.Padding(FMargin(8.0f, 0.f))
 		.VAlign(VAlign_Center)
 		.OnContextMenuOpening(this, &FLandscapeEditorCustomNodeBuilder_LayersBrushStack::OnBrushContextMenuOpening, InBrushIndex)
 		.OnSelected(this, &FLandscapeEditorCustomNodeBuilder_LayersBrushStack::OnBrushSelectionChanged, InBrushIndex)
@@ -170,53 +172,54 @@ TSharedPtr<SWidget> FLandscapeEditorCustomNodeBuilder_LayersBrushStack::Generate
 			.Padding(0, 0, 4, 0)
 			.VAlign(VAlign_Center)
 			[
-				SNew(SButton)
-				.ContentPadding(0.0f)
-				.ButtonStyle(FAppStyle::Get(), "NoBorder")
-				.OnClicked(this, &FLandscapeEditorCustomNodeBuilder_LayersBrushStack::OnToggleAffectsHeightmap, InBrushIndex)
+				SNew(SCheckBox)
+				.Style(FAppStyle::Get(), "ToggleButtonCheckbox")
+				.IsChecked_Lambda([InBrushIndex, this]() { return this->IsAffectingHeightmap(InBrushIndex) ? ECheckBoxState::Checked : ECheckBoxState::Unchecked; })
+				.OnCheckStateChanged_Raw(this, &FLandscapeEditorCustomNodeBuilder_LayersBrushStack::OnToggleAffectsHeightmap, InBrushIndex)
 				.ToolTipText(LOCTEXT("LandscapeBrushAffectsHeightmap", "Toggle Affects Heightmap"))
 				.HAlign(HAlign_Center)
-				.VAlign(VAlign_Center)
-				.Content()
+				.Padding(FMargin(5.f, 3.f))
 				[
 					SNew(SImage)
+					.ColorAndOpacity(FSlateColor::UseForeground())
 					.Image(this, &FLandscapeEditorCustomNodeBuilder_LayersBrushStack::GetAffectsHeightmapBrush, InBrushIndex)
 				]
 			]
 			+ SHorizontalBox::Slot()
 			.AutoWidth()
-            .VAlign(VAlign_Center)
+			.Padding(0, 0, 4, 0)
+			.VAlign(VAlign_Center)
             [
-				SNew(SButton)
-                .ContentPadding(0.0f)
-                .ButtonStyle(FAppStyle::Get(), "NoBorder")
-                .OnClicked(this, &FLandscapeEditorCustomNodeBuilder_LayersBrushStack::OnToggleAffectsWeightmap, InBrushIndex)
-                .ToolTipText(LOCTEXT("LandscapeBrushAffectsWeightmap", "Toggle Affects Weightmap"))
-                .HAlign(HAlign_Center)
-                .VAlign(VAlign_Center)
-                .Content()
-                [
-                    SNew(SImage)
-                    .Image(this, &FLandscapeEditorCustomNodeBuilder_LayersBrushStack::GetAffectsWeightmapBrush, InBrushIndex)
-                ]
+				SNew(SCheckBox)
+				.Style(FAppStyle::Get(), "ToggleButtonCheckbox")
+				.IsChecked_Lambda([InBrushIndex, this]() { return this->IsAffectingWeightmap(InBrushIndex) ? ECheckBoxState::Checked : ECheckBoxState::Unchecked; })
+				.OnCheckStateChanged_Raw(this, &FLandscapeEditorCustomNodeBuilder_LayersBrushStack::OnToggleAffectsWeightmap, InBrushIndex)
+				.ToolTipText(LOCTEXT("LandscapeBrushAffectsWeightmap", "Toggle Affects Weightmap"))
+				.HAlign(HAlign_Center)
+				.Padding(FMargin(5.f, 3.f))
+				[
+					SNew(SImage)
+					.ColorAndOpacity(FSlateColor::UseForeground())
+					.Image(this, &FLandscapeEditorCustomNodeBuilder_LayersBrushStack::GetAffectsWeightmapBrush, InBrushIndex)
+				]
             ]
 			+ SHorizontalBox::Slot()
-				.AutoWidth()
-				.VAlign(VAlign_Center)
+			.AutoWidth()
+			.VAlign(VAlign_Center)
+			[
+				SNew(SCheckBox)
+				.Style(FAppStyle::Get(), "ToggleButtonCheckbox")
+				.IsChecked_Lambda([InBrushIndex, this]() { return this->IsAffectingVisibilityLayer(InBrushIndex) ? ECheckBoxState::Checked : ECheckBoxState::Unchecked; })
+				.OnCheckStateChanged_Raw(this, &FLandscapeEditorCustomNodeBuilder_LayersBrushStack::OnToggleAffectsVisibilityLayer, InBrushIndex)
+				.ToolTipText(LOCTEXT("LandscapeBrushAffectsVisibilityLayer", "Toggle Affects Visibility Layer"))
+				.HAlign(HAlign_Center)
+				.Padding(FMargin(5.f, 3.f))
 				[
-					SNew(SButton)
-					.ContentPadding(0.0f)
-					.ButtonStyle(FAppStyle::Get(), "NoBorder")
-					.OnClicked(this, &FLandscapeEditorCustomNodeBuilder_LayersBrushStack::OnToggleAffectsVisibilityLayer, InBrushIndex)
-					.ToolTipText(LOCTEXT("LandscapeBrushAffectsVisibilityLayer", "Toggle Affects Visibility Layer"))
-					.HAlign(HAlign_Center)
-					.VAlign(VAlign_Center)
-					.Content()
-					[
-						SNew(SImage)
-						.Image(this, &FLandscapeEditorCustomNodeBuilder_LayersBrushStack::GetAffectsVisibilityLayerBrush, InBrushIndex)
-					]
+					SNew(SImage)
+					.ColorAndOpacity(FSlateColor::UseForeground())
+					.Image(this, &FLandscapeEditorCustomNodeBuilder_LayersBrushStack::GetAffectsVisibilityLayerBrush, InBrushIndex)
 				]
+			]
 		];
 	
 	return RowWidget;
@@ -309,36 +312,46 @@ void FLandscapeEditorCustomNodeBuilder_LayersBrushStack::FillAddBrushMenu(FMenuB
 	}
 }
 
-FReply FLandscapeEditorCustomNodeBuilder_LayersBrushStack::OnToggleAffectsHeightmap(int32 InBrushIndex)
+void FLandscapeEditorCustomNodeBuilder_LayersBrushStack::OnToggleAffectsHeightmap(ECheckBoxState InCheckBoxState, int32 InBrushIndex)
 {
+	if ((InCheckBoxState != ECheckBoxState::Checked) && (InCheckBoxState != ECheckBoxState::Unchecked))
+	{
+		return;
+	}
+
 	if (ALandscapeBlueprintBrushBase* Brush = GetBrush(InBrushIndex))
 	{
 		const FScopedTransaction Transaction(LOCTEXT("Landscape_Brush_AffectsHeightmap", "Set Brush Can Affect Heightmap"));
-		bool bCanAffectHeightmap = Brush->CanAffectHeightmap();
-		Brush->SetCanAffectHeightmap(!bCanAffectHeightmap);
+		Brush->SetCanAffectHeightmap(InCheckBoxState == ECheckBoxState::Checked);
 	}
-	return FReply::Handled();
 }
 
-FReply FLandscapeEditorCustomNodeBuilder_LayersBrushStack::OnToggleAffectsWeightmap(int32 InBrushIndex)
+void FLandscapeEditorCustomNodeBuilder_LayersBrushStack::OnToggleAffectsWeightmap(ECheckBoxState InCheckBoxState, int32 InBrushIndex)
 {
+	if ((InCheckBoxState != ECheckBoxState::Checked) && (InCheckBoxState != ECheckBoxState::Unchecked))
+	{
+		return;
+	}
+
 	if (ALandscapeBlueprintBrushBase* Brush = GetBrush(InBrushIndex))
 	{
 		const FScopedTransaction Transaction(LOCTEXT("Landscape_Brush_AffectsWeightmap", "Set Brush Can Affect Weightmap"));
-		bool bCanAffectWeightmap = Brush->CanAffectWeightmap();
-		Brush->SetCanAffectWeightmap(!bCanAffectWeightmap);
+		Brush->SetCanAffectWeightmap(InCheckBoxState == ECheckBoxState::Checked);
 	}
-	return FReply::Handled();
 }
-FReply FLandscapeEditorCustomNodeBuilder_LayersBrushStack::OnToggleAffectsVisibilityLayer(int32 InBrushIndex)
+
+void FLandscapeEditorCustomNodeBuilder_LayersBrushStack::OnToggleAffectsVisibilityLayer(ECheckBoxState InCheckBoxState, int32 InBrushIndex)
 {
+	if ((InCheckBoxState != ECheckBoxState::Checked) && (InCheckBoxState != ECheckBoxState::Unchecked))
+	{
+		return;
+	}
+
 	if (ALandscapeBlueprintBrushBase* Brush = GetBrush(InBrushIndex))
 	{
 		const FScopedTransaction Transaction(LOCTEXT("Landscape_Brush_AffectsVisibilityLayer", "Set Brush Can Affect Visibility Layer"));
-		bool bCanAffectVisibilityLayer = Brush->CanAffectVisibilityLayer();
-		Brush->SetCanAffectVisibilityLayer(!bCanAffectVisibilityLayer);
+		Brush->SetCanAffectVisibilityLayer(InCheckBoxState == ECheckBoxState::Checked);
 	}
-	return FReply::Handled();
 }
 
 FReply FLandscapeEditorCustomNodeBuilder_LayersBrushStack::OnToggleVisibility(int32 InBrushIndex)
@@ -358,21 +371,39 @@ void FLandscapeEditorCustomNodeBuilder_LayersBrushStack::OnToggleVisibility(ALan
 }
 
 const FSlateBrush* FLandscapeEditorCustomNodeBuilder_LayersBrushStack::GetAffectsWeightmapBrush(int32 InBrushIndex) const
+{
+	return IsAffectingWeightmap(InBrushIndex) ? FAppStyle::GetBrush("LandscapeEditor.Brush.AffectsWeightmap") : FAppStyle::GetBrush("LandscapeEditor.Brush.AffectsWeightmap.Disabled");
+}
+
+bool FLandscapeEditorCustomNodeBuilder_LayersBrushStack::IsAffectingWeightmap(int32 InBrushIndex) const
 {
 	ALandscapeBlueprintBrushBase* Brush = GetBrush(InBrushIndex);
-	return Brush && Brush->CanAffectWeightmap() ? FAppStyle::GetBrush("LandscapeEditor.Brush.AffectsWeight.Enabled") : FAppStyle::GetBrush("LandscapeEditor.Brush.AffectsWeight.Disabled");
+
+	return Brush && Brush->CanAffectWeightmap();
 }
 
 const FSlateBrush* FLandscapeEditorCustomNodeBuilder_LayersBrushStack::GetAffectsHeightmapBrush(int32 InBrushIndex) const
+{
+	return IsAffectingHeightmap(InBrushIndex) ? FAppStyle::GetBrush("LandscapeEditor.Brush.AffectsHeightmap") : FAppStyle::GetBrush("LandscapeEditor.Brush.AffectsHeightmap.Disabled");
+}
+
+bool FLandscapeEditorCustomNodeBuilder_LayersBrushStack::IsAffectingHeightmap(int32 InBrushIndex) const
 {
 	ALandscapeBlueprintBrushBase* Brush = GetBrush(InBrushIndex);
-	return Brush && Brush->CanAffectHeightmap() ? FAppStyle::GetBrush("LandscapeEditor.Brush.AffectsHeight.Enabled") : FAppStyle::GetBrush("LandscapeEditor.Brush.AffectsHeight.Disabled");
+
+	return Brush && Brush->CanAffectHeightmap();
 }
 
 const FSlateBrush* FLandscapeEditorCustomNodeBuilder_LayersBrushStack::GetAffectsVisibilityLayerBrush(int32 InBrushIndex) const
+{
+	return IsAffectingVisibilityLayer(InBrushIndex) ? FAppStyle::GetBrush("LandscapeEditor.Brush.AffectsVisibilityLayer") : FAppStyle::GetBrush("LandscapeEditor.Brush.AffectsVisibilityLayer.Disabled");
+}
+
+bool FLandscapeEditorCustomNodeBuilder_LayersBrushStack::IsAffectingVisibilityLayer(int32 InBrushIndex) const
 {
 	ALandscapeBlueprintBrushBase* Brush = GetBrush(InBrushIndex);
-	return Brush && Brush->CanAffectVisibilityLayer() ? FAppStyle::GetBrush("LandscapeEditor.Brush.AffectsVisibilityLayer.Enabled") : FAppStyle::GetBrush("LandscapeEditor.Brush.AffectsVisibilityLayer.Disabled");
+
+	return Brush && Brush->CanAffectVisibilityLayer();
 }
 
 const FSlateBrush* FLandscapeEditorCustomNodeBuilder_LayersBrushStack::GetVisibilityBrush(int32 InBrushIndex) const
@@ -435,12 +466,14 @@ FText FLandscapeEditorCustomNodeBuilder_LayersBrushStack::GetBrushText(int32 InB
 
 FSlateColor FLandscapeEditorCustomNodeBuilder_LayersBrushStack::GetBrushTextColor(int32 InBrushIndex) const
 {
-	if (ALandscapeBlueprintBrushBase* Brush = GetBrush(InBrushIndex))
+	ALandscapeBlueprintBrushBase* Brush = GetBrush(InBrushIndex);
+
+	if (Brush == nullptr)
 	{
-		return FSlateColor::UseForeground();
+		return FSlateColor::UseSubduedForeground();
 	}
 
-	return FSlateColor::UseSubduedForeground();
+	return IsBrushSelected(InBrushIndex) ? FStyleColors::ForegroundHover : FSlateColor::UseForeground();
 }
 
 ALandscapeBlueprintBrushBase* FLandscapeEditorCustomNodeBuilder_LayersBrushStack::GetBrush(int32 InBrushIndex) const
diff --git a/Engine/Source/Editor/LandscapeEditor/Private/LandscapeEditorDetailCustomization_LayersBrushStack.h b/Engine/Source/Editor/LandscapeEditor/Private/LandscapeEditorDetailCustomization_LayersBrushStack.h
--- a/Engine/Source/Editor/LandscapeEditor/Private/LandscapeEditorDetailCustomization_LayersBrushStack.h
+++ b/Engine/Source/Editor/LandscapeEditor/Private/LandscapeEditorDetailCustomization_LayersBrushStack.h
@@ -63,16 +63,23 @@ protected:
 	void FillAddBrushMenu(FMenuBuilder& MenuBuilder, TArray<ALandscapeBlueprintBrushBase*> Brushes);
 
 	FReply OnToggleVisibility(int32 InBrushIndex);
-	FReply OnToggleAffectsHeightmap(int32 InBrushIndex);
-	FReply OnToggleAffectsWeightmap(int32 InBrushIndex);
-	FReply OnToggleAffectsVisibilityLayer(int32 InBrushIndex);
+	void OnToggleAffectsHeightmap(ECheckBoxState InCheckBoxState, int32 InBrushIndex);
+	void OnToggleAffectsWeightmap(ECheckBoxState InCheckBoxState, int32 InBrushIndex);
+	void OnToggleAffectsVisibilityLayer(ECheckBoxState InCheckBoxState, int32 InBrushIndex);
 
 	void OnToggleVisibility(ALandscapeBlueprintBrushBase* Brush);
 
 	const FSlateBrush* GetAffectsHeightmapBrush(int32 InBrushIndex) const;
+	bool IsAffectingHeightmap(int32 InBrushIndex) const;
+
 	const FSlateBrush* GetAffectsWeightmapBrush(int32 InBrushIndex) const;
+	bool IsAffectingWeightmap(int32 InBrushIndex) const;
+
 	const FSlateBrush* GetAffectsVisibilityLayerBrush(int32 InBrushIndex) const;
+	bool IsAffectingVisibilityLayer(int32 InBrushIndex) const;
+
 	const FSlateBrush* GetVisibilityBrush(int32 InBrushIndex) const;
+	
 	bool IsBrushSelected(int32 InBrushIndex) const;
 	bool IsBrushEnabled(int32 InBrushIndex) const;
 	void OnBrushSelectionChanged(int32 InBrushIndex);
diff --git a/Engine/Source/Editor/LandscapeEditor/Private/LandscapeEditorDetailCustomization_TargetLayers.cpp b/Engine/Source/Editor/LandscapeEditor/Private/LandscapeEditorDetailCustomization_TargetLayers.cpp
--- a/Engine/Source/Editor/LandscapeEditor/Private/LandscapeEditorDetailCustomization_TargetLayers.cpp
+++ b/Engine/Source/Editor/LandscapeEditor/Private/LandscapeEditorDetailCustomization_TargetLayers.cpp
@@ -537,7 +537,7 @@ TSharedPtr<SWidget> FLandscapeEditorCustomNodeBuilder_TargetLayers::GenerateRow(
 						SNew(STextBlock)
 						.Font(IDetailLayoutBuilder::GetDetailFont())
 						.Text(Target->TargetName)
-						.ShadowOffset(FVector2D::UnitVector)
+						.ColorAndOpacity_Static(&FLandscapeEditorCustomNodeBuilder_TargetLayers::GetTargetTextColor, Target)
 					]
 				]
 			];
@@ -617,7 +617,7 @@ TSharedPtr<SWidget> FLandscapeEditorCustomNodeBuilder_TargetLayers::GenerateRow(
 							SNew(STextBlock)
 							.Font(IDetailLayoutBuilder::GetDetailFont())
 							.Text(Target->TargetName)
-							.ShadowOffset(FVector2D::UnitVector)
+							.ColorAndOpacity_Static(&FLandscapeEditorCustomNodeBuilder_TargetLayers::GetTargetTextColor, Target)
 						]
 						+ SHorizontalBox::Slot()
 						.HAlign(HAlign_Right)
@@ -626,7 +626,7 @@ TSharedPtr<SWidget> FLandscapeEditorCustomNodeBuilder_TargetLayers::GenerateRow(
 							.Visibility_Lambda([=] { return (Target->LayerInfoObj.IsValid() && Target->LayerInfoObj->bNoWeightBlend) ? EVisibility::Visible : EVisibility::Collapsed; })
 							.Font(IDetailLayoutBuilder::GetDetailFont())
 							.Text(LOCTEXT("NoWeightBlend", "No Weight-Blend"))
-							.ShadowOffset(FVector2D::UnitVector)
+							.ColorAndOpacity_Static(&FLandscapeEditorCustomNodeBuilder_TargetLayers::GetTargetTextColor, Target)
 						]
 					]
 					+ SVerticalBox::Slot()
@@ -717,6 +717,7 @@ TSharedPtr<SWidget> FLandscapeEditorCustomNodeBuilder_TargetLayers::GenerateRow(
 							[
 								SNew(STextBlock)
 								.Text(LOCTEXT("SubtractiveBlend", "Subtractive Blend"))
+								.ColorAndOpacity_Static(&FLandscapeEditorCustomNodeBuilder_TargetLayers::GetTargetTextColor, Target)
 							]
 						]
 					]
@@ -1514,6 +1515,11 @@ void FLandscapeEditorCustomNodeBuilder_TargetLayers::OnDebugModeColorChannelChan
 	}
 }
 
+FSlateColor FLandscapeEditorCustomNodeBuilder_TargetLayers::GetTargetTextColor(const TSharedRef<FLandscapeTargetListInfo> InTarget)
+{
+	return FLandscapeEditorCustomNodeBuilder_TargetLayers::GetTargetLayerIsSelected(InTarget) ? FStyleColors::ForegroundHover : FSlateColor::UseForeground();
+}
+
 //////////////////////////////////////////////////////////////////////////
 
 void SLandscapeEditorSelectableBorder::Construct(const FArguments& InArgs)
diff --git a/Engine/Source/Editor/LandscapeEditor/Private/LandscapeEditorDetailCustomization_TargetLayers.h b/Engine/Source/Editor/LandscapeEditor/Private/LandscapeEditorDetailCustomization_TargetLayers.h
--- a/Engine/Source/Editor/LandscapeEditor/Private/LandscapeEditorDetailCustomization_TargetLayers.h
+++ b/Engine/Source/Editor/LandscapeEditor/Private/LandscapeEditorDetailCustomization_TargetLayers.h
@@ -94,6 +94,7 @@ protected:
 	static EVisibility GetDebugModeColorChannelVisibility(const TSharedRef<FLandscapeTargetListInfo> Target);
 	static ECheckBoxState DebugModeColorChannelIsChecked(const TSharedRef<FLandscapeTargetListInfo> Target, int32 Channel);
 	static void OnDebugModeColorChannelChanged(ECheckBoxState NewCheckedState, const TSharedRef<FLandscapeTargetListInfo> Target, int32 Channel);
+	static FSlateColor GetTargetTextColor(const TSharedRef<FLandscapeTargetListInfo> InTarget);
 
 	static EVisibility GetLayersSubstractiveBlendVisibility(const TSharedRef<FLandscapeTargetListInfo> Target);
 	static ECheckBoxState IsLayersSubstractiveBlendChecked(const TSharedRef<FLandscapeTargetListInfo> Target);
