Fixes new landscape location/rotation/scale reset to default button not showing up when using the gizmos.
Adds LandscapeEditor DetailsView RightColumnMinWidth to allow ResetToDefault button to be properly hit testable.

diff --git a/Engine/Source/Editor/LandscapeEditor/Private/LandscapeEditorDetailCustomization_NewLandscape.cpp b/Engine/Source/Editor/LandscapeEditor/Private/LandscapeEditorDetailCustomization_NewLandscape.cpp
--- a/Engine/Source/Editor/LandscapeEditor/Private/LandscapeEditorDetailCustomization_NewLandscape.cpp
+++ b/Engine/Source/Editor/LandscapeEditor/Private/LandscapeEditorDetailCustomization_NewLandscape.cpp
@@ -278,7 +278,16 @@ void FLandscapeEditorDetailCustomization_NewLandscape::CustomizeDetails(IDetailL
 	TSharedRef<IPropertyHandle> PropertyHandle_Location_X = PropertyHandle_Location->GetChildHandle("X").ToSharedRef();
 	TSharedRef<IPropertyHandle> PropertyHandle_Location_Y = PropertyHandle_Location->GetChildHandle("Y").ToSharedRef();
 	TSharedRef<IPropertyHandle> PropertyHandle_Location_Z = PropertyHandle_Location->GetChildHandle("Z").ToSharedRef();
-	NewLandscapeCategory.AddProperty(PropertyHandle_Location)
+	
+	IDetailPropertyRow& LocationPropertyRow = NewLandscapeCategory.AddProperty(PropertyHandle_Location);
+
+	FIsResetToDefaultVisible IsResetLocationVisible = FIsResetToDefaultVisible::CreateSP(this, &FLandscapeEditorDetailCustomization_NewLandscape::ShouldShowResetLocationlToDefault);
+	FResetToDefaultHandler ResetLocationHandler = FResetToDefaultHandler::CreateSP(this, &FLandscapeEditorDetailCustomization_NewLandscape::ResetLocationToDefault);
+	FResetToDefaultOverride ResetLocationOverride = FResetToDefaultOverride::Create(IsResetLocationVisible, ResetLocationHandler);
+
+	LocationPropertyRow.OverrideResetToDefault(ResetLocationOverride);
+
+	LocationPropertyRow
 	.CustomWidget()
 	.NameContent()
 	[
@@ -307,7 +316,16 @@ void FLandscapeEditorDetailCustomization_NewLandscape::CustomizeDetails(IDetailL
 	TSharedRef<IPropertyHandle> PropertyHandle_Rotation_Roll  = PropertyHandle_Rotation->GetChildHandle("Roll").ToSharedRef();
 	TSharedRef<IPropertyHandle> PropertyHandle_Rotation_Pitch = PropertyHandle_Rotation->GetChildHandle("Pitch").ToSharedRef();
 	TSharedRef<IPropertyHandle> PropertyHandle_Rotation_Yaw   = PropertyHandle_Rotation->GetChildHandle("Yaw").ToSharedRef();
-	NewLandscapeCategory.AddProperty(PropertyHandle_Rotation)
+
+	IDetailPropertyRow& RotationPropertyRow = NewLandscapeCategory.AddProperty(PropertyHandle_Rotation);
+
+	FIsResetToDefaultVisible IsResetRotationVisible = FIsResetToDefaultVisible::CreateSP(this, &FLandscapeEditorDetailCustomization_NewLandscape::ShouldShowResetRotationToDefault);
+	FResetToDefaultHandler ResetRotationHandler = FResetToDefaultHandler::CreateSP(this, &FLandscapeEditorDetailCustomization_NewLandscape::ResetRotationToDefault);
+	FResetToDefaultOverride ResetRotationOverride = FResetToDefaultOverride::Create(IsResetRotationVisible, ResetRotationHandler);
+
+	RotationPropertyRow.OverrideResetToDefault(ResetRotationOverride);
+
+	RotationPropertyRow
 	.CustomWidget()
 	.NameContent()
 	[
@@ -332,7 +350,16 @@ void FLandscapeEditorDetailCustomization_NewLandscape::CustomizeDetails(IDetailL
 	TSharedRef<IPropertyHandle> PropertyHandle_Scale_X = PropertyHandle_Scale->GetChildHandle("X").ToSharedRef();
 	TSharedRef<IPropertyHandle> PropertyHandle_Scale_Y = PropertyHandle_Scale->GetChildHandle("Y").ToSharedRef();
 	TSharedRef<IPropertyHandle> PropertyHandle_Scale_Z = PropertyHandle_Scale->GetChildHandle("Z").ToSharedRef();
-	NewLandscapeCategory.AddProperty(PropertyHandle_Scale)
+	
+	IDetailPropertyRow& ScalePropertyRow = NewLandscapeCategory.AddProperty(PropertyHandle_Scale);
+
+	FIsResetToDefaultVisible IsResetScaleVisible = FIsResetToDefaultVisible::CreateSP(this, &FLandscapeEditorDetailCustomization_NewLandscape::ShouldShowResetScaleToDefault);
+	FResetToDefaultHandler ResetScaleHandler = FResetToDefaultHandler::CreateSP(this, &FLandscapeEditorDetailCustomization_NewLandscape::ResetScaleToDefault);
+	FResetToDefaultOverride ResetScaleOverride = FResetToDefaultOverride::Create(IsResetScaleVisible, ResetScaleHandler);
+
+	ScalePropertyRow.OverrideResetToDefault(ResetScaleOverride);
+
+	ScalePropertyRow
 	.CustomWidget()
 	.NameContent()
 	[
@@ -1506,4 +1533,43 @@ bool FLandscapeEditorDetailCustomization_NewLandscape::ShouldShowResetMaterialTo
 	return DefaultMaterial.Get() != Object;
 }
 
+void FLandscapeEditorDetailCustomization_NewLandscape::ResetLocationToDefault(TSharedPtr<IPropertyHandle> InPropertyHandle)
+{
+	InPropertyHandle->SetValue(ULandscapeEditorObject::NewLandscape_DefaultLocation);
+}
+
+bool FLandscapeEditorDetailCustomization_NewLandscape::ShouldShowResetLocationlToDefault(TSharedPtr<IPropertyHandle> InPropertyHandle)
+{
+	FVector Location;
+	InPropertyHandle->GetValue(Location);
+
+	return Location != ULandscapeEditorObject::NewLandscape_DefaultLocation;
+}
+
+void FLandscapeEditorDetailCustomization_NewLandscape::ResetRotationToDefault(TSharedPtr<IPropertyHandle> InPropertyHandle)
+{
+	InPropertyHandle->SetValue(ULandscapeEditorObject::NewLandscape_DefaultRotation);
+}
+
+bool FLandscapeEditorDetailCustomization_NewLandscape::ShouldShowResetRotationToDefault(TSharedPtr<IPropertyHandle> InPropertyHandle)
+{
+	FRotator Rotation;
+	InPropertyHandle->GetValue(Rotation);
+
+	return Rotation != ULandscapeEditorObject::NewLandscape_DefaultRotation;
+}
+
+void FLandscapeEditorDetailCustomization_NewLandscape::ResetScaleToDefault(TSharedPtr<IPropertyHandle> InPropertyHandle)
+{
+	InPropertyHandle->SetValue(ULandscapeEditorObject::NewLandscape_DefaultScale);
+}
+
+bool FLandscapeEditorDetailCustomization_NewLandscape::ShouldShowResetScaleToDefault(TSharedPtr<IPropertyHandle> InPropertyHandle)
+{
+	FVector Scale;
+	InPropertyHandle->GetValue(Scale);
+
+	return Scale != ULandscapeEditorObject::NewLandscape_DefaultScale;
+}
+
 #undef LOCTEXT_NAMESPACE
diff --git a/Engine/Source/Editor/LandscapeEditor/Private/LandscapeEditorDetailCustomization_NewLandscape.h b/Engine/Source/Editor/LandscapeEditor/Private/LandscapeEditorDetailCustomization_NewLandscape.h
--- a/Engine/Source/Editor/LandscapeEditor/Private/LandscapeEditorDetailCustomization_NewLandscape.h
+++ b/Engine/Source/Editor/LandscapeEditor/Private/LandscapeEditorDetailCustomization_NewLandscape.h
@@ -93,6 +93,15 @@ public:
 	void ResetMaterialToDefault(TSharedPtr<IPropertyHandle> InPropertyHandle);
 	bool ShouldShowResetMaterialToDefault(TSharedPtr<IPropertyHandle> InPropertyHandle);
 
+	void ResetLocationToDefault(TSharedPtr<IPropertyHandle> InPropertyHandle);
+	bool ShouldShowResetLocationlToDefault(TSharedPtr<IPropertyHandle> InPropertyHandle);
+
+	void ResetRotationToDefault(TSharedPtr<IPropertyHandle> InPropertyHandle);
+	bool ShouldShowResetRotationToDefault(TSharedPtr<IPropertyHandle> InPropertyHandle);
+
+	void ResetScaleToDefault(TSharedPtr<IPropertyHandle> InPropertyHandle);
+	bool ShouldShowResetScaleToDefault(TSharedPtr<IPropertyHandle> InPropertyHandle);
+
 protected:
 
 	bool bUsingSlider;
diff --git a/Engine/Source/Editor/LandscapeEditor/Private/LandscapeEditorObject.cpp b/Engine/Source/Editor/LandscapeEditor/Private/LandscapeEditorObject.cpp
--- a/Engine/Source/Editor/LandscapeEditor/Private/LandscapeEditorObject.cpp
+++ b/Engine/Source/Editor/LandscapeEditor/Private/LandscapeEditorObject.cpp
@@ -23,6 +23,10 @@ static TAutoConsoleVariable<bool> CVarLandscapeSimulateAlphaBrushTextureLoadFail
 	false,
 	TEXT("Debug utility to simulate a loading failure (e.g. invalid source data, which can happen in cooked editor or with a badly virtualized texture) when loading the alpha brush texture"));
 
+const FVector ULandscapeEditorObject::NewLandscape_DefaultLocation = FVector(0, 0, 100);
+const FRotator ULandscapeEditorObject::NewLandscape_DefaultRotation = FRotator::ZeroRotator;
+const FVector ULandscapeEditorObject::NewLandscape_DefaultScale = FVector(100, 100, 100);
+
 ULandscapeEditorObject::ULandscapeEditorObject(const FObjectInitializer& ObjectInitializer)
 	: Super(ObjectInitializer)
 
@@ -88,9 +92,9 @@ ULandscapeEditorObject::ULandscapeEditorObject(const FObjectInitializer& ObjectI
 	, NewLandscape_QuadsPerSection(63)
 	, NewLandscape_SectionsPerComponent(1)
 	, NewLandscape_ComponentCount(8, 8)
-	, NewLandscape_Location(0, 0, 100)
-	, NewLandscape_Rotation(0, 0, 0)
-	, NewLandscape_Scale(100, 100, 100)
+	, NewLandscape_Location(NewLandscape_DefaultLocation)
+	, NewLandscape_Rotation(NewLandscape_DefaultRotation)
+	, NewLandscape_Scale(NewLandscape_DefaultScale)
 	, ImportLandscape_Width(0)
 	, ImportLandscape_Height(0)
 	, ImportLandscape_AlphamapType(ELandscapeImportAlphamapType::Additive)
diff --git a/Engine/Source/Editor/LandscapeEditor/Private/SLandscapeEditor.cpp b/Engine/Source/Editor/LandscapeEditor/Private/SLandscapeEditor.cpp
--- a/Engine/Source/Editor/LandscapeEditor/Private/SLandscapeEditor.cpp
+++ b/Engine/Source/Editor/LandscapeEditor/Private/SLandscapeEditor.cpp
@@ -453,6 +453,8 @@ void SLandscapeEditor::Construct(const FArguments& InArgs, TSharedRef<FLandscape
 	FDetailsViewArgs DetailsViewArgs;
 	DetailsViewArgs.NameAreaSettings = FDetailsViewArgs::HideNameArea;
 	DetailsViewArgs.bAllowSearch = false;
+	// Minimum size to allow the ResetToDefault button to be hit testable.
+	DetailsViewArgs.RightColumnMinWidth = 35;
 
 	DetailsPanel = PropertyEditorModule.CreateDetailView(DetailsViewArgs);
 	DetailsPanel->SetIsPropertyVisibleDelegate(FIsPropertyVisible::CreateSP(this, &SLandscapeEditor::GetIsPropertyVisible));
diff --git a/Engine/Source/Editor/LandscapeEditor/Public/LandscapeEditorObject.h b/Engine/Source/Editor/LandscapeEditor/Public/LandscapeEditorObject.h
--- a/Engine/Source/Editor/LandscapeEditor/Public/LandscapeEditorObject.h
+++ b/Engine/Source/Editor/LandscapeEditor/Public/LandscapeEditorObject.h
@@ -541,14 +541,20 @@ class ULandscapeEditorObject : public UObject
 	UPROPERTY(Category="New Landscape", EditAnywhere, meta=(DisplayName="Number of Components", ShowForTools="NewLandscape"))
 	FIntPoint NewLandscape_ComponentCount;
 
+	static const FVector NewLandscape_DefaultLocation;
+
 	// The location of the new landscape
 	UPROPERTY(Category="New Landscape", EditAnywhere, meta=(DisplayName="Location", ShowForTools="NewLandscape"))
 	FVector NewLandscape_Location;
 
+	static const FRotator NewLandscape_DefaultRotation;
+
 	// The rotation of the new landscape
 	UPROPERTY(Category="New Landscape", EditAnywhere, meta=(DisplayName="Rotation", ShowForTools="NewLandscape"))
 	FRotator NewLandscape_Rotation;
 
+	static const FVector NewLandscape_DefaultScale;
+
 	// The scale of the new landscape. This is the distance between each vertex on the landscape, defaulting to 100 units.
 	UPROPERTY(Category="New Landscape", EditAnywhere, meta=(DisplayName="Scale", ShowForTools="NewLandscape"))
 	FVector NewLandscape_Scale;
