Fix spelling error in landscape tool tip

diff --git a/Engine/Source/Editor/LandscapeEditor/Public/LandscapeEditorObject.h b/Engine/Source/Editor/LandscapeEditor/Public/LandscapeEditorObject.h
--- a/Engine/Source/Editor/LandscapeEditor/Public/LandscapeEditorObject.h
+++ b/Engine/Source/Editor/LandscapeEditor/Public/LandscapeEditorObject.h
@@ -637,7 +637,7 @@ public:
 	UPROPERTY(Category="Brush Settings", EditAnywhere, NonTransactional, meta=(DisplayName="Brush Falloff", ShowForBrushes="BrushSet_Circle,BrushSet_Gizmo,BrushSet_Pattern", ShowForTargetTypes = "Heightmap,Visibility", ClampMin="0", ClampMax="1", UIMin = "0", UIMax = "1"))
 	float BrushFalloff;
 	
-	// The falloff at the edge of the point brush, as a fraction of the brush's size. 0 = no falloff, 1 = all falloff
+	// The falloff at the edge of the paint brush, as a fraction of the brush's size. 0 = no falloff, 1 = all falloff
 	UPROPERTY(Category="Brush Settings", EditAnywhere, NonTransactional, meta=(DisplayName="Brush Falloff", ShowForBrushes="BrushSet_Circle,BrushSet_Gizmo,BrushSet_Pattern", ShowForTargetTypes = "Weightmap", ClampMin="0", ClampMax="1", UIMin = "0", UIMax = "1"))
 	float PaintBrushFalloff;
 
