Turned landscape brush size's max slider and max clamp value into project settings

Misc:
* Removed FLandscapeEditorServices, which didn't make much sense in terms of design wrt dependency inversion : instead, FLanscapeEditorModule is the one implementing ILandscapeEditorServices

#tests editor

diff --git a/Engine/Source/Editor/LandscapeEditor/Private/LandscapeEdMode.cpp b/Engine/Source/Editor/LandscapeEditor/Private/LandscapeEdMode.cpp
--- a/Engine/Source/Editor/LandscapeEditor/Private/LandscapeEdMode.cpp
+++ b/Engine/Source/Editor/LandscapeEditor/Private/LandscapeEdMode.cpp
@@ -1722,16 +1722,16 @@ void FEdModeLandscape::ChangeBrushSize(bool bIncrease)
 	else
 	{
 		float Radius = UISettings->GetCurrentToolBrushRadius();
+		const ULandscapeSettings* LandscapeSettings = GetDefault<ULandscapeSettings>();
 		const float SliderMin = 10.0f;
-		const float SliderMax = 8192.0f;
-		float Diff = 0.05f; //6.0f / SliderMax;
+		const float SliderMax = LandscapeSettings->GetBrushSizeUIMax();
+		float Diff = 0.05f; 
 		if (!bIncrease)
 		{
 			Diff = -Diff;
 		}
 
 		float NewValue = Radius * (1.0f + Diff);
-
 		if (bIncrease)
 		{
 			NewValue = FMath::Max(NewValue, Radius + 1.0f);
diff --git a/Engine/Source/Editor/LandscapeEditor/Private/LandscapeEditorDetails.cpp b/Engine/Source/Editor/LandscapeEditor/Private/LandscapeEditorDetails.cpp
--- a/Engine/Source/Editor/LandscapeEditor/Private/LandscapeEditorDetails.cpp
+++ b/Engine/Source/Editor/LandscapeEditor/Private/LandscapeEditorDetails.cpp
@@ -16,6 +16,7 @@
 #include "LandscapeEditorDetailCustomization_MiscTools.h"
 #include "LandscapeEditorDetailCustomization_AlphaBrush.h"
 #include "LandscapeEditorDetailCustomization_ImportExport.h"
+#include "LandscapeSettings.h"
 #include "DetailWidgetRow.h"
 #include "LandscapeEditorDetailCustomization_TargetLayers.h"
 #include "DetailCategoryBuilder.h"
@@ -81,7 +82,18 @@ void FLandscapeEditorDetails::CustomizeDetails(IDetailLayoutBuilder& DetailBuild
 	};
 	
 	DetailBuilder.SortCategories(CategorySorter);
-	
+
+	// UIMax and ClampMax for the brush radius come from the project settings : 	
+	const ULandscapeSettings* Settings = GetDefault<ULandscapeSettings>();
+	TSharedRef<IPropertyHandle> BrushRadiusProperty = DetailBuilder.GetProperty(GET_MEMBER_NAME_CHECKED(ULandscapeEditorObject, BrushRadius));
+	TSharedRef<IPropertyHandle> PaintBrushRadiusProperty = DetailBuilder.GetProperty(GET_MEMBER_NAME_CHECKED(ULandscapeEditorObject, PaintBrushRadius));
+	const FString UIMaxString = FString::Printf(TEXT("%f"), Settings->GetBrushSizeUIMax());
+	const FString ClampMaxString = FString::Printf(TEXT("%f"), Settings->GetBrushSizeClampMax());
+	BrushRadiusProperty->SetInstanceMetaData(TEXT("UIMax"), *UIMaxString);
+	BrushRadiusProperty->SetInstanceMetaData(TEXT("ClampMax"), *ClampMaxString);
+	PaintBrushRadiusProperty->SetInstanceMetaData(TEXT("UIMax"), *UIMaxString);
+	PaintBrushRadiusProperty->SetInstanceMetaData(TEXT("ClampMax"), *ClampMaxString);
+
 	LandscapeEditorCategory.AddCustomRow(FText::GetEmpty())
 	.Visibility(TAttribute<EVisibility>::Create(TAttribute<EVisibility>::FGetter::CreateStatic(&FLandscapeEditorDetails::GetTargetLandscapeSelectorVisibility)))
 	[
diff --git a/Engine/Source/Editor/LandscapeEditor/Private/LandscapeEditorModule.cpp b/Engine/Source/Editor/LandscapeEditor/Private/LandscapeEditorModule.cpp
--- a/Engine/Source/Editor/LandscapeEditor/Private/LandscapeEditorModule.cpp
+++ b/Engine/Source/Editor/LandscapeEditor/Private/LandscapeEditorModule.cpp
@@ -18,9 +18,9 @@
 #include "Classes/ActorFactoryLandscape.h"
 #include "LandscapeFileFormatPng.h"
 #include "LandscapeFileFormatRaw.h"
-#include "Settings/EditorExperimentalSettings.h"
 #include "LandscapeEditorServices.h"
 #include "LandscapeImageFileCache.h"
+#include "SLandscapeLayerListDialog.h"
 
 #include "Framework/MultiBox/MultiBoxBuilder.h"
 #include "PropertyEditorModule.h"
@@ -60,7 +60,7 @@ struct FRegisteredLandscapeWeightmapFileFormat
 	FRegisteredLandscapeWeightmapFileFormat(TSharedRef<ILandscapeWeightmapFileFormat> InFileFormat);
 };
 
-class FLandscapeEditorModule : public ILandscapeEditorModule
+class FLandscapeEditorModule : public ILandscapeEditorModule, public ILandscapeEditorServices
 {
 public:
 
@@ -149,8 +149,7 @@ public:
 		}
 		
 		ILandscapeModule& LandscapeModule = FModuleManager::GetModuleChecked<ILandscapeModule>("Landscape");
-		LandscapeEditorServices.Reset(new FLandscapeEditorServices);
-		LandscapeModule.SetLandscapeEditorServices(LandscapeEditorServices.Get());
+		LandscapeModule.SetLandscapeEditorServices(this);
 
 		LandscapeImageFileCache.Reset(new FLandscapeImageFileCache());
 	}
@@ -185,11 +184,10 @@ public:
 		// GEditor->ActorFactories.RemoveAll([](const UActorFactory* ActorFactory) { return ActorFactory->IsA<UActorFactoryLandscape>(); });
 
 		ILandscapeModule& LandscapeModule = FModuleManager::GetModuleChecked<ILandscapeModule>("Landscape");
-		if (LandscapeModule.GetLandscapeEditorServices() == LandscapeEditorServices.Get())
+		if (LandscapeModule.GetLandscapeEditorServices() == this)
 		{
 			LandscapeModule.SetLandscapeEditorServices(nullptr);
 		}
-		LandscapeEditorServices.Reset();
 		LandscapeImageFileCache.Reset();
 	}
 
@@ -307,6 +305,9 @@ public:
 		return GLandscapeViewMode == ViewMode;
 	}
 
+	/**
+	 * ILandscapeEditorModule implementation
+	 */
 	virtual void RegisterHeightmapFileFormat(TSharedRef<ILandscapeHeightmapFileFormat> FileFormat) override
 	{
 		HeightmapFormats.Emplace(FileFormat);
@@ -364,6 +365,12 @@ public:
 
 	FLandscapeImageFileCache& GetImageFileCache() const override;
 
+	/**
+	* ILandscapeEditorServices implementation
+	*/
+	virtual int32 GetOrCreateEditLayer(FName InEditLayerName, ALandscape* InTargetLandscape) override;
+	virtual void RefreshDetailPanel() override;
+
 protected:
 	TSharedPtr<FExtender> ViewportMenuExtender;
 	TSharedPtr<FUICommandList> GlobalUICommandList;
@@ -373,7 +380,6 @@ protected:
 	mutable FString WeightmapImportDialogTypeString;
 	mutable FString HeightmapExportDialogTypeString;
 	mutable FString WeightmapExportDialogTypeString;
-	TUniquePtr<ILandscapeEditorServices> LandscapeEditorServices;
 	TUniquePtr<FLandscapeImageFileCache> LandscapeImageFileCache;
 };
 
@@ -549,5 +555,26 @@ FLandscapeImageFileCache& FLandscapeEditorModule::GetImageFileCache() const
 	return *LandscapeImageFileCache;
 }
 
+int32 FLandscapeEditorModule::GetOrCreateEditLayer(FName InEditLayerName, ALandscape* InTargetLandscape)
+{
+	// Insertion logic is left to the user through modal drag + drop dialog : 
+	int32 ExistingLayerIndex = InTargetLandscape->GetLayerIndex(InEditLayerName);
+	if (ExistingLayerIndex == INDEX_NONE)
+	{
+		InTargetLandscape->CreateLayer(InEditLayerName);
+		TSharedPtr<SLandscapeLayerListDialog> Dialog = SNew(SLandscapeLayerListDialog, InTargetLandscape->LandscapeLayers);
+		Dialog->ShowModal();
+		ExistingLayerIndex = Dialog->GetInsertedLayerIndex();
+	}
+	return ExistingLayerIndex;
+}
+
+void FLandscapeEditorModule::RefreshDetailPanel()
+{
+	if (FEdModeLandscape* LandscapeMode = (FEdModeLandscape*)GLevelEditorModeTools().GetActiveMode(FBuiltinEditorModes::EM_Landscape))
+	{
+		LandscapeMode->RefreshDetailPanel();
+	}
+}
 
 #undef LOCTEXT_NAMESPACE
diff --git a/Engine/Source/Editor/LandscapeEditor/Private/LandscapeEditorServices.cpp b/Engine/Source/Editor/LandscapeEditor/Private/LandscapeEditorServices.cpp
deleted file mode 100644
diff --git a/Engine/Source/Editor/LandscapeEditor/Public/LandscapeEditorObject.h b/Engine/Source/Editor/LandscapeEditor/Public/LandscapeEditorObject.h
--- a/Engine/Source/Editor/LandscapeEditor/Public/LandscapeEditorObject.h
+++ b/Engine/Source/Editor/LandscapeEditor/Public/LandscapeEditorObject.h
@@ -615,11 +615,11 @@ public:
 	// Common Brush Settings:
 
 	// The radius of the sculpt brush, in unreal units
-	UPROPERTY(Category="Brush Settings", EditAnywhere, NonTransactional, meta=(DisplayName="Brush Size", ShowForBrushes="BrushSet_Circle,BrushSet_Alpha,BrushSet_Pattern", ShowForTargetTypes = "Heightmap,Visibility", ClampMin="1", ClampMax="65536", UIMin="1", UIMax="8192", SliderExponent="3", MaxFractionalDigits="2"))
+	UPROPERTY(Category="Brush Settings", EditAnywhere, NonTransactional, meta=(DisplayName="Brush Size", ShowForBrushes="BrushSet_Circle,BrushSet_Alpha,BrushSet_Pattern", ShowForTargetTypes = "Heightmap,Visibility", ClampMin="1", ClampMax="65536", UIMin="1", UIMax="8192", SliderExponent="3", MaxFractionalDigits="2", ToolTip = "Radius of the sculpt editing brush. The maximum slider/clamp value can be set with the BrushSizeUIMax/BrushSizeClampMax values in the Landscape project settings."))
 	float BrushRadius;
 
 	// The radius of the paint brush, in unreal units
-	UPROPERTY(Category="Brush Settings", EditAnywhere, NonTransactional, meta=(DisplayName="Brush Size", ShowForBrushes="BrushSet_Circle,BrushSet_Alpha,BrushSet_Pattern", ShowForTargetTypes = "Weightmap", ClampMin="1", ClampMax="65536", UIMin="1", UIMax="8192", SliderExponent="3"))
+	UPROPERTY(Category="Brush Settings", EditAnywhere, NonTransactional, meta=(DisplayName="Brush Size", ShowForBrushes="BrushSet_Circle,BrushSet_Alpha,BrushSet_Pattern", ShowForTargetTypes = "Weightmap", ClampMin="1", ClampMax="65536", UIMin="1", UIMax="8192", SliderExponent="3", ToolTip = "Radius of the paint editing brush. The maximum slider/clamp value can be set with the BrushSizeUIMax/BrushSizeClampMax values in the Landscape project settings."))
 	float PaintBrushRadius;
 
 	// The falloff at the edge of the sculpt brush, as a fraction of the brush's size. 0 = no falloff, 1 = all falloff
diff --git a/Engine/Source/Editor/LandscapeEditor/Public/SLandscapeLayerListDialog.h b/Engine/Source/Editor/LandscapeEditor/Public/SLandscapeLayerListDialog.h
--- a/Engine/Source/Editor/LandscapeEditor/Public/SLandscapeLayerListDialog.h
+++ b/Engine/Source/Editor/LandscapeEditor/Public/SLandscapeLayerListDialog.h
@@ -18,8 +18,6 @@ typedef SListView<TSharedPtr<FWidgetLayerListItem> > SWidgetLayerListView;
 class SLandscapeLayerListDialog : public SCustomDialog
 {
 public:
-	friend class FLandscapeEditorServices;
-	
 	virtual ~SLandscapeLayerListDialog() {}
 	
 	SLATE_BEGIN_ARGS( SLandscapeLayerListDialog	){}
@@ -32,6 +30,8 @@ public:
 	 */
 	void Construct(const FArguments& InArgs, TArray<FLandscapeLayer>& InLayers);
 
+	int32 GetInsertedLayerIndex() const { return InsertedLayerIndex; }
+
 private:
 	/* Called when LayerList is changed to update WidgetLayerList and the dialog window. */
 	void OnLayerListUpdated();
diff --git a/Engine/Source/Runtime/Landscape/Private/LandscapeSettings.cpp b/Engine/Source/Runtime/Landscape/Private/LandscapeSettings.cpp
new file mode 100644
--- /dev/null
+++ b/Engine/Source/Runtime/Landscape/Private/LandscapeSettings.cpp
@@ -0,0 +1,24 @@
+// Copyright Epic Games, Inc. All Rights Reserved.
+
+#include "LandscapeSettings.h"
+#include "LandscapeModule.h"
+#include "Modules/ModuleManager.h"
+#include "LandscapeEditorServices.h"
+
+#if WITH_EDITOR
+
+void ULandscapeSettings::PostEditChangeProperty(FPropertyChangedEvent& PropertyChangedEvent)
+{
+	Super::PostEditChangeProperty(PropertyChangedEvent);
+
+	const FName PropertyName = PropertyChangedEvent.MemberProperty ? PropertyChangedEvent.MemberProperty->GetFName() : NAME_None;
+	if ((PropertyName == GET_MEMBER_NAME_CHECKED(ULandscapeSettings, BrushSizeUIMax))
+		|| (PropertyName == GET_MEMBER_NAME_CHECKED(ULandscapeSettings, BrushSizeClampMax)))
+	{
+		// If landscape mode is active, refresh the detail panel to apply the changes immediately : 
+		ILandscapeModule& LandscapeModule = FModuleManager::GetModuleChecked<ILandscapeModule>("Landscape");
+		LandscapeModule.GetLandscapeEditorServices()->RefreshDetailPanel();
+	}
+}
+
+#endif // WITH_EDITOR
diff --git a/Engine/Source/Runtime/Landscape/Public/LandscapeEditorServices.h b/Engine/Source/Runtime/Landscape/Public/LandscapeEditorServices.h
--- a/Engine/Source/Runtime/Landscape/Public/LandscapeEditorServices.h
+++ b/Engine/Source/Runtime/Landscape/Public/LandscapeEditorServices.h
@@ -18,15 +18,7 @@ public:
 	 * @return The index at which the edit layer named InEditLayerName exists
 	 */
 	virtual int32 GetOrCreateEditLayer(FName InEditLayerName, ALandscape* InTargetLandscape) = 0;
-};
-
-#if WITH_EDITOR
-class FLandscapeEditorServices : public ILandscapeEditorServices
-{
-public:
-	virtual ~FLandscapeEditorServices() {}
 
-	// Implements pure virtual function. Insertion logic is left to the user through modal drag + drop dialog.
-	virtual int32 GetOrCreateEditLayer(FName InEditLayerName, ALandscape* InTargetLandscape);
+	/** Requires the landscape editor mode, if active to refresh its detail panel */
+	virtual void RefreshDetailPanel() = 0;
 };
-#endif // WITH_EDITOR
\ No newline at end of file
diff --git a/Engine/Source/Runtime/Landscape/Public/LandscapeSettings.h b/Engine/Source/Runtime/Landscape/Public/LandscapeSettings.h
--- a/Engine/Source/Runtime/Landscape/Public/LandscapeSettings.h
+++ b/Engine/Source/Runtime/Landscape/Public/LandscapeSettings.h
@@ -10,8 +10,6 @@
 class UMaterialInterface;
 class ULandscapeLayerInfoObject;
 
-DECLARE_DELEGATE_OneParam(FnOnMaxComponentsChanged, uint64);
-
 UENUM()
 enum class ELandscapeDirtyingMode : uint8
 {
@@ -40,11 +38,20 @@ public:
 
 	int32 GetSideResolutionLimit() const { return SideResolutionLimit; }
 
+	float GetBrushSizeUIMax() const { return BrushSizeUIMax; }
+	float GetBrushSizeClampMax() const { return BrushSizeClampMax; }
+
 	/** Returns the default landscape material that should be used when creating a new landscape. */
 	TSoftObjectPtr<UMaterialInterface> GetDefaultLandscapeMaterial() const { return DefaultLandscapeMaterial; }
 
 	/** Returns the default landscape layer info object that will be assigned to unset layers when creating a new landscape. */
 	TSoftObjectPtr<ULandscapeLayerInfoObject> GetDefaultLayerInfoObject() const { return DefaultLayerInfoObject; }
+
+#if WITH_EDITOR
+	//~ Begin UObject Interface.
+	virtual void PostEditChangeProperty(FPropertyChangedEvent& PropertyChangedEvent) override;
+#endif // WITH_EDITOR
+
 public:
 	UPROPERTY(config, EditAnywhere, Category = "Layers", meta=(UIMin = "1", UIMax = "32", ClampMin = "1", ClampMax = "32", ToolTip = "This option controls the maximum editing layers that can be added to a Landscape"))
 	int32 MaxNumberOfLayers = 8;
@@ -74,7 +81,13 @@ protected:
 	UPROPERTY(EditAnywhere, config, Category = "Layers", meta = (ToolTip = "Default Layer Info Object"))
 	TSoftObjectPtr<ULandscapeLayerInfoObject> DefaultLayerInfoObject;
 
+	UPROPERTY(EditAnywhere, config, Category = "Configuration", meta = (ToolTip = "Maximum size that can be set via the slider for the landscape sculpt/paint brushes"))
+	float BrushSizeUIMax = 8192;
+
+	UPROPERTY(EditAnywhere, config, Category = "Configuration", meta = (ToolTip = "Maximum size that can be set manually for the landscape sculpt/paint brushes"))
+	float BrushSizeClampMax = 65536;
+
 	UPROPERTY(transient)
 	bool bRestrictiveMode = false;
-
 };
+
