Various UX improvments for landscape :
* Added an option to be able to consider the paint layer's hardness parameter when using the erosion tool. Set to false by default because this used to be misleading for users
* Set hardness to 0 by default for landscape layer info assets (for the same reason as above)
* Added several missing tooltips
* Added support for double-clicking on a paint layer thumbnail to open the corresponding landscape layer info asset

diff --git a/Engine/Config/BaseEditorPerProjectUserSettings.ini b/Engine/Config/BaseEditorPerProjectUserSettings.ini
--- a/Engine/Config/BaseEditorPerProjectUserSettings.ini
+++ b/Engine/Config/BaseEditorPerProjectUserSettings.ini
@@ -741,6 +741,7 @@ ErodeIterationNum=28
 ErodeSurfaceThickness=256
 ErosionNoiseMode=2
 ErosionNoiseScale=60.000000
+bErosionUseLayerHardness=False
 RainAmount=128
 SedimentCapacity=0.300000
 HErodeIterationNum=28
diff --git a/Engine/Source/Editor/LandscapeEditor/Private/LandscapeEdModeErosionTools.cpp b/Engine/Source/Editor/LandscapeEditor/Private/LandscapeEdModeErosionTools.cpp
--- a/Engine/Source/Editor/LandscapeEditor/Private/LandscapeEdModeErosionTools.cpp
+++ b/Engine/Source/Editor/LandscapeEditor/Private/LandscapeEdModeErosionTools.cpp
@@ -114,19 +114,25 @@ public:
 		LayerHeightDataCache.Read(X1, Y1, X2, Y2, HeightData);
 
 		TArray<uint8> WeightDatas; // Weight*Layers...
-		WeightCache.CacheData(X1, Y1, X2, Y2);
-		WeightCache.GetCachedData(X1, Y1, X2, Y2, WeightDatas, LayerNum);
+		if (UISettings->bErosionUseLayerHardness || bWeightApplied)
+		{
+			WeightCache.CacheData(X1, Y1, X2, Y2);
+			WeightCache.GetCachedData(X1, Y1, X2, Y2, WeightDatas, LayerNum);
+		}
 
 		// Apply the brush	
 		uint16 Thresh = static_cast<uint16>(UISettings->ErodeThresh);
 		int32 WeightMoveThresh = FMath::Min<int32>(FMath::Max<int32>(Thickness >> 2, Thresh), Thickness >> 1);
 
 		TArray<float> CenterWeights;
-		CenterWeights.Empty(LayerNum);
-		CenterWeights.AddUninitialized(LayerNum);
 		TArray<float> NeighborWeight;
-		NeighborWeight.Empty(NeighborNum*LayerNum);
-		NeighborWeight.AddUninitialized(NeighborNum*LayerNum);
+		if (bWeightApplied)
+		{
+			CenterWeights.Empty(LayerNum);
+			CenterWeights.AddUninitialized(LayerNum);
+			NeighborWeight.Empty(NeighborNum*LayerNum);
+			NeighborWeight.AddUninitialized(NeighborNum*LayerNum);
+		}
 
 		bool bHasChanged = false;
 		for (int32 i = 0; i < Iteration; i++)
@@ -170,6 +176,7 @@ public:
 						if (SlopeTotal > 0)
 						{
 							float Softness = 1.0f;
+							if (UISettings->bErosionUseLayerHardness)
 							{
 								for (int32 Idx = 0; Idx < LayerNum; Idx++)
 								{
diff --git a/Engine/Source/Editor/LandscapeEditor/Private/LandscapeEditorDetailCustomization_TargetLayers.cpp b/Engine/Source/Editor/LandscapeEditor/Private/LandscapeEditorDetailCustomization_TargetLayers.cpp
--- a/Engine/Source/Editor/LandscapeEditor/Private/LandscapeEditorDetailCustomization_TargetLayers.cpp
+++ b/Engine/Source/Editor/LandscapeEditor/Private/LandscapeEditorDetailCustomization_TargetLayers.cpp
@@ -29,6 +29,7 @@
 #include "IDetailPropertyRow.h"
 #include "DetailCategoryBuilder.h"
 #include "PropertyCustomizationHelpers.h"
+#include "Subsystems/AssetEditorSubsystem.h"
 
 #include "SLandscapeEditor.h"
 #include "Dialogs/DlgPickAssetPath.h"
@@ -504,7 +505,7 @@ TSharedPtr<SWidget> FLandscapeEditorCustomNodeBuilder_TargetLayers::GenerateRow(
 			return RowWidget;
 		}
 	}
-	
+
 	if (Target->TargetType != ELandscapeToolTargetType::Weightmap)
 	{
 		RowWidget = SNew(SLandscapeEditorSelectableBorder)
@@ -512,34 +513,34 @@ TSharedPtr<SWidget> FLandscapeEditorCustomNodeBuilder_TargetLayers::GenerateRow(
 			.VAlign(VAlign_Center)
 			.OnContextMenuOpening_Static(&FLandscapeEditorCustomNodeBuilder_TargetLayers::OnTargetLayerContextMenuOpening, Target)
 			.OnSelected_Static(&FLandscapeEditorCustomNodeBuilder_TargetLayers::OnTargetSelectionChanged, Target)
-			.IsSelected_Static(&FLandscapeEditorCustomNodeBuilder_TargetLayers::GetTargetLayerIsSelected, Target)			
+			.IsSelected_Static(&FLandscapeEditorCustomNodeBuilder_TargetLayers::GetTargetLayerIsSelected, Target)
 			.Visibility(this, &FLandscapeEditorCustomNodeBuilder_TargetLayers::ShouldShowLayer, Target)
 			[
 				SNew(SHorizontalBox)
-				+ SHorizontalBox::Slot()
-				.AutoWidth()
-				.VAlign(VAlign_Center)
-				.Padding(FMargin(2))
-				[
-					SNew(SImage)
-					.Image(FAppStyle::GetBrush(Target->TargetType == ELandscapeToolTargetType::Heightmap ? TEXT("LandscapeEditor.Target_Heightmap") : TEXT("LandscapeEditor.Target_Visibility")))
-				]
-				+ SHorizontalBox::Slot()
-				.VAlign(VAlign_Center)
-				.Padding(4, 0)
-				[
-					SNew(SVerticalBox)
-					+ SVerticalBox::Slot()
-					.AutoHeight()
+					+ SHorizontalBox::Slot()
+					.AutoWidth()
 					.VAlign(VAlign_Center)
-					.Padding(0, 2)
+					.Padding(FMargin(2))
 					[
-						SNew(STextBlock)
-						.Font(IDetailLayoutBuilder::GetDetailFont())
-						.Text(Target->TargetName)
-						.ColorAndOpacity_Static(&FLandscapeEditorCustomNodeBuilder_TargetLayers::GetTargetTextColor, Target)
+						SNew(SImage)
+							.Image(FAppStyle::GetBrush(Target->TargetType == ELandscapeToolTargetType::Heightmap ? TEXT("LandscapeEditor.Target_Heightmap") : TEXT("LandscapeEditor.Target_Visibility")))
+					]
+					+ SHorizontalBox::Slot()
+					.VAlign(VAlign_Center)
+					.Padding(4, 0)
+					[
+						SNew(SVerticalBox)
+							+ SVerticalBox::Slot()
+							.AutoHeight()
+							.VAlign(VAlign_Center)
+							.Padding(0, 2)
+							[
+								SNew(STextBlock)
+									.Font(IDetailLayoutBuilder::GetDetailFont())
+									.Text(Target->TargetName)
+									.ColorAndOpacity_Static(&FLandscapeEditorCustomNodeBuilder_TargetLayers::GetTargetTextColor, Target)
+							]
 					]
-				]
 			];
 	}
 	else
@@ -553,229 +554,239 @@ TSharedPtr<SWidget> FLandscapeEditorCustomNodeBuilder_TargetLayers::GenerateRow(
 			.OnSelected_Static(&FLandscapeEditorCustomNodeBuilder_TargetLayers::OnTargetSelectionChanged, Target)
 			.IsSelected_Static(&FLandscapeEditorCustomNodeBuilder_TargetLayers::GetTargetLayerIsSelected, Target)
 			.Visibility(this, &FLandscapeEditorCustomNodeBuilder_TargetLayers::ShouldShowLayer, Target)
-			[				
+			[
 				SNew(SHorizontalBox)
 
-				+ SHorizontalBox::Slot()
-				.AutoWidth()
-				.VAlign(VAlign_Center)
-				[
-					SNew(SBox)
-					.Padding(FMargin(2.0f, 0.0f, 2.0f, 0.0f))
-					[
-						SNew(SImage)
-						.Image(FCoreStyle::Get().GetBrush("VerticalBoxDragIndicator"))
-					]
-				]
-
-				+ SHorizontalBox::Slot()
-				.AutoWidth()
-				.VAlign(VAlign_Center)
-				.Padding(FMargin(2))
-				[
-					SNew(SBox)
-					.Visibility_Static(&FLandscapeEditorCustomNodeBuilder_TargetLayers::GetDebugModeLayerUsageVisibility, Target)
-					.WidthOverride(48.0f)
-					.HeightOverride(48.0f)
-					[
-						SNew(SImage)
-						.Image(FCoreStyle::Get().GetBrush("WhiteBrush"))
-						.ColorAndOpacity_Static(&FLandscapeEditorCustomNodeBuilder_TargetLayers::GetLayerUsageDebugColor, Target)
-					]
-				]
-
-				+ SHorizontalBox::Slot()
-				.AutoWidth()
-				.VAlign(VAlign_Center)
-				.Padding(FMargin(2))
-				[
-					(Target->bValid)
-					? (TSharedRef<SWidget>)(
-					SNew(SLandscapeAssetThumbnail, Target->ThumbnailMIC.Get(), ThumbnailPool)
-					.Visibility_Static(&FLandscapeEditorCustomNodeBuilder_TargetLayers::GetDebugModeLayerUsageVisibility_Invert, Target)
-					.ThumbnailSize(FIntPoint(48, 48))
-					)
-					: (TSharedRef<SWidget>)(
-					SNew(SImage)
-					.Visibility_Static(&FLandscapeEditorCustomNodeBuilder_TargetLayers::GetDebugModeLayerUsageVisibility_Invert, Target)
-					.Image(FAppStyle::GetBrush(TEXT("LandscapeEditor.Target_Invalid")))
-					)
-				]
-				+ SHorizontalBox::Slot()
-				.VAlign(VAlign_Center)
-				.Padding(4, 0)
-				[
-					SNew(SVerticalBox)
-					+ SVerticalBox::Slot()
-					.AutoHeight()
-					.VAlign(VAlign_Center)
-					.Padding(0, 2, 0, 0)
-					[
-						SNew(SHorizontalBox)
-						+ SHorizontalBox::Slot()
-						[
-							SNew(STextBlock)
-							.Font(IDetailLayoutBuilder::GetDetailFont())
-							.Text(Target->TargetName)
-							.ColorAndOpacity_Static(&FLandscapeEditorCustomNodeBuilder_TargetLayers::GetTargetTextColor, Target)
-						]
-						+ SHorizontalBox::Slot()
-						.HAlign(HAlign_Right)
-						[
-							SNew(STextBlock)
-							.Visibility_Lambda([=] { return (Target->LayerInfoObj.IsValid() && Target->LayerInfoObj->bNoWeightBlend) ? EVisibility::Visible : EVisibility::Collapsed; })
-							.Font(IDetailLayoutBuilder::GetDetailFont())
-							.Text(LOCTEXT("NoWeightBlend", "No Weight-Blend"))
-							.ColorAndOpacity_Static(&FLandscapeEditorCustomNodeBuilder_TargetLayers::GetTargetTextColor, Target)
-						]
-					]
-					+ SVerticalBox::Slot()
-					.AutoHeight()
+					+ SHorizontalBox::Slot()
+					.AutoWidth()
 					.VAlign(VAlign_Center)
 					[
-						SNew(SHorizontalBox)
-						.Visibility_Static(&FLandscapeEditorCustomNodeBuilder_TargetLayers::GetTargetLayerInfoSelectorVisibility, Target)
-						+ SHorizontalBox::Slot()
-						.FillWidth(1)
-						.VAlign(VAlign_Center)
-						[
-							SNew(SObjectPropertyEntryBox)
-							.IsEnabled((bool)Target->bValid)
-							.ObjectPath(Target->LayerInfoObj != NULL ? Target->LayerInfoObj->GetPathName() : FString())
-							.AllowedClass(ULandscapeLayerInfoObject::StaticClass())
-							.OnObjectChanged_Static(&FLandscapeEditorCustomNodeBuilder_TargetLayers::OnTargetLayerSetObject, Target)
-							.OnShouldFilterAsset_Static(&FLandscapeEditorCustomNodeBuilder_TargetLayers::ShouldFilterLayerInfo, Target->LayerName)
-							.AllowCreate(false)
-							.AllowClear(false)
-						]
-						+ SHorizontalBox::Slot()
-						.AutoWidth()
-						.VAlign(VAlign_Center)
-						[
-							SNew(SComboButton)
-							.ButtonStyle( FAppStyle::Get(), "HoverHintOnly" )
-							.HasDownArrow(false)
-							.ContentPadding(4.0f)
-							.ForegroundColor(FSlateColor::UseForeground())
-							.IsFocusable(false)
-							.ToolTipText(LOCTEXT("Tooltip_Create", "Create Layer Info"))
-							.IsEnabled_Static(&FLandscapeEditorCustomNodeBuilder_TargetLayers::GetTargetLayerCreateEnabled, Target)
-							.OnGetMenuContent_Static(&FLandscapeEditorCustomNodeBuilder_TargetLayers::OnGetTargetLayerCreateMenu, Target)
-							.ButtonContent()
-							[
-								SNew(SImage)
-								.Image(FAppStyle::GetBrush("LandscapeEditor.Target_Create"))
-							]
-						]
-						+ SHorizontalBox::Slot()
-						.AutoWidth()
-						.VAlign(VAlign_Center)
-						[
-							SNew(SButton)
-							.ButtonStyle( FAppStyle::Get(), "HoverHintOnly" )
-							.ContentPadding(4.0f)
-							.ForegroundColor(FSlateColor::UseForeground())
-							.IsFocusable(false)
-							.ToolTipText(LOCTEXT("Tooltip_MakePublic", "Make Layer Public (move layer info into asset file)"))
-							.Visibility_Static(&FLandscapeEditorCustomNodeBuilder_TargetLayers::GetTargetLayerMakePublicVisibility, Target)
-							.OnClicked_Static(&FLandscapeEditorCustomNodeBuilder_TargetLayers::OnTargetLayerMakePublicClicked, Target)
+						SNew(SBox)
+							.Padding(FMargin(2.0f, 0.0f, 2.0f, 0.0f))
 							[
 								SNew(SImage)
-								.Image(FAppStyle::GetBrush("LandscapeEditor.Target_MakePublic"))
+									.Image(FCoreStyle::Get().GetBrush("VerticalBoxDragIndicator"))
 							]
-						]
-						+ SHorizontalBox::Slot()
-						.AutoWidth()
-						.VAlign(VAlign_Center)
-						[
-							SNew(SButton)
-							.ButtonStyle( FAppStyle::Get(), "HoverHintOnly" )
-							.ContentPadding(4.0f)
-							.ForegroundColor(FSlateColor::UseForeground())
-							.IsFocusable(false)
-							.ToolTipText(LOCTEXT("Tooltip_Delete", "Delete Layer"))
-							.Visibility_Static(&FLandscapeEditorCustomNodeBuilder_TargetLayers::GetTargetLayerDeleteVisibility, Target)
-							.OnClicked_Static(&FLandscapeEditorCustomNodeBuilder_TargetLayers::OnTargetLayerDeleteClicked, Target)
+					]
+
+					+ SHorizontalBox::Slot()
+					.AutoWidth()
+					.VAlign(VAlign_Center)
+					.Padding(FMargin(2))
+					[
+						SNew(SBox)
+							.Visibility_Static(&FLandscapeEditorCustomNodeBuilder_TargetLayers::GetDebugModeLayerUsageVisibility, Target)
+							.WidthOverride(48.0f)
+							.HeightOverride(48.0f)
 							[
 								SNew(SImage)
-								.Image(FAppStyle::GetBrush("LandscapeEditor.Target_Delete"))
+									.Image(FCoreStyle::Get().GetBrush("WhiteBrush"))
+									.ColorAndOpacity_Static(&FLandscapeEditorCustomNodeBuilder_TargetLayers::GetLayerUsageDebugColor, Target)
 							]
-						]
 					]
-					+ SVerticalBox::Slot()
-					.AutoHeight()
+
+					+ SHorizontalBox::Slot()
+					.AutoWidth()
+					.VAlign(VAlign_Center)
+					.Padding(FMargin(2))
 					[
-						SNew(SHorizontalBox)
-						.Visibility_Static(&FLandscapeEditorCustomNodeBuilder_TargetLayers::GetLayersSubstractiveBlendVisibility, Target)
-						+ SHorizontalBox::Slot()
-						.AutoWidth()
-						.Padding(0, 2, 2, 2)
-						[
-							SNew(SCheckBox)
-							.IsChecked_Static(&FLandscapeEditorCustomNodeBuilder_TargetLayers::IsLayersSubstractiveBlendChecked, Target)
-							.OnCheckStateChanged_Static(&FLandscapeEditorCustomNodeBuilder_TargetLayers::OnLayersSubstractiveBlendChanged, Target)
-							[
-								SNew(STextBlock)
-								.Text(LOCTEXT("SubtractiveBlend", "Subtractive Blend"))
-								.ColorAndOpacity_Static(&FLandscapeEditorCustomNodeBuilder_TargetLayers::GetTargetTextColor, Target)
-							]
-						]
+						(Target->bValid)
+							? (TSharedRef<SWidget>)(
+								SNew(SLandscapeAssetThumbnail, Target->ThumbnailMIC.Get(), ThumbnailPool)
+								.Visibility_Static(&FLandscapeEditorCustomNodeBuilder_TargetLayers::GetDebugModeLayerUsageVisibility_Invert, Target)
+								.ThumbnailSize(FIntPoint(48, 48))
+								// Open landscape layer info asset on double-click on the thumbnail : 
+								.OnAccessAsset_Lambda([Target](UObject* InObject)
+								{ 
+									// Note : the object being returned here is the landscape MIC so it's not what we use for opening the landscape layer info asset : 
+									if ((Target->TargetType == ELandscapeToolTargetType::Weightmap) && (Target->LayerInfoObj != nullptr))
+									{
+										UAssetEditorSubsystem* AssetEditorSubsystem = GEditor->GetEditorSubsystem<UAssetEditorSubsystem>();
+										return AssetEditorSubsystem->OpenEditorForAsset(Target->LayerInfoObj.Get());
+									}
+									return false;
+								}))
+							: (TSharedRef<SWidget>)(
+								SNew(SImage)
+								.Visibility_Static(&FLandscapeEditorCustomNodeBuilder_TargetLayers::GetDebugModeLayerUsageVisibility_Invert, Target)
+								.Image(FAppStyle::GetBrush(TEXT("LandscapeEditor.Target_Invalid")))
+								)
 					]
-					+ SVerticalBox::Slot()
-					.AutoHeight()
+					+ SHorizontalBox::Slot()
+					.VAlign(VAlign_Center)
+					.Padding(4, 0)
 					[
-						SNew(SHorizontalBox)
-						.Visibility_Static(&FLandscapeEditorCustomNodeBuilder_TargetLayers::GetDebugModeColorChannelVisibility, Target)
-						+ SHorizontalBox::Slot()
-						.AutoWidth()
-						.Padding(0, 2, 2, 2)
-						[
-							SNew(SCheckBox)
-							.IsChecked_Static(&FLandscapeEditorCustomNodeBuilder_TargetLayers::DebugModeColorChannelIsChecked, Target, 0)
-							.OnCheckStateChanged_Static(&FLandscapeEditorCustomNodeBuilder_TargetLayers::OnDebugModeColorChannelChanged, Target, 0)
+						SNew(SVerticalBox)
+							+ SVerticalBox::Slot()
+							.AutoHeight()
+							.VAlign(VAlign_Center)
+							.Padding(0, 2, 0, 0)
 							[
-								SNew(STextBlock)
-								.Text(LOCTEXT("ViewMode.Debug_None", "None"))
+								SNew(SHorizontalBox)
+									+ SHorizontalBox::Slot()
+									[
+										SNew(STextBlock)
+											.Font(IDetailLayoutBuilder::GetDetailFont())
+											.Text(Target->TargetName)
+											.ColorAndOpacity_Static(&FLandscapeEditorCustomNodeBuilder_TargetLayers::GetTargetTextColor, Target)
+									]
+									+ SHorizontalBox::Slot()
+									.HAlign(HAlign_Right)
+									[
+										SNew(STextBlock)
+											.Visibility_Lambda([=] { return (Target->LayerInfoObj.IsValid() && Target->LayerInfoObj->bNoWeightBlend) ? EVisibility::Visible : EVisibility::Collapsed; })
+											.Font(IDetailLayoutBuilder::GetDetailFont())
+											.Text(LOCTEXT("NoWeightBlend", "No Weight-Blend"))
+											.ColorAndOpacity_Static(&FLandscapeEditorCustomNodeBuilder_TargetLayers::GetTargetTextColor, Target)
+									]
 							]
-						]
-						+ SHorizontalBox::Slot()
-						.AutoWidth()
-						.Padding(2)
-						[
-							SNew(SCheckBox)
-							.IsChecked_Static(&FLandscapeEditorCustomNodeBuilder_TargetLayers::DebugModeColorChannelIsChecked, Target, 1)
-							.OnCheckStateChanged_Static(&FLandscapeEditorCustomNodeBuilder_TargetLayers::OnDebugModeColorChannelChanged, Target, 1)
+							+ SVerticalBox::Slot()
+							.AutoHeight()
+							.VAlign(VAlign_Center)
 							[
-								SNew(STextBlock)
-								.Text(LOCTEXT("ViewMode.Debug_R", "R"))
+								SNew(SHorizontalBox)
+									.Visibility_Static(&FLandscapeEditorCustomNodeBuilder_TargetLayers::GetTargetLayerInfoSelectorVisibility, Target)
+									+ SHorizontalBox::Slot()
+									.FillWidth(1)
+									.VAlign(VAlign_Center)
+									[
+										SNew(SObjectPropertyEntryBox)
+											.IsEnabled((bool)Target->bValid)
+											.ObjectPath(Target->LayerInfoObj != NULL ? Target->LayerInfoObj->GetPathName() : FString())
+											.AllowedClass(ULandscapeLayerInfoObject::StaticClass())
+											.OnObjectChanged_Static(&FLandscapeEditorCustomNodeBuilder_TargetLayers::OnTargetLayerSetObject, Target)
+											.OnShouldFilterAsset_Static(&FLandscapeEditorCustomNodeBuilder_TargetLayers::ShouldFilterLayerInfo, Target->LayerName)
+											.AllowCreate(false)
+											.AllowClear(false)
+									]
+									+ SHorizontalBox::Slot()
+									.AutoWidth()
+									.VAlign(VAlign_Center)
+									[
+										SNew(SComboButton)
+											.ButtonStyle(FAppStyle::Get(), "HoverHintOnly")
+											.HasDownArrow(false)
+											.ContentPadding(4.0f)
+											.ForegroundColor(FSlateColor::UseForeground())
+											.IsFocusable(false)
+											.ToolTipText(LOCTEXT("Tooltip_Create", "Create Layer Info"))
+											.IsEnabled_Static(&FLandscapeEditorCustomNodeBuilder_TargetLayers::GetTargetLayerCreateEnabled, Target)
+											.OnGetMenuContent_Static(&FLandscapeEditorCustomNodeBuilder_TargetLayers::OnGetTargetLayerCreateMenu, Target)
+											.ButtonContent()
+											[
+												SNew(SImage)
+													.Image(FAppStyle::GetBrush("LandscapeEditor.Target_Create"))
+											]
+									]
+									+ SHorizontalBox::Slot()
+									.AutoWidth()
+									.VAlign(VAlign_Center)
+									[
+										SNew(SButton)
+											.ButtonStyle(FAppStyle::Get(), "HoverHintOnly")
+											.ContentPadding(4.0f)
+											.ForegroundColor(FSlateColor::UseForeground())
+											.IsFocusable(false)
+											.ToolTipText(LOCTEXT("Tooltip_MakePublic", "Make Layer Public (move layer info into asset file)"))
+											.Visibility_Static(&FLandscapeEditorCustomNodeBuilder_TargetLayers::GetTargetLayerMakePublicVisibility, Target)
+											.OnClicked_Static(&FLandscapeEditorCustomNodeBuilder_TargetLayers::OnTargetLayerMakePublicClicked, Target)
+											[
+												SNew(SImage)
+													.Image(FAppStyle::GetBrush("LandscapeEditor.Target_MakePublic"))
+											]
+									]
+									+ SHorizontalBox::Slot()
+									.AutoWidth()
+									.VAlign(VAlign_Center)
+									[
+										SNew(SButton)
+											.ButtonStyle(FAppStyle::Get(), "HoverHintOnly")
+											.ContentPadding(4.0f)
+											.ForegroundColor(FSlateColor::UseForeground())
+											.IsFocusable(false)
+											.ToolTipText(LOCTEXT("Tooltip_Delete", "Delete Layer"))
+											.Visibility_Static(&FLandscapeEditorCustomNodeBuilder_TargetLayers::GetTargetLayerDeleteVisibility, Target)
+											.OnClicked_Static(&FLandscapeEditorCustomNodeBuilder_TargetLayers::OnTargetLayerDeleteClicked, Target)
+											[
+												SNew(SImage)
+													.Image(FAppStyle::GetBrush("LandscapeEditor.Target_Delete"))
+											]
+									]
 							]
-						]
-						+ SHorizontalBox::Slot()
-						.AutoWidth()
-						.Padding(2)
-						[
-							SNew(SCheckBox)
-							.IsChecked_Static(&FLandscapeEditorCustomNodeBuilder_TargetLayers::DebugModeColorChannelIsChecked, Target, 2)
-							.OnCheckStateChanged_Static(&FLandscapeEditorCustomNodeBuilder_TargetLayers::OnDebugModeColorChannelChanged, Target, 2)
+							+ SVerticalBox::Slot()
+							.AutoHeight()
 							[
-								SNew(STextBlock)
-								.Text(LOCTEXT("ViewMode.Debug_G", "G"))
+								SNew(SHorizontalBox)
+									.Visibility_Static(&FLandscapeEditorCustomNodeBuilder_TargetLayers::GetLayersSubstractiveBlendVisibility, Target)
+									+ SHorizontalBox::Slot()
+									.AutoWidth()
+									.Padding(0, 2, 2, 2)
+									[
+										SNew(SCheckBox)
+											.IsChecked_Static(&FLandscapeEditorCustomNodeBuilder_TargetLayers::IsLayersSubstractiveBlendChecked, Target)
+											.OnCheckStateChanged_Static(&FLandscapeEditorCustomNodeBuilder_TargetLayers::OnLayersSubstractiveBlendChanged, Target)
+											[
+												SNew(STextBlock)
+													.Text(LOCTEXT("SubtractiveBlend", "Subtractive Blend"))
+													.ColorAndOpacity_Static(&FLandscapeEditorCustomNodeBuilder_TargetLayers::GetTargetTextColor, Target)
+											]
+									]
 							]
-						]
-						+ SHorizontalBox::Slot()
-						.AutoWidth()
-						.Padding(2)
-						[
-							SNew(SCheckBox)
-							.IsChecked_Static(&FLandscapeEditorCustomNodeBuilder_TargetLayers::DebugModeColorChannelIsChecked, Target, 4)
-							.OnCheckStateChanged_Static(&FLandscapeEditorCustomNodeBuilder_TargetLayers::OnDebugModeColorChannelChanged, Target, 4)
+							+ SVerticalBox::Slot()
+							.AutoHeight()
 							[
-								SNew(STextBlock)
-								.Text(LOCTEXT("ViewMode.Debug_B", "B"))
+								SNew(SHorizontalBox)
+									.Visibility_Static(&FLandscapeEditorCustomNodeBuilder_TargetLayers::GetDebugModeColorChannelVisibility, Target)
+									+ SHorizontalBox::Slot()
+									.AutoWidth()
+									.Padding(0, 2, 2, 2)
+									[
+										SNew(SCheckBox)
+											.IsChecked_Static(&FLandscapeEditorCustomNodeBuilder_TargetLayers::DebugModeColorChannelIsChecked, Target, 0)
+											.OnCheckStateChanged_Static(&FLandscapeEditorCustomNodeBuilder_TargetLayers::OnDebugModeColorChannelChanged, Target, 0)
+											[
+												SNew(STextBlock)
+													.Text(LOCTEXT("ViewMode.Debug_None", "None"))
+											]
+									]
+									+ SHorizontalBox::Slot()
+									.AutoWidth()
+									.Padding(2)
+									[
+										SNew(SCheckBox)
+											.IsChecked_Static(&FLandscapeEditorCustomNodeBuilder_TargetLayers::DebugModeColorChannelIsChecked, Target, 1)
+											.OnCheckStateChanged_Static(&FLandscapeEditorCustomNodeBuilder_TargetLayers::OnDebugModeColorChannelChanged, Target, 1)
+											[
+												SNew(STextBlock)
+													.Text(LOCTEXT("ViewMode.Debug_R", "R"))
+											]
+									]
+									+ SHorizontalBox::Slot()
+									.AutoWidth()
+									.Padding(2)
+									[
+										SNew(SCheckBox)
+											.IsChecked_Static(&FLandscapeEditorCustomNodeBuilder_TargetLayers::DebugModeColorChannelIsChecked, Target, 2)
+											.OnCheckStateChanged_Static(&FLandscapeEditorCustomNodeBuilder_TargetLayers::OnDebugModeColorChannelChanged, Target, 2)
+											[
+												SNew(STextBlock)
+													.Text(LOCTEXT("ViewMode.Debug_G", "G"))
+											]
+									]
+									+ SHorizontalBox::Slot()
+									.AutoWidth()
+									.Padding(2)
+									[
+										SNew(SCheckBox)
+											.IsChecked_Static(&FLandscapeEditorCustomNodeBuilder_TargetLayers::DebugModeColorChannelIsChecked, Target, 4)
+											.OnCheckStateChanged_Static(&FLandscapeEditorCustomNodeBuilder_TargetLayers::OnDebugModeColorChannelChanged, Target, 4)
+											[
+												SNew(STextBlock)
+													.Text(LOCTEXT("ViewMode.Debug_B", "B"))
+											]
+									]
 							]
-						]
 					]
-				]
 			];
 	}
 
diff --git a/Engine/Source/Editor/LandscapeEditor/Private/LandscapeEditorObject.cpp b/Engine/Source/Editor/LandscapeEditor/Private/LandscapeEditorObject.cpp
--- a/Engine/Source/Editor/LandscapeEditor/Private/LandscapeEditorObject.cpp
+++ b/Engine/Source/Editor/LandscapeEditor/Private/LandscapeEditorObject.cpp
@@ -60,6 +60,7 @@ ULandscapeEditorObject::ULandscapeEditorObject(const FObjectInitializer& ObjectI
 	, ErodeIterationNum(28)
 	, ErosionNoiseMode(ELandscapeToolErosionMode::Lower)
 	, ErosionNoiseScale(60.0f)
+	, bErosionUseLayerHardness(false)
 
 	, RainAmount(128)
 	, SedimentCapacity(0.3f)
@@ -275,6 +276,7 @@ void ULandscapeEditorObject::Load()
 	GConfig->GetInt(TEXT("LandscapeEdit"), TEXT("ErosionNoiseMode"), InErosionNoiseMode, GEditorPerProjectIni);
 	ErosionNoiseMode = (ELandscapeToolErosionMode)InErosionNoiseMode;
 	GConfig->GetFloat(TEXT("LandscapeEdit"), TEXT("ErosionNoiseScale"), ErosionNoiseScale, GEditorPerProjectIni);
+	GConfig->GetBool(TEXT("LandscapeEdit"), TEXT("bErosionUseLayerHardness"), bErosionUseLayerHardness, GEditorPerProjectIni);
 
 	GConfig->GetInt(TEXT("LandscapeEdit"), TEXT("RainAmount"), RainAmount, GEditorPerProjectIni);
 	GConfig->GetFloat(TEXT("LandscapeEdit"), TEXT("SedimentCapacity"), SedimentCapacity, GEditorPerProjectIni);
@@ -425,6 +427,7 @@ void ULandscapeEditorObject::Save()
 	GConfig->SetInt(TEXT("LandscapeEdit"), TEXT("ErodeSurfaceThickness"), ErodeSurfaceThickness, GEditorPerProjectIni);
 	GConfig->SetInt(TEXT("LandscapeEdit"), TEXT("ErosionNoiseMode"), (int32)ErosionNoiseMode, GEditorPerProjectIni);
 	GConfig->SetFloat(TEXT("LandscapeEdit"), TEXT("ErosionNoiseScale"), ErosionNoiseScale, GEditorPerProjectIni);
+	GConfig->SetBool(TEXT("LandscapeEdit"), TEXT("bErosionUseLayerHardness"), bErosionUseLayerHardness, GEditorPerProjectIni);
 
 	GConfig->SetInt(TEXT("LandscapeEdit"), TEXT("RainAmount"), RainAmount, GEditorPerProjectIni);
 	GConfig->SetFloat(TEXT("LandscapeEdit"), TEXT("SedimentCapacity"), SedimentCapacity, GEditorPerProjectIni);
diff --git a/Engine/Source/Editor/LandscapeEditor/Private/SLandscapeEditor.cpp b/Engine/Source/Editor/LandscapeEditor/Private/SLandscapeEditor.cpp
--- a/Engine/Source/Editor/LandscapeEditor/Private/SLandscapeEditor.cpp
+++ b/Engine/Source/Editor/LandscapeEditor/Private/SLandscapeEditor.cpp
@@ -27,6 +27,7 @@ void SLandscapeAssetThumbnail::Construct(const FArguments& InArgs, UObject* Asse
 	FIntPoint ThumbnailSize = InArgs._ThumbnailSize;
 
 	AssetThumbnail = MakeShareable(new FAssetThumbnail(Asset, ThumbnailSize.X, ThumbnailSize.Y, ThumbnailPool));
+	OnAccessAsset = InArgs._OnAccessAsset;
 
 	ChildSlot
 	[
@@ -68,6 +69,21 @@ void SLandscapeAssetThumbnail::SetAsset(UObject* Asset)
 	AssetThumbnail->SetAsset(Asset);
 }
 
+FReply SLandscapeAssetThumbnail::OnMouseButtonDoubleClick(const FGeometry& InMyGeometry, const FPointerEvent& InMouseEvent)
+{
+	if (InMyGeometry.IsUnderLocation(InMouseEvent.GetScreenSpacePosition()))
+	{
+		if (InMouseEvent.GetEffectingButton() == EKeys::LeftMouseButton && OnAccessAsset.IsBound())
+		{
+			if (OnAccessAsset.Execute(AssetThumbnail->GetAsset()))
+			{
+				return FReply::Handled();
+			}
+		}
+	}
+	return FReply::Unhandled();
+}
+
 //////////////////////////////////////////////////////////////////////////
 
 
diff --git a/Engine/Source/Editor/LandscapeEditor/Private/SLandscapeEditor.h b/Engine/Source/Editor/LandscapeEditor/Private/SLandscapeEditor.h
--- a/Engine/Source/Editor/LandscapeEditor/Private/SLandscapeEditor.h
+++ b/Engine/Source/Editor/LandscapeEditor/Private/SLandscapeEditor.h
@@ -10,6 +10,7 @@
 #include "Toolkits/IToolkitHost.h"
 #include "AssetThumbnail.h"
 #include "Toolkits/BaseToolkit.h"
+#include "Framework/SlateDelegates.h"
 
 class Error;
 class IDetailsView;
@@ -25,17 +26,20 @@ class SLandscapeAssetThumbnail : public SCompoundWidget
 public:
 	SLATE_BEGIN_ARGS( SLandscapeAssetThumbnail )
 		: _ThumbnailSize( 64,64 ) {}
-		SLATE_ARGUMENT( FIntPoint, ThumbnailSize )
+		SLATE_ARGUMENT(FIntPoint, ThumbnailSize)
+		SLATE_EVENT(FAccessAsset, OnAccessAsset)
 	SLATE_END_ARGS()
 
 	void Construct(const FArguments& InArgs, UObject* Asset, TSharedRef<FAssetThumbnailPool> ThumbnailPool);
 	~SLandscapeAssetThumbnail();
 
 	void SetAsset(UObject* Asset);
+	virtual FReply OnMouseButtonDoubleClick(const FGeometry& InMyGeometry, const FPointerEvent& InMouseEvent) override;
 
 private:
 	void OnMaterialCompilationFinished(UMaterialInterface* MaterialInterface);
 
+	FAccessAsset OnAccessAsset;
 	TSharedPtr<FAssetThumbnail> AssetThumbnail;
 };
 
diff --git a/Engine/Source/Editor/LandscapeEditor/Public/LandscapeEditorObject.h b/Engine/Source/Editor/LandscapeEditor/Public/LandscapeEditorObject.h
--- a/Engine/Source/Editor/LandscapeEditor/Public/LandscapeEditorObject.h
+++ b/Engine/Source/Editor/LandscapeEditor/Public/LandscapeEditorObject.h
@@ -392,6 +392,10 @@ class ULandscapeEditorObject : public UObject
 	UPROPERTY(Category="Tool Settings", EditAnywhere, NonTransactional, meta=(DisplayName="Noise Scale", ShowForTools="Erosion", ClampMin="1", ClampMax="512", UIMin="1.1", UIMax="256"))
 	float ErosionNoiseScale;
 
+	// Whether the erosion tool should take into account the paint layer's hardness parameter (a hardness of 0 means the layer is fully affected by erosion, while 1 means fully unaffected)
+	UPROPERTY(Category = "Tool Settings", EditAnywhere, NonTransactional, meta = (DisplayName = "Use Layer Hardness", ShowForTools = "Erosion"))
+	bool bErosionUseLayerHardness;
+
 	// Hydraulic Erosion Tool:
 
 	// The amount of rain to apply to the surface. Larger values will result in more erosion
diff --git a/Engine/Source/Runtime/Landscape/Classes/LandscapeLayerInfoObject.h b/Engine/Source/Runtime/Landscape/Classes/LandscapeLayerInfoObject.h
--- a/Engine/Source/Runtime/Landscape/Classes/LandscapeLayerInfoObject.h
+++ b/Engine/Source/Runtime/Landscape/Classes/LandscapeLayerInfoObject.h
@@ -25,45 +25,42 @@ class ULandscapeLayerInfoObject : public UObject
 {
 	GENERATED_UCLASS_BODY()
 
-	UPROPERTY(VisibleAnywhere, Category=LandscapeLayerInfoObject, AssetRegistrySearchable)
+	UPROPERTY(VisibleAnywhere, Category = LandscapeLayerInfoObject, AssetRegistrySearchable)
 	FName LayerName;
 
-	UPROPERTY(EditAnywhere, Category=LandscapeLayerInfoObject)
+	UPROPERTY(EditAnywhere, Category = LandscapeLayerInfoObject, Meta = (DisplayName = "Physical Material", Tooltip = "Physical material to use when this layer is the predominant one at a given location. Note: this is ignored if the Landscape Physical Material node is used in the landscape material. "))
 	TObjectPtr<UPhysicalMaterial> PhysMaterial;
 
-	UPROPERTY(EditAnywhere, Category=LandscapeLayerInfoObject)
+	UPROPERTY(EditAnywhere, Category = LandscapeLayerInfoObject, Meta = (ClampMin = "0", ClampMax = "1", Tooltip = "Defines how much 'resistance' areas painted with this layer will offer to the Erosion tool. A hardness of 0 means the layer is fully affected by erosion, while 1 means fully unaffected."))
 	float Hardness;
 
 #if WITH_EDITORONLY_DATA
-	/* The minimum weight that needs to be painted for that layer to be picked up as the dominant physical layer */
-	UPROPERTY(EditAnywhere, Category=LandscapeLayerInfoObject, Meta = (ClampMin="0", ClampMax="1"))
+	UPROPERTY(EditAnywhere, Category = LandscapeLayerInfoObject, Meta = (ClampMin = "0", ClampMax = "1", Tooltip = "The minimum weight that needs to be painted for that layer to be picked up as the dominant physical layer."))
 	float MinimumCollisionRelevanceWeight;
 
-	UPROPERTY(EditAnywhere, Category=LandscapeLayerInfoObject)
+	UPROPERTY(EditAnywhere, Category = LandscapeLayerInfoObject, Meta = (Tooltip = "Prevents this layer to be weight-blended with others."))
 	uint32 bNoWeightBlend:1;
 
-	/** Texture to modulate the Splines Falloff Layer Alpha */
-	UPROPERTY(EditAnywhere, Category = SplineFalloffModulation, Meta = (DisplayName="Texture"))
+	UPROPERTY(EditAnywhere, Category = SplineFalloffModulation, Meta = (DisplayName="Texture", Tooltip = "Texture to modulate the Splines Falloff Layer Alpha"))
 	TObjectPtr<UTexture2D> SplineFalloffModulationTexture;
 
-	UPROPERTY(EditAnywhere, Category = SplineFalloffModulation, Meta = (DisplayName = "Color Mask"))
+	UPROPERTY(EditAnywhere, Category = SplineFalloffModulation, Meta = (DisplayName = "Color Mask", Tooltip = "Defines which channel of the Spline Falloff Modulation Texture to use."))
 	ESplineModulationColorMask SplineFalloffModulationColorMask;
 
-	UPROPERTY(EditAnywhere, Category = SplineFalloffModulation, Meta = (DisplayName = "Tiling", ClampMin="0.01"))
+	UPROPERTY(EditAnywhere, Category = SplineFalloffModulation, Meta = (DisplayName = "Tiling", ClampMin = "0.01", Tooltip = "Defines the tiling to use when sampling the Spline Falloff Modulation Texture."))
 	float SplineFalloffModulationTiling;
 
-	UPROPERTY(EditAnywhere, Category = SplineFalloffModulation, Meta = (DisplayName = "Bias", ClampMin="0"))
+	UPROPERTY(EditAnywhere, Category = SplineFalloffModulation, Meta = (DisplayName = "Bias", ClampMin = "0", Tooltip = "Defines the offset to use when sampling the Spline Falloff Modulation Texture."))
 	float SplineFalloffModulationBias;
 
-	UPROPERTY(EditAnywhere, Category = SplineFalloffModulation, Meta = (DisplayName = "Scale", ClampMin="0"))
+	UPROPERTY(EditAnywhere, Category = SplineFalloffModulation, Meta = (DisplayName = "Scale", ClampMin = "0", Tooltip = "Allows to scale the value sampled from the Spline Falloff Modulation Texture."))
 	float SplineFalloffModulationScale;
 				
 	UPROPERTY(NonTransactional, Transient)
 	bool IsReferencedFromLoadedData;
 #endif // WITH_EDITORONLY_DATA
 
-	/* The color to use for layer usage debug */
-	UPROPERTY(EditAnywhere, Category=LandscapeLayerInfoObject)
+	UPROPERTY(EditAnywhere, Category = LandscapeLayerInfoObject, Meta = (Tooltip = "The color to use for layer usage debug"))
 	FLinearColor LayerUsageDebugColor;
 
 #if WITH_EDITOR
diff --git a/Engine/Source/Runtime/Landscape/Private/LandscapeEdit.cpp b/Engine/Source/Runtime/Landscape/Private/LandscapeEdit.cpp
--- a/Engine/Source/Runtime/Landscape/Private/LandscapeEdit.cpp
+++ b/Engine/Source/Runtime/Landscape/Private/LandscapeEdit.cpp
@@ -5302,7 +5302,7 @@ ULandscapeLayerInfoObject::ULandscapeLayerInfoObject(const FObjectInitializer& O
 	, IsReferencedFromLoadedData(false)
 #endif // WITH_EDITORONLY_DATA
 {
-	Hardness = 0.5f;
+	Hardness = 0.0f;
 #if WITH_EDITORONLY_DATA
 	MinimumCollisionRelevanceWeight = 0.0f;
 	bNoWeightBlend = false;
@@ -5332,7 +5332,6 @@ FLinearColor ULandscapeLayerInfoObject::GenerateLayerUsageDebugColor() const
 #if WITH_EDITOR
 void ULandscapeLayerInfoObject::PostEditChangeProperty(FPropertyChangedEvent& PropertyChangedEvent)
 {
-	static const FName NAME_Hardness = GET_MEMBER_NAME_CHECKED(ULandscapeLayerInfoObject, Hardness);
 	static const FName NAME_PhysMaterial = GET_MEMBER_NAME_CHECKED(ULandscapeLayerInfoObject, PhysMaterial);
 	static const FName NAME_LayerUsageDebugColor = GET_MEMBER_NAME_CHECKED(ULandscapeLayerInfoObject, LayerUsageDebugColor);
 	static const FName NAME_MinimumCollisionRelevanceWeight = GET_MEMBER_NAME_CHECKED(ULandscapeLayerInfoObject, MinimumCollisionRelevanceWeight);
@@ -5347,11 +5346,7 @@ void ULandscapeLayerInfoObject::PostEditChangeProperty(FPropertyChangedEvent& Pr
 
 	if (GIsEditor)
 	{
-		if (PropertyName == NAME_Hardness)
-		{
-			Hardness = FMath::Clamp<float>(Hardness, 0.0f, 1.0f);
-		}
-		else if (PropertyName == NAME_PhysMaterial || PropertyName == NAME_MinimumCollisionRelevanceWeight)
+		if (PropertyName == NAME_PhysMaterial || PropertyName == NAME_MinimumCollisionRelevanceWeight)
 		{
 			for (TObjectIterator<ALandscapeProxy> It; It; ++It)
 			{
