Fix for missing or incorrect tool tips in Landscape Mode / Sculpt tab
Adding back the "Reset to default" option to some landscape tool properties

diff --git a/Engine/Source/Editor/LandscapeEditor/Private/LandscapeEditorDetailCustomization_AlphaBrush.cpp b/Engine/Source/Editor/LandscapeEditor/Private/LandscapeEditorDetailCustomization_AlphaBrush.cpp
--- a/Engine/Source/Editor/LandscapeEditor/Private/LandscapeEditorDetailCustomization_AlphaBrush.cpp
+++ b/Engine/Source/Editor/LandscapeEditor/Private/LandscapeEditorDetailCustomization_AlphaBrush.cpp
@@ -200,7 +200,6 @@ void FLandscapeEditorDetailCustomization_AlphaBrush::CustomizeDetails(IDetailLay
 	DetailBuilder.HideProperty(PropertyHandle_AlphaTextureChannel);
 
 	BrushSettingsCategory.AddProperty(PropertyHandle_AlphaTexture)
-	.OverrideResetToDefault(FResetToDefaultOverride::Hide())
 	.CustomWidget()
 	.NameContent()
 	[
@@ -248,8 +247,7 @@ void FLandscapeEditorDetailCustomization_AlphaBrush::CustomizeDetails(IDetailLay
 		]
 	];
 
-	BrushSettingsCategory.AddProperty(PropertyHandle_AlphaTextureChannel)
-	.OverrideResetToDefault(FResetToDefaultOverride::Hide());
+	BrushSettingsCategory.AddProperty(PropertyHandle_AlphaTextureChannel);
 
 	if (IsBrushSetActive("BrushSet_Pattern"))
 	{
diff --git a/Engine/Source/Editor/LandscapeEditor/Private/LandscapeEditorDetailCustomization_MiscTools.cpp b/Engine/Source/Editor/LandscapeEditor/Private/LandscapeEditorDetailCustomization_MiscTools.cpp
--- a/Engine/Source/Editor/LandscapeEditor/Private/LandscapeEditorDetailCustomization_MiscTools.cpp
+++ b/Engine/Source/Editor/LandscapeEditor/Private/LandscapeEditorDetailCustomization_MiscTools.cpp
@@ -39,6 +39,7 @@ void FLandscapeEditorDetailCustomization_MiscTools::CustomizeDetails(IDetailLayo
 		[
 			SNew(SButton)
 			.Text(LOCTEXT("Component.ClearSelection", "Clear Component Selection"))
+			.ToolTipText(LOCTEXT("Component.ClearSelectionToolTip", "Removes all components from the current selection"))
 			.HAlign(HAlign_Center)
 			.OnClicked_Static(&FLandscapeEditorDetailCustomization_MiscTools::OnClearComponentSelectionButtonClicked)
 		];
@@ -51,6 +52,7 @@ void FLandscapeEditorDetailCustomization_MiscTools::CustomizeDetails(IDetailLayo
 		[
 			SNew(SButton)
 			.Text(LOCTEXT("Mask.ClearSelection", "Clear Region Selection"))
+			.ToolTipText(LOCTEXT("Mask.ClearSelectionToolTip", "Removes all painted regions from the current selection"))
 			.HAlign(HAlign_Center)
 			.OnClicked_Static(&FLandscapeEditorDetailCustomization_MiscTools::OnClearRegionSelectionButtonClicked)
 		];
@@ -228,6 +230,7 @@ void FLandscapeEditorDetailCustomization_MiscTools::CustomizeDetails(IDetailLayo
 				[
 					SNew(SButton)
 					.Text(LOCTEXT("Ramp.Reset", "Reset"))
+					.ToolTipText(LOCTEXT("Ramp.ResetToolTip", "Clear the added ramp points"))
 					.HAlign(HAlign_Center)
 					.OnClicked_Static(&FLandscapeEditorDetailCustomization_MiscTools::OnResetRampButtonClicked)
 				]
@@ -237,6 +240,7 @@ void FLandscapeEditorDetailCustomization_MiscTools::CustomizeDetails(IDetailLayo
 					SNew(SButton)
 					.IsEnabled_Static(&FLandscapeEditorDetailCustomization_MiscTools::GetApplyRampButtonIsEnabled)
 					.Text(LOCTEXT("Ramp.Apply", "Add Ramp"))
+					.ToolTipText(LOCTEXT("Ramp.ApplyToolTip", "Applies the current ramp to the height map of the currently selected edit layer"))
 					.HAlign(HAlign_Center)
 					.OnClicked_Static(&FLandscapeEditorDetailCustomization_MiscTools::OnApplyRampButtonClicked)
 				]
@@ -259,6 +263,7 @@ void FLandscapeEditorDetailCustomization_MiscTools::CustomizeDetails(IDetailLayo
 				[
 					SNew(SButton)
 					.Text(LOCTEXT("Mirror.Reset", "Recenter"))
+					.ToolTipText(LOCTEXT("Mirror.ResetToolTip", "Center the mirror point on the current landscape"))
 					.HAlign(HAlign_Center)
 					.OnClicked_Lambda(&FLandscapeEditorDetailCustomization_MiscTools::OnResetMirrorPointButtonClicked)
 				]
@@ -268,6 +273,7 @@ void FLandscapeEditorDetailCustomization_MiscTools::CustomizeDetails(IDetailLayo
 					SNew(SButton)
 					.IsEnabled_Lambda([]() { FEdModeLandscape* LandscapeEdMode = GetEditorMode(); return LandscapeEdMode && LandscapeEdMode->CanEditLayer(); })
 					.Text(LOCTEXT("Mirror.Apply", "Apply"))
+					.ToolTipText(LOCTEXT("Mirror.ApplyToolTip", "Apply the mirror operation to the current landscape edit layer"))
 					.HAlign(HAlign_Center)
 					.OnClicked_Static(&FLandscapeEditorDetailCustomization_MiscTools::OnApplyMirrorButtonClicked)
 				]
diff --git a/Engine/Source/Editor/LandscapeEditor/Private/LandscapeEditorDetails.cpp b/Engine/Source/Editor/LandscapeEditor/Private/LandscapeEditorDetails.cpp
--- a/Engine/Source/Editor/LandscapeEditor/Private/LandscapeEditorDetails.cpp
+++ b/Engine/Source/Editor/LandscapeEditor/Private/LandscapeEditorDetails.cpp
@@ -162,7 +162,7 @@ void FLandscapeEditorDetails::CustomizeDetails(IDetailLayoutBuilder& DetailBuild
 		SNew(STextBlock)
 		.TextStyle( &FCoreStyle::Get().GetWidgetStyle<FTextBlockStyle>( "SmallText" ) )
 		.Text(LOCTEXT("BrushSelector", "Brush Type"))
-
+		.ToolTipText(LOCTEXT("BrushSelectorToolTip", "Selects the type of brush to use"))
 	]
 	.ValueContent()
 	[
@@ -176,6 +176,7 @@ void FLandscapeEditorDetails::CustomizeDetails(IDetailLayoutBuilder& DetailBuild
 		SNew(STextBlock)
 		.TextStyle( &FCoreStyle::Get().GetWidgetStyle<FTextBlockStyle>( "SmallText" ) )
 		.Text(LOCTEXT("BrushFalloff", "Brush Falloff"))
+		.ToolTipText(LOCTEXT("BrushFalloffToolTip", "Selects the profile shape of the brush falloff"))
 
 	]
 	.ValueContent()
@@ -265,7 +266,7 @@ FText FLandscapeEditorDetails::GetLocalizedName(FString Name)
 	}
 
 	FText Result;
-	ensure(FText::FindText(TEXT("LandscapeEditor"), Name, Result));
+	ensure(FText::FindText(TEXT(LOCTEXT_NAMESPACE), Name, Result));
 	return Result;
 }
 
diff --git a/Engine/Source/Editor/LandscapeEditor/Public/LandscapeEditorObject.h b/Engine/Source/Editor/LandscapeEditor/Public/LandscapeEditorObject.h
--- a/Engine/Source/Editor/LandscapeEditor/Public/LandscapeEditorObject.h
+++ b/Engine/Source/Editor/LandscapeEditor/Public/LandscapeEditorObject.h
@@ -301,6 +301,7 @@ class ULandscapeEditorObject : public UObject
 	UPROPERTY(NonTransactional)
 	float MaximumValueRadius;
 
+	// Use the combined result of the underlying layers as input to the operation.  When not checked, it will use only the data in the currently selected layer as input.
 	UPROPERTY(Category="Tool Settings", EditAnywhere, NonTransactional, meta=(ShowForTools="Flatten,Smooth,Erosion,HydraErosion,Ramp", ShowForTargetTypes="Heightmap", ShowForLandscapeLayerSystem))
 	bool bCombinedLayersOperation;
 
@@ -403,14 +404,14 @@ class ULandscapeEditorObject : public UObject
 	int32 RainAmount;
 
 	// The amount of sediment that the water can carry. Larger values will result in more erosion
-	UPROPERTY(Category="Tool Settings", EditAnywhere, NonTransactional, meta=(DisplayName="Sediment Cap.", ShowForTools="HydraErosion", ClampMin="0.1", ClampMax="1.0"))
+	UPROPERTY(Category="Tool Settings", EditAnywhere, NonTransactional, meta=(DisplayName="Sediment Capacity", ShowForTools="HydraErosion", ClampMin="0.1", ClampMax="1.0"))
 	float SedimentCapacity;
 
 	// Number of erosion iterations, more means more erosion but is slower
 	UPROPERTY(Category="Tool Settings", EditAnywhere, NonTransactional, meta=(DisplayName="Iterations", ShowForTools="HydraErosion", ClampMin="1", ClampMax="300", UIMin="1", UIMax="150"))
 	int32 HErodeIterationNum;
 
-	// Initial Rain Distribution
+	// Selects how rain is distributed over the brush area for hydro erosion
 	UPROPERTY(Category="Tool Settings", EditAnywhere, NonTransactional, meta=(DisplayName="Initial Rain Distribution", ShowForTools="HydraErosion"))
 	ELandscapeToolHydroErosionMode RainDistMode;
 
@@ -489,8 +490,7 @@ class ULandscapeEditorObject : public UObject
 	UPROPERTY(Category="Tool Settings", EditAnywhere, NonTransactional, meta=(DisplayName="Smoothing Width", ShowForTools="Mirror", ClampMin="0", UIMin="0", UIMax="20"))
 	int32 MirrorSmoothingWidth;
 
-	// Blueprint Brush Tool
-
+	// Selects the blueprint brush to apply to the current edit layer. Click on the landscape to apply it.
 	UPROPERTY(Category = "Tool Settings", EditAnywhere, Transient, meta = (DisplayName = "Blueprint Brush", ShowForTools = "BlueprintBrush"))
 	TSubclassOf<ALandscapeBlueprintBrushBase> BlueprintBrush;
 
@@ -660,6 +660,7 @@ public:
 	UPROPERTY(Category="Brush Settings", EditAnywhere, NonTransactional, meta=(DisplayName="Texture Pan V", ShowForBrushes="BrushSet_Pattern", ClampMin="0", ClampMax="1"))
 	float AlphaBrushPanV;
 
+	// Tile the pattern in world space (this only takes scale and position into account, not rotation)
 	UPROPERTY(Category = "Brush Settings", EditAnywhere, NonTransactional, meta = (DisplayName = "Use World-Space", ShowForBrushes = "BrushSet_Pattern"))
 	bool bUseWorldSpacePatternBrush;
 
