Resize Landscape: Use transform when calculating new position of Landscape after resize.

diff --git a/Engine/Source/Editor/LandscapeEditor/Private/LandscapeEdMode.cpp b/Engine/Source/Editor/LandscapeEditor/Private/LandscapeEdMode.cpp
--- a/Engine/Source/Editor/LandscapeEditor/Private/LandscapeEdMode.cpp
+++ b/Engine/Source/Editor/LandscapeEditor/Private/LandscapeEdMode.cpp
@@ -3538,7 +3538,7 @@ ALandscape* FEdModeLandscape::ChangeComponentSetting(int32 NumComponentsX, int32
 					NewMaxY = NewMinY + NewVertsY - 1;
 
 					// offset landscape to component boundary
-					LandscapeOffset = FVector(NewMinX, NewMinY, 0) * OldLandscape->GetActorScale();
+					LandscapeOffset = OldLandscape->GetActorTransform().TransformVector(FVector(NewMinX, NewMinY, 0));
 					LandscapeOffsetQuads = FIntPoint(NewMinX, NewMinY);
 				}
 
diff --git a/Engine/Source/Editor/LandscapeEditor/Private/LandscapeEditorDetailCustomization_NewLandscape.cpp b/Engine/Source/Editor/LandscapeEditor/Private/LandscapeEditorDetailCustomization_NewLandscape.cpp
--- a/Engine/Source/Editor/LandscapeEditor/Private/LandscapeEditorDetailCustomization_NewLandscape.cpp
+++ b/Engine/Source/Editor/LandscapeEditor/Private/LandscapeEditorDetailCustomization_NewLandscape.cpp
@@ -1064,8 +1064,8 @@ FReply FLandscapeEditorDetailCustomization_NewLandscape::OnCreateButtonClicked()
 	
 	const FIntPoint TotalLandscapeComponentSize { UISettings->NewLandscape_ComponentCount.X, UISettings->NewLandscape_ComponentCount.Y };
 
-	const int32 ComponentCountX = bNeedsLandscapeRegions ? UISettings->WorldPartitionRegionSize : TotalLandscapeComponentSize.X;
-	const int32 ComponentCountY = bNeedsLandscapeRegions ? UISettings->WorldPartitionRegionSize : TotalLandscapeComponentSize.Y;
+	const int32 ComponentCountX = bNeedsLandscapeRegions ? FMath::Min(static_cast<int32>(UISettings->WorldPartitionRegionSize), TotalLandscapeComponentSize.X) : TotalLandscapeComponentSize.X;
+	const int32 ComponentCountY = bNeedsLandscapeRegions ? FMath::Min(static_cast<int32>(UISettings->WorldPartitionRegionSize), TotalLandscapeComponentSize.Y) : TotalLandscapeComponentSize.Y;
 	const int32 QuadsPerComponent = UISettings->NewLandscape_SectionsPerComponent * QuadsPerSection;
 	const int32 SizeX = ComponentCountX * QuadsPerComponent + 1;
 	const int32 SizeY = ComponentCountY * QuadsPerComponent + 1;
