Landscape: Make it possible to drag landscape brushes to different edit layers in landscape mode.

diff --git a/Engine/Source/Editor/LandscapeEditor/Private/LandscapeEditorDetailCustomization_Layers.cpp b/Engine/Source/Editor/LandscapeEditor/Private/LandscapeEditorDetailCustomization_Layers.cpp
--- a/Engine/Source/Editor/LandscapeEditor/Private/LandscapeEditorDetailCustomization_Layers.cpp
+++ b/Engine/Source/Editor/LandscapeEditor/Private/LandscapeEditorDetailCustomization_Layers.cpp
@@ -21,6 +21,7 @@
 #include "Widgets/Input/SMultiLineEditableTextBox.h"
 #include "EditorModeManager.h"
 #include "EditorModes.h"
+#include "LandscapeEditorDetailCustomization_LayersBrushStack.h" // FLandscapeBrushDragDropOp
 #include "LandscapeEditorModule.h"
 #include "LandscapeEditorObject.h"
 #include "Landscape.h"
@@ -915,22 +916,61 @@ FReply FLandscapeEditorCustomNodeBuilder_Layers::HandleAcceptDrop(FDragDropEvent
 {
 	TSharedPtr<FLandscapeListElementDragDropOp> DragDropOperation = DragDropEvent.GetOperationAs<FLandscapeListElementDragDropOp>();
 
-	if (DragDropOperation.IsValid())
+	if (!DragDropOperation.IsValid())
+	{
+		return FReply::Unhandled();
+	}
+
+	FEdModeLandscape* LandscapeEdMode = GetEditorMode();
+	ALandscape* Landscape = LandscapeEdMode ? LandscapeEdMode->GetLandscape() : nullptr;
+	if (!Landscape)
+	{
+		return FReply::Unhandled();
+	}
+
+	// See if we're actually getting a drag from the blueprint brush list, rather than
+	// from the edit layer list
+	if (DragDropOperation->IsOfType<FLandscapeBrushDragDropOp>())
 	{
-		FEdModeLandscape* LandscapeEdMode = GetEditorMode();
-		ALandscape* Landscape = LandscapeEdMode ? LandscapeEdMode->GetLandscape() : nullptr;
-		if (Landscape)
+		int32 StartingBrushIndex = DragDropOperation->SlotIndexBeingDragged;
+		int32 StartingLayerIndex = LandscapeEdMode->GetCurrentLayerIndex();
+		int32 DestinationLayerIndex = SlotIndexToLayerIndex(SlotIndex);
+
+		if (StartingLayerIndex == DestinationLayerIndex)
 		{
-			int32 StartingLayerIndex = SlotIndexToLayerIndex(DragDropOperation->SlotIndexBeingDragged);
-			int32 DestinationLayerIndex = SlotIndexToLayerIndex(SlotIndex);
-			const FScopedTransaction Transaction(LOCTEXT("Landscape_Layers_Reorder", "Reorder Layer"));
-			if (Landscape->ReorderLayer(StartingLayerIndex, DestinationLayerIndex))
-			{
-				LandscapeEdMode->SetCurrentLayer(DestinationLayerIndex);
-				LandscapeEdMode->RefreshDetailPanel();
-				return FReply::Handled();
-			}
+			// See comment further below about not returning Handled()
+			return FReply::Unhandled();
+		}
+
+		ALandscapeBlueprintBrushBase* Brush = Landscape->GetBrushForLayer(StartingLayerIndex, StartingBrushIndex);
+		if (!ensure(Brush))
+		{
+			return FReply::Unhandled();
 		}
+
+		const FScopedTransaction Transaction(LOCTEXT("Landscape_LayerBrushes_MoveLayers", "Move Brush to Layer"));
+		Landscape->RemoveBrushFromLayer(StartingLayerIndex, StartingBrushIndex);
+		Landscape->AddBrushToLayer(DestinationLayerIndex, Brush);
+
+		LandscapeEdMode->SetCurrentLayer(DestinationLayerIndex);
+		LandscapeEdMode->RefreshDetailPanel();
+
+		// HACK: We don't return FReply::Handled() here because otherwise, SDragAndDropVerticalBox::OnDrop
+		// will apply UI slot reordering after we return. Properly speaking, we should have a way to signal 
+		// that the operation was handled yet that it is not one that SDragAndDropVerticalBox should deal with.
+		// For now, however, just make sure to return Unhandled.
+		return FReply::Unhandled();
+	}
+
+	// This must be a drag from our own list.
+	int32 StartingLayerIndex = SlotIndexToLayerIndex(DragDropOperation->SlotIndexBeingDragged);
+	int32 DestinationLayerIndex = SlotIndexToLayerIndex(SlotIndex);
+	const FScopedTransaction Transaction(LOCTEXT("Landscape_Layers_Reorder", "Reorder Layer"));
+	if (Landscape->ReorderLayer(StartingLayerIndex, DestinationLayerIndex))
+	{
+		LandscapeEdMode->SetCurrentLayer(DestinationLayerIndex);
+		LandscapeEdMode->RefreshDetailPanel();
+		return FReply::Handled();
 	}
 
 	return FReply::Unhandled();
diff --git a/Engine/Source/Editor/LandscapeEditor/Private/LandscapeEditorDetailCustomization_LayersBrushStack.cpp b/Engine/Source/Editor/LandscapeEditor/Private/LandscapeEditorDetailCustomization_LayersBrushStack.cpp
--- a/Engine/Source/Editor/LandscapeEditor/Private/LandscapeEditorDetailCustomization_LayersBrushStack.cpp
+++ b/Engine/Source/Editor/LandscapeEditor/Private/LandscapeEditorDetailCustomization_LayersBrushStack.cpp
@@ -509,7 +509,7 @@ FReply FLandscapeEditorCustomNodeBuilder_LayersBrushStack::HandleDragDetected(co
 				TSharedPtr<SWidget> Row = GenerateRow(SlotIndex);
 				if (Row.IsValid())
 				{
-					return FReply::Handled().BeginDragDrop(FLandscapeListElementDragDropOp::New(SlotIndex, Slot, Row));
+					return FReply::Handled().BeginDragDrop(FLandscapeBrushDragDropOp::New(SlotIndex, Slot, Row));
 				}
 			}
 		}
@@ -519,7 +519,7 @@ FReply FLandscapeEditorCustomNodeBuilder_LayersBrushStack::HandleDragDetected(co
 
 TOptional<SDragAndDropVerticalBox::EItemDropZone> FLandscapeEditorCustomNodeBuilder_LayersBrushStack::HandleCanAcceptDrop(const FDragDropEvent& DragDropEvent, SDragAndDropVerticalBox::EItemDropZone DropZone, SVerticalBox::FSlot* Slot)
 {
-	TSharedPtr<FLandscapeListElementDragDropOp> DragDropOperation = DragDropEvent.GetOperationAs<FLandscapeListElementDragDropOp>();
+	TSharedPtr<FLandscapeBrushDragDropOp> DragDropOperation = DragDropEvent.GetOperationAs<FLandscapeBrushDragDropOp>();
 
 	if (DragDropOperation.IsValid())
 	{
@@ -531,7 +531,7 @@ TOptional<SDragAndDropVerticalBox::EItemDropZone> FLandscapeEditorCustomNodeBuil
 
 FReply FLandscapeEditorCustomNodeBuilder_LayersBrushStack::HandleAcceptDrop(FDragDropEvent const& DragDropEvent, SDragAndDropVerticalBox::EItemDropZone DropZone, int32 SlotIndex, SVerticalBox::FSlot* Slot)
 {
-	TSharedPtr<FLandscapeListElementDragDropOp> DragDropOperation = DragDropEvent.GetOperationAs<FLandscapeListElementDragDropOp>();
+	TSharedPtr<FLandscapeBrushDragDropOp> DragDropOperation = DragDropEvent.GetOperationAs<FLandscapeBrushDragDropOp>();
 
 	if (DragDropOperation.IsValid())
 	{
@@ -539,10 +539,10 @@ FReply FLandscapeEditorCustomNodeBuilder_LayersBrushStack::HandleAcceptDrop(FDra
 		ALandscape* Landscape = LandscapeEdMode ? LandscapeEdMode->GetLandscape() : nullptr;
 		if (Landscape)
 		{
-			int32 StartingLayerIndex = DragDropOperation->SlotIndexBeingDragged;
-			int32 DestinationLayerIndex = SlotIndex;
+			int32 StartingBrushIndex = DragDropOperation->SlotIndexBeingDragged;
+			int32 DestinationBrushIndex = SlotIndex;
 			const FScopedTransaction Transaction(LOCTEXT("Landscape_LayerBrushes_Reorder", "Reorder Layer Brush"));
-			if (Landscape->ReorderLayerBrush(LandscapeEdMode->GetCurrentLayerIndex(), StartingLayerIndex, DestinationLayerIndex))
+			if (Landscape->ReorderLayerBrush(LandscapeEdMode->GetCurrentLayerIndex(), StartingBrushIndex, DestinationBrushIndex))
 			{
 				LandscapeEdMode->RefreshDetailPanel();
 				return FReply::Handled();
@@ -553,4 +553,21 @@ FReply FLandscapeEditorCustomNodeBuilder_LayersBrushStack::HandleAcceptDrop(FDra
 	return FReply::Unhandled();
 }
 
+TSharedRef<FLandscapeBrushDragDropOp> FLandscapeBrushDragDropOp::New(int32 InSlotIndexBeingDragged, 
+	SVerticalBox::FSlot* InSlotBeingDragged, TSharedPtr<SWidget> InWidgetToShow)
+{
+	// Same code as FLandscapeListElementDragDropOp, just returns a different type.
+
+	TSharedRef<FLandscapeBrushDragDropOp> Operation = MakeShared<FLandscapeBrushDragDropOp>();
+
+	Operation->MouseCursor = EMouseCursor::GrabHandClosed;
+	Operation->SlotIndexBeingDragged = InSlotIndexBeingDragged;
+	Operation->SlotBeingDragged = InSlotBeingDragged;
+	Operation->WidgetToShow = InWidgetToShow;
+
+	Operation->Construct();
+
+	return Operation;
+}
+
 #undef LOCTEXT_NAMESPACE
\ No newline at end of file
diff --git a/Engine/Source/Editor/LandscapeEditor/Private/LandscapeEditorDetailCustomization_LayersBrushStack.h b/Engine/Source/Editor/LandscapeEditor/Private/LandscapeEditorDetailCustomization_LayersBrushStack.h
--- a/Engine/Source/Editor/LandscapeEditor/Private/LandscapeEditorDetailCustomization_LayersBrushStack.h
+++ b/Engine/Source/Editor/LandscapeEditor/Private/LandscapeEditorDetailCustomization_LayersBrushStack.h
@@ -11,6 +11,7 @@
 #include "Widgets/DeclarativeSyntaxSupport.h"
 #include "Widgets/SWidget.h"
 #include "Widgets/Layout/SBorder.h"
+#include "LandscapeEditorDetailCustomization_Layers.h" // FLandscapeListElementDragDropOp
 #include "LandscapeEdMode.h"
 #include "IDetailCustomNodeBuilder.h"
 #include "IDetailCustomization.h"
@@ -93,3 +94,15 @@ protected:
 	TOptional<SDragAndDropVerticalBox::EItemDropZone> HandleCanAcceptDrop(const FDragDropEvent& DragDropEvent, SDragAndDropVerticalBox::EItemDropZone DropZone, SVerticalBox::FSlot* Slot);
 	FReply HandleAcceptDrop(FDragDropEvent const& DragDropEvent, SDragAndDropVerticalBox::EItemDropZone DropZone, int32 SlotIndex, SVerticalBox::FSlot* Slot);
 };
+
+class FLandscapeBrushDragDropOp : public FLandscapeListElementDragDropOp
+{
+public:
+	DRAG_DROP_OPERATOR_TYPE(FLandscapeBrushDragDropOp, FLandscapeListElementDragDropOp)
+
+	static TSharedRef<FLandscapeBrushDragDropOp> New(int32 InSlotIndexBeingDragged, 
+		SVerticalBox::FSlot* InSlotBeingDragged, TSharedPtr<SWidget> InWidgetToShow);
+
+public:
+	virtual ~FLandscapeBrushDragDropOp() {}
+};
