Changed "Mode" property in landscape import/export to either "Import Mode" or "Export Mode" to be more explicit to the user

diff --git a/Engine/Source/Editor/LandscapeEditor/Private/LandscapeEditorDetailCustomization_ImportExport.cpp b/Engine/Source/Editor/LandscapeEditor/Private/LandscapeEditorDetailCustomization_ImportExport.cpp
--- a/Engine/Source/Editor/LandscapeEditor/Private/LandscapeEditorDetailCustomization_ImportExport.cpp
+++ b/Engine/Source/Editor/LandscapeEditor/Private/LandscapeEditorDetailCustomization_ImportExport.cpp
@@ -185,6 +185,32 @@ void FLandscapeEditorDetailCustomization_ImportExport::CustomizeDetails(IDetailL
 	TSharedPtr<SToolTip> ExportSingleFileTooltip = SNew(SToolTip)
 		.Text(PropertyHandle_ExportSingleFile->GetToolTipText());
 
+	TSharedRef<IPropertyHandle> PropertyHandle_ImportExportMode = DetailBuilder.GetProperty(GET_MEMBER_NAME_CHECKED(ULandscapeEditorObject, ImportExportMode));
+	TSharedPtr<SWidget> NameWidget;
+	TSharedPtr<SWidget> ValueWidget;
+	IDetailPropertyRow& ImportExportModeRow = ImportExportCategory.AddProperty(PropertyHandle_ImportExportMode);
+	ImportExportModeRow.GetDefaultWidgets(NameWidget, ValueWidget);
+	ImportExportModeRow.CustomWidget()
+		.NameContent()
+		[
+			SNew(SBox)
+				.VAlign(VAlign_Center)
+				.Padding(FMargin(0, 0, 2, 0))
+				[
+					SNew(STextBlock)
+						.Text(MakeAttributeLambda([]() { return IsImporting() ? LOCTEXT("ImportMode", "Import Mode") : LOCTEXT("ExportMode", "Export Mode"); }))
+						.Font(DetailBuilder.GetDetailFont())
+						.ToolTipText(MakeAttributeLambda([]() { return IsImporting() 
+							? LOCTEXT("ImportModeTooltip", "Specifies whether all or only loaded landscape regions should be imported") 
+							: LOCTEXT("ExportModeToolTip", "Specifies whether all or only loaded landscape regions should be imported"); }))
+				]
+		]
+		.ValueContent()
+		[
+			ValueWidget.ToSharedRef()
+		];
+
+
 	ImportExportCategory.AddProperty(PropertyHandle_ExportSingleFile)
 		.Visibility(MakeAttributeLambda([]() { return FLandscapeEditorDetailCustomization_ImportExport::GetImportExportVisibility(false); }))
 		.IsEnabled(MakeAttributeLambda([]() { return FLandscapeEditorDetailCustomization_ImportExport::GetExportSingleFileIsEnabled(); }))
diff --git a/Engine/Source/Editor/LandscapeEditor/Public/LandscapeEditorObject.h b/Engine/Source/Editor/LandscapeEditor/Public/LandscapeEditorObject.h
--- a/Engine/Source/Editor/LandscapeEditor/Public/LandscapeEditorObject.h
+++ b/Engine/Source/Editor/LandscapeEditor/Public/LandscapeEditorObject.h
@@ -581,10 +581,10 @@ class ULandscapeEditorObject : public UObject
 	UPROPERTY(Category = "Import / Export", EditAnywhere, NonTransactional, meta = (DisplayName="Export Selected Edit Layer", ShowForTools = "ImportExport", ToolTip="When true exports the selected edit layer, if false exports the blended result"))
 	bool bExportEditLayer = true;
 
-	UPROPERTY(Category = "Import / Export", EditAnywhere, NonTransactional, meta = (DisplayName = "Export Single File", ShowForTools = "ImportExport", ToolTip = "(World Partition only) When true, exports the landscape as a single file, if false exports each grid tile individually."))
+	UPROPERTY(Category = "Import / Export", EditAnywhere, NonTransactional, meta = (ShowForTools = "ImportExport", ToolTip = "(World Partition only) When true, exports the landscape as a single file, if false exports each grid tile individually."))
 	bool bExportSingleFile = false;
 
-	UPROPERTY(Category = "Import / Export", EditAnywhere, NonTransactional, meta = (DisplayName = "Mode", ShowForTools = "ImportExport", ToolTip = "Import Loaded or All Landscape Regions"))
+	UPROPERTY(Category = "Import / Export", EditAnywhere, NonTransactional, meta = (ShowForTools = "ImportExport"))
 	ELandscapeImportExportMode ImportExportMode = ELandscapeImportExportMode::LoadedOnly;
 
 	UPROPERTY(NonTransactional)
