Re-select the affected landscape after using the resize tool in landscape mode
#tests editor

diff --git a/Engine/Source/Editor/LandscapeEditor/Private/LandscapeEditorDetailCustomization_ResizeLandscape.cpp b/Engine/Source/Editor/LandscapeEditor/Private/LandscapeEditorDetailCustomization_ResizeLandscape.cpp
--- a/Engine/Source/Editor/LandscapeEditor/Private/LandscapeEditorDetailCustomization_ResizeLandscape.cpp
+++ b/Engine/Source/Editor/LandscapeEditor/Private/LandscapeEditorDetailCustomization_ResizeLandscape.cpp
@@ -504,9 +504,10 @@ FReply FLandscapeEditorDetailCustomization_ResizeLandscape::OnApplyButtonClicked
 		const int32 SectionsPerComponent = LandscapeEdMode->UISettings->ResizeLandscape_SectionsPerComponent;
 		const int32 QuadsPerSection = LandscapeEdMode->UISettings->ResizeLandscape_QuadsPerSection;
 		const bool bResample = (LandscapeEdMode->UISettings->ResizeLandscape_ConvertMode == ELandscapeConvertMode::Resample);
-		LandscapeEdMode->ChangeComponentSetting(ComponentCount.X, ComponentCount.Y, SectionsPerComponent, QuadsPerSection, bResample);
+		ALandscape* NewLandscape = LandscapeEdMode->ChangeComponentSetting(ComponentCount.X, ComponentCount.Y, SectionsPerComponent, QuadsPerSection, bResample);
 
 		LandscapeEdMode->UpdateLandscapeList();
+		LandscapeEdMode->SetTargetLandscape(NewLandscape->GetLandscapeInfo());
 	}
 
 	return FReply::Handled();
