World Partition:
- FEditorLoadedActorCache was reporting unregistered loaded actors and the code that was using it was expecting registered actors.
- Replaced FEditorLoadedActorCache by FWorldPartitionHelpers::GetLoadedActorGuidsForLevel that gathers registered actors GUIDs and updated callers.
- UWorldPartition::Tick is now private to avoid code like LandscapeEditorUtils::SaveLandscapeProxies doing manual ticks to get around unwanted internal behaviors.
- TExternalDirtyActorsTracker can now filter out tracked actors depending on the parent class wanted behavior.
- Don't track dirty actors coming from non transactional operations.

diff --git a/Engine/Source/Editor/DataLayerEditor/Private/DataLayer/DataLayerHierarchy.cpp b/Engine/Source/Editor/DataLayerEditor/Private/DataLayer/DataLayerHierarchy.cpp
--- a/Engine/Source/Editor/DataLayerEditor/Private/DataLayer/DataLayerHierarchy.cpp
+++ b/Engine/Source/Editor/DataLayerEditor/Private/DataLayer/DataLayerHierarchy.cpp
@@ -33,7 +33,6 @@
 #include "WorldPartition/WorldPartitionHelpers.h"
 #include "WorldPartition/WorldPartitionSubsystem.h"
 #include "WorldPartition/WorldPartitionActorDescInstance.h"
-#include "EditorLoadedActorCache.h"
 
 TUniquePtr<FDataLayerHierarchy> FDataLayerHierarchy::Create(FDataLayerMode* Mode, const TWeakObjectPtr<UWorld>& World)
 {
@@ -277,9 +276,8 @@ void FDataLayerHierarchy::CreateItems(TArray<FSceneOutlinerTreeItemPtr>& OutItem
 				const UDataLayerManager* DataLayerManager = OwningWorld->GetDataLayerManager();
 				if (WorldPartitionSubsystem && DataLayerManager)
 				{
-					FEditorLoadedActorCache LoadedActorCache;
 					const ULevelInstanceSubsystem* LevelInstanceSubsystem = UWorld::GetSubsystem<ULevelInstanceSubsystem>(OwningWorld);
-					WorldPartitionSubsystem->ForEachWorldPartition([this, DataLayerManager, CurrentLevel, LevelInstanceSubsystem, IsDataLayerShown, &WorldToLevelDataLayerMap, &LoadedActorCache, &OutItems](UWorldPartition* WorldPartition)
+					WorldPartitionSubsystem->ForEachWorldPartition([this, DataLayerManager, CurrentLevel, LevelInstanceSubsystem, IsDataLayerShown, &WorldToLevelDataLayerMap, &OutItems](UWorldPartition* WorldPartition)
 					{
 						// Skip WorldPartition if it's not the one of the current level (the editing level instance)
 						// or if we hide the content of level instances
@@ -291,7 +289,7 @@ void FDataLayerHierarchy::CreateItems(TArray<FSceneOutlinerTreeItemPtr>& OutItem
 							return true;
 						}
 
-						const TSet<FGuid>& LoadedActors = LoadedActorCache.GetLoadedActorsForLevel(OuterLevel);
+						const TSet<FGuid> LoadedActors =  FWorldPartitionHelpers::GetLoadedActorGuidsForLevel(OuterLevel);
 
 						// Create an FDataLayerActorDescTreeItem for each unloaded actor of this WorldPartition
 						FWorldPartitionHelpers::ForEachActorDescInstance(WorldPartition, [this, IsDataLayerShown, DataLayerManager, CurrentLevel, &WorldToLevelDataLayerMap, &LoadedActors, &OutItems](const FWorldPartitionActorDescInstance* ActorDescInstance)
diff --git a/Engine/Source/Editor/LandscapeEditor/Private/LandscapeEditorDetailCustomization_NewLandscape.cpp b/Engine/Source/Editor/LandscapeEditor/Private/LandscapeEditorDetailCustomization_NewLandscape.cpp
--- a/Engine/Source/Editor/LandscapeEditor/Private/LandscapeEditorDetailCustomization_NewLandscape.cpp
+++ b/Engine/Source/Editor/LandscapeEditor/Private/LandscapeEditorDetailCustomization_NewLandscape.cpp
@@ -62,7 +62,6 @@
 #include "ActorFactories/ActorFactory.h"
 
 #include "SourceControlHelpers.h"
-#include "WorldPartition/WorldPartition.h"
 
 #include "Misc/ScopedSlowTask.h"
 
@@ -1187,7 +1186,6 @@ FReply FLandscapeEditorDetailCustomization_NewLandscape::OnCreateButtonClicked()
 		LandscapeProxy->OnMaterialChangedDelegate().AddRaw(LandscapeEdMode, &FEdModeLandscape::OnLandscapeMaterialChangedDelegate);
 	}
 
-	UWorldPartition* WorldPartition = World->GetWorldPartition();
 	ULandscapeSubsystem* LandscapeSubsystem = World->GetSubsystem<ULandscapeSubsystem>();
 	ALandscapeProxy* LandscapeProxy = LandscapeEdMode->CurrentToolTarget.LandscapeInfo->GetLandscapeProxy();
 	
@@ -1214,7 +1212,7 @@ FReply FLandscapeEditorDetailCustomization_NewLandscape::OnCreateButtonClicked()
 
 		TArray<ALocationVolume*> RegionVolumes;
 		FBox LandscapeBounds;
-		auto AddComponentsToRegion = [&Progress, NumRegions, bIsNewLandscape, LandscapeEdMode, WorldPartition, World, LandscapeProxy, &UISettings, QuadsPerSection, LandscapeInfo, LandscapeSubsystem, &RegionVolumes, &LandscapeBounds,&MaterialLayerDataPerLayers](const FIntPoint& RegionCoordinate, const TArray<FIntPoint>& NewComponents)
+		auto AddComponentsToRegion = [&Progress, NumRegions, bIsNewLandscape, LandscapeEdMode, World, LandscapeProxy, &UISettings, QuadsPerSection, LandscapeInfo, LandscapeSubsystem, &RegionVolumes, &LandscapeBounds,&MaterialLayerDataPerLayers](const FIntPoint& RegionCoordinate, const TArray<FIntPoint>& NewComponents)
 		{
 			TRACE_CPUPROFILER_EVENT_SCOPE(AddComponentsToRegion);
 						
@@ -1250,7 +1248,7 @@ FReply FLandscapeEditorDetailCustomization_NewLandscape::OnCreateButtonClicked()
 
 			// ensures all the final height textures have been updated.
 			LandscapeInfo->ForceLayersFullUpdate();
-			LandscapeEditorUtils::SaveLandscapeProxies(MakeArrayView(CreatedStreamingProxies), WorldPartition);
+			LandscapeEditorUtils::SaveLandscapeProxies(MakeArrayView(CreatedStreamingProxies));
 			LandscapeBounds += LandscapeInfo->GetCompleteBounds();
 
 			Progress.EnterProgressFrame(1.0f , FText::Format(LOCTEXT("LandscapeCreateRegion", "Creating Landscape Editor Regions ({0}, {1})"), RegionCoordinate.X, RegionCoordinate.Y));
@@ -1275,7 +1273,7 @@ FReply FLandscapeEditorDetailCustomization_NewLandscape::OnCreateButtonClicked()
 		TArray<ALandscapeProxy*> AllProxies;
 		
 		// save the initial region & unload it
-		LandscapeInfo->ForEachLandscapeProxy([&WorldPartition, &AllProxies](ALandscapeProxy* Proxy) {
+		LandscapeInfo->ForEachLandscapeProxy([&AllProxies](ALandscapeProxy* Proxy) {
 			if (Proxy->IsA<ALandscapeStreamingProxy>())
 			{		
 				AllProxies.Add(Proxy);
@@ -1284,7 +1282,7 @@ FReply FLandscapeEditorDetailCustomization_NewLandscape::OnCreateButtonClicked()
 			return true;
 		});
 
-		LandscapeEditorUtils::SaveLandscapeProxies(MakeArrayView(AllProxies), WorldPartition);
+		LandscapeEditorUtils::SaveLandscapeProxies(MakeArrayView(AllProxies));
 	}
 
 	return FReply::Handled();
diff --git a/Engine/Source/Editor/LandscapeEditor/Private/LandscapeEditorUtils.cpp b/Engine/Source/Editor/LandscapeEditor/Private/LandscapeEditorUtils.cpp
--- a/Engine/Source/Editor/LandscapeEditor/Private/LandscapeEditorUtils.cpp
+++ b/Engine/Source/Editor/LandscapeEditor/Private/LandscapeEditorUtils.cpp
@@ -14,7 +14,9 @@
 #include "EditorModeManager.h"
 #include "EditorModes.h"
 
+#include "Algo/Transform.h"
 #include "WorldPartition/WorldPartition.h"
+#include "WorldPartition/WorldPartitionHelpers.h"
 #include "Framework/Application/SlateApplication.h"
 #include "Misc/MessageDialog.h"
 
@@ -85,23 +87,17 @@ namespace LandscapeEditorUtils
 
 	}
 
-	void SaveLandscapeProxies(TArrayView<ALandscapeProxy*> Proxies, UWorldPartition* WorldPartition)
+	void SaveLandscapeProxies(TArrayView<ALandscapeProxy*> Proxies)
 	{
-		// Save the Proxies a side effect of this is storing all the Saved Actors in the WP DirtyActors list
+		// Save the proxies
 		{
 			TRACE_CPUPROFILER_EVENT_SCOPE(SaveCreatedActors);
 			LandscapeEditorUtils::SaveObjects(Proxies);
 		}
-		
-		// Hold References to the Proxies we're going to save, this prevents WorldPartition Pinning these Actors on Tick
-		TArray< FWorldPartitionReference> References;
-		for (ALandscapeProxy* Proxy : Proxies)
-		{
-			References.Add(FWorldPartitionReference(WorldPartition, Proxy->GetActorGuid()));
-		}
 
-		// Clear the Dirty Actors in WP
-		WorldPartition->Tick(0.0f);
+		// Grab references to proxies so they get unloaded after this function returns
+		TArray<FWorldPartitionReference> ProxyReferences;
+		Algo::Transform(Proxies, ProxyReferences, [](ALandscapeProxy* Proxy) { return FWorldPartitionReference(FWorldPartitionHelpers::GetWorldPartition(Proxy), Proxy->GetActorGuid()); });
 	}
 }
 
diff --git a/Engine/Source/Editor/LandscapeEditor/Private/LandscapeRegionUtils.cpp b/Engine/Source/Editor/LandscapeEditor/Private/LandscapeRegionUtils.cpp
--- a/Engine/Source/Editor/LandscapeEditor/Private/LandscapeRegionUtils.cpp
+++ b/Engine/Source/Editor/LandscapeEditor/Private/LandscapeRegionUtils.cpp
@@ -4,8 +4,6 @@
 
 #include "Landscape.h"
 
-#include "WorldPartition/WorldPartition.h"
-
 #include "SourceControlHelpers.h"
 #include "UObject/SavePackage.h"
 #include "Builders/CubeBuilder.h"
@@ -100,9 +98,6 @@ void ForEachRegion_LoadProcessUnload(ULandscapeInfo* InLandscapeInfo, const FInt
 	const int32 RegionSizeInTexels = InLandscapeInfo->ComponentSizeQuads * InLandscapeInfo->RegionSizeInComponents + 1;
 	const FIntRect RegionCoordinates = InDomain / RegionSizeInTexels;
 
-	UWorldPartition* WorldPartition = InWorld->GetWorldPartition();
-	check(WorldPartition != nullptr);
-
 	TArray<AActor*> Children;
 	InLandscapeInfo->LandscapeActor->GetAttachedActors(Children);
 	TArray<ALocationVolume*> LandscapeRegions;
@@ -136,7 +131,7 @@ void ForEachRegion_LoadProcessUnload(ULandscapeInfo* InLandscapeInfo, const FInt
 
 		InLandscapeInfo->ForceLayersFullUpdate();
 
-		LandscapeEditorUtils::SaveLandscapeProxies(MakeArrayView(LandscapeProxies), WorldPartition);
+		LandscapeEditorUtils::SaveLandscapeProxies(MakeArrayView(LandscapeProxies));
 
 		Region->Unload();
 
diff --git a/Engine/Source/Editor/LandscapeEditor/Public/LandscapeEditorUtils.h b/Engine/Source/Editor/LandscapeEditor/Public/LandscapeEditorUtils.h
--- a/Engine/Source/Editor/LandscapeEditor/Public/LandscapeEditorUtils.h
+++ b/Engine/Source/Editor/LandscapeEditor/Public/LandscapeEditorUtils.h
@@ -8,7 +8,6 @@
 class ALandscapeProxy;
 class ALandscapeStreamingProxy;
 class ULandscapeLayerInfoObject;
-class UWorldPartition;
 
 namespace LandscapeEditorUtils
 {
@@ -26,5 +25,5 @@ namespace LandscapeEditorUtils
 		UEditorLoadingAndSavingUtils::SavePackages(Packages, /* bOnlyDirty = */ false);
 	}
 
-	void SaveLandscapeProxies(TArrayView<ALandscapeProxy*> Proxies, UWorldPartition* WorldPartition);
+	void SaveLandscapeProxies(TArrayView<ALandscapeProxy*> Proxies);
 }
diff --git a/Engine/Source/Editor/SceneOutliner/Private/ActorHierarchy.cpp b/Engine/Source/Editor/SceneOutliner/Private/ActorHierarchy.cpp
--- a/Engine/Source/Editor/SceneOutliner/Private/ActorHierarchy.cpp
+++ b/Engine/Source/Editor/SceneOutliner/Private/ActorHierarchy.cpp
@@ -27,7 +27,6 @@
 #include "WorldPartition/WorldPartitionActorDesc.h"
 #include "WorldPartition/IWorldPartitionEditorModule.h"
 #include "WorldPartition/WorldPartitionSubsystem.h"
-#include "EditorLoadedActorCache.h"
 #include "ActorFolder.h"
 #include "ActorMode.h"
 
@@ -503,11 +502,9 @@ void FActorHierarchy::CreateUnloadedItems(UWorld* World, TArray<FSceneOutlinerTr
 	{
 		if (UWorldPartitionSubsystem* WorldPartitionSubsystem = UWorld::GetSubsystem<UWorldPartitionSubsystem>(World))
 		{
-			FEditorLoadedActorCache LoadedActorCache;
-			
 			const ULevelInstanceSubsystem* LevelInstanceSubsystem = World->GetSubsystem<ULevelInstanceSubsystem>();
 
-			WorldPartitionSubsystem->ForEachWorldPartition([this, LevelInstanceSubsystem, &LoadedActorCache, &OutItems, &AddChildContainer](UWorldPartition* WorldPartition)
+			WorldPartitionSubsystem->ForEachWorldPartition([this, LevelInstanceSubsystem, &OutItems, &AddChildContainer](UWorldPartition* WorldPartition)
 			{
 				UWorld* OuterWorld = WorldPartition->GetTypedOuter<UWorld>();
 				ULevel* OuterLevel = OuterWorld ? OuterWorld->PersistentLevel : nullptr;
@@ -522,8 +519,8 @@ void FActorHierarchy::CreateUnloadedItems(UWorld* World, TArray<FSceneOutlinerTr
 					}
 				}
 
-				const TSet<FGuid>& LoadedActors = LoadedActorCache.GetLoadedActorsForLevel(OuterLevel);
-				
+				const TSet<FGuid> LoadedActors = FWorldPartitionHelpers::GetLoadedActorGuidsForLevel(OuterLevel);
+
 				FWorldPartitionHelpers::ForEachActorDescInstance(WorldPartition, [this, &LoadedActors, &OutItems, &AddChildContainer, LevelInstanceSubsystem](const FWorldPartitionActorDescInstance* ActorDescInstance)
 				{
 					if (ActorDescInstance != nullptr && !LoadedActors.Contains(ActorDescInstance->GetGuid()) && FActorDescTreeItem::ShouldDisplayInOutliner(ActorDescInstance))
diff --git a/Engine/Source/Runtime/Engine/Private/EditorLoadedActorCache.cpp b/Engine/Source/Runtime/Engine/Private/EditorLoadedActorCache.cpp
deleted file mode 100644
diff --git a/Engine/Source/Runtime/Engine/Private/WorldPartition/WorldPartitionHelpers.cpp b/Engine/Source/Runtime/Engine/Private/WorldPartition/WorldPartitionHelpers.cpp
--- a/Engine/Source/Runtime/Engine/Private/WorldPartition/WorldPartitionHelpers.cpp
+++ b/Engine/Source/Runtime/Engine/Private/WorldPartition/WorldPartitionHelpers.cpp
@@ -324,6 +324,12 @@ bool FWorldPartitionHelpers::FixupRedirectedAssetPath(FName& InOutAssetPath)
 	return true;
 }
 
+TSet<FGuid> FWorldPartitionHelpers::GetLoadedActorGuidsForLevel(const ULevel* InLevel)
+{
+	TSet<FGuid> Result;
+	Algo::TransformIf(InLevel->Actors, Result, [](const AActor* Actor) { return IsValid(Actor); }, [](const AActor* Actor) { return Actor->GetActorGuid(); });
+	return MoveTemp(Result);
+}
 #endif // #if WITH_EDITOR
 
 bool FWorldPartitionHelpers::ConvertEditorPathToRuntimePath(const FSoftObjectPath& InPath, FSoftObjectPath& OutPath)
diff --git a/Engine/Source/Runtime/Engine/Public/EditorLoadedActorCache.h b/Engine/Source/Runtime/Engine/Public/EditorLoadedActorCache.h
deleted file mode 100644
diff --git a/Engine/Source/Runtime/Engine/Public/ExternalDirtyActorsTracker.h b/Engine/Source/Runtime/Engine/Public/ExternalDirtyActorsTracker.h
--- a/Engine/Source/Runtime/Engine/Public/ExternalDirtyActorsTracker.h
+++ b/Engine/Source/Runtime/Engine/Public/ExternalDirtyActorsTracker.h
@@ -36,6 +36,8 @@ public:
 
 	const MapType& GetDirtyActors() const { return DirtyActors; }
 
+	/** Returns false if tracking this actor in unnecessary */
+	virtual bool OnAddDirtyActor(const TWeakObjectPtr<AActor> InActor) { return true; }
 	virtual void OnRemoveInvalidDirtyActor(const TWeakObjectPtr<AActor> InActor, typename StoreType::Type& InValue) {}
 	virtual void OnRemoveNonDirtyActor(const TWeakObjectPtr<AActor> InActor, typename StoreType::Type& InValue) {}
 
@@ -50,7 +52,10 @@ protected:
 				{
 					if (InPackage->IsDirty())
 					{
-						DirtyActors.Add(Actor, StoreType::Store(Owner, Actor));
+						if (OnAddDirtyActor(Actor))
+						{
+							DirtyActors.Add(Actor, StoreType::Store(Owner, Actor));
+						}
 					}
 					else
 					{
diff --git a/Engine/Source/Runtime/Engine/Public/WorldPartition/WorldPartition.h b/Engine/Source/Runtime/Engine/Public/WorldPartition/WorldPartition.h
--- a/Engine/Source/Runtime/Engine/Public/WorldPartition/WorldPartition.h
+++ b/Engine/Source/Runtime/Engine/Public/WorldPartition/WorldPartition.h
@@ -459,7 +459,10 @@ public:
 
 	ENGINE_API bool IsMainWorldPartition() const;
 
+private:
 	ENGINE_API void Tick(float DeltaSeconds);
+
+public:
 	ENGINE_API bool CanAddCellToWorld(const IWorldPartitionCell* InCell) const;
 	ENGINE_API bool IsStreamingCompleted(const TArray<FWorldPartitionStreamingSource>* InStreamingSources) const;
 	ENGINE_API bool IsStreamingCompleted(EWorldPartitionRuntimeCellState QueryState, const TArray<FWorldPartitionStreamingQuerySource>& QuerySources, bool bExactState) const;
@@ -596,6 +599,7 @@ private:
 		FWorldPartitionExternalDirtyActorsTracker(UWorldPartition* InWorldPartition);
 
 		//~ Begin TExternalDirtyActorsTracker interface
+		virtual bool OnAddDirtyActor(const TWeakObjectPtr<AActor> InActor) override { return !!GUndo; }
 		virtual void OnRemoveNonDirtyActor(const TWeakObjectPtr<AActor> InActor, FWorldPartitionReference& InValue) override;
 		virtual void Tick(float InDeltaTime) override;
 		//~ End TExternalDirtyActorsTracker interface
diff --git a/Engine/Source/Runtime/Engine/Public/WorldPartition/WorldPartitionHelpers.h b/Engine/Source/Runtime/Engine/Public/WorldPartition/WorldPartitionHelpers.h
--- a/Engine/Source/Runtime/Engine/Public/WorldPartition/WorldPartitionHelpers.h
+++ b/Engine/Source/Runtime/Engine/Public/WorldPartition/WorldPartitionHelpers.h
@@ -177,6 +177,8 @@ public:
 	
 	UE_DEPRECATED(5.4, "Use UAssetRegistryHelpers::FixupRedirectedAssetPath instead.")
 	static ENGINE_API bool FixupRedirectedAssetPath(FSoftObjectPath& InOutSoftObjectPath);
+
+	static ENGINE_API TSet<FGuid> GetLoadedActorGuidsForLevel(const ULevel* InLevel);
 #endif // WITH_EDITOR
 
 	// Editor/Runtime conversions
