Introducing Landscape Edit Layer class : essentially, a type of object describing landscape edit layers that can be extended by derivation in order to specialize the behavior of the edit layer class. Most edit layer settings should eventually move here and this will be the building block for allowing plugins to extend the edit layer system in a more solid way than the BP brush system.
* 2 implementations for now : ULandscapeEditLayer (standard edit layer, that can be sculpted/painted) and ULandscapeEditLayerSplines (replaces the current "Reserved for Splines" edit layer system)
* The edit layers have a bunch of virtual functions to customize their behavior (affects heightmaps/weightmaps/visibility? contains persistent textures? can be edited via standard tools? supports collapsing? etc.)
* A composition pattern was used : FLandscapeLayer owns an Instanced ULandscapeEditLayerBase property (null not allowed) although the public "LandscapeLayers" list was deprecated in favor of a private "LandscapeEditLayers" in order to stop users being able to mutate the list of layers and their properties freely
* ALandscape::GetLayer was deprecated in favor of ALandscape::GetLayerConst in order to enforce a single, controlled, access to the layer data (via public methods), instead of the current free-for-all way

Landscape UX changes/improvements :
* Landscape edit layers are now added via a more standard '+' button (like standard array properties) and will then request the user to pick a class of edit layer (standard or splines for now, but any non-abstract edit layer class is discoverable) via a modal dialog. The type of the edit layer cannot be changed after the fact (the "Reserve For Splines" / "Remove Reserve for Splines" actions are therefore gone)
* Deletion of edit layer is done via a button or via context menu
* Edit layers context menu will now display all available options (instead of hiding some) but will disable the ones that are not available and will display the reason why that's the case in the tooltip (e.g. cannot rename a locked edit layer)
* The Collapse Layer action has been made available all the time (instead of only when the Manage panel is active) and now also supports collapsing the splines layer onto a standard layer underneath (but not the other way around), because it's supported, since the Splines layer does have standard heightmaps/weightmaps
* The edit layers have the ability to define custom actions generically, which will dynamically extend the context menu on that edit layer (the "Update Splines" actions on the Splines edit layer uses this system)
* Added a drag indicator on edit layers to streamline this reorderable list

Left to do:
* Move most settings from FLandscapeLayer to ULandscapeEditLayer and implement a generic system for responding to property changes in both the landscape and landscape UI, instead of exposing several ad-hoc functions (SetLayerName, SetLayerLocked, etc.) on the landscape and the landscape UI code (plus, it will allow to get those to react to changes to settings defined in an plugin's edit layer)
* Prevent useless heightmaps/weightmaps from being created on procedural edit layers that don't need persistent textures (water, landscape patch, etc.) This is a memory optimization : at least, we already prevent those to be edited, which is a common mistake with the current BP brush system, as users can sculpt/paint on a water edit layer, since it's a standard edit layer and they're allowed to
* Move the systems that act as proper edit layers (LandscapePatchManager, WaterBrushManager) from BP brush to the landscape edit layer class system
* Integrate proper icons for edit layer classes to help the user tell at a glance the type of layer (e.g. splines layer)
#tests Editor

diff --git a/Engine/Plugins/Enterprise/DatasmithImporter/Source/DatasmithImporter/Private/DatasmithLandscapeImporter.cpp b/Engine/Plugins/Enterprise/DatasmithImporter/Source/DatasmithImporter/Private/DatasmithLandscapeImporter.cpp
--- a/Engine/Plugins/Enterprise/DatasmithImporter/Source/DatasmithImporter/Private/DatasmithLandscapeImporter.cpp
+++ b/Engine/Plugins/Enterprise/DatasmithImporter/Source/DatasmithImporter/Private/DatasmithLandscapeImporter.cpp
@@ -81,7 +81,7 @@ AActor* FDatasmithLandscapeImporter::ImportLandscapeActor( const TSharedRef< IDa
 		MaterialLayerDataPerLayer.Add(FGuid(), TArray<FLandscapeImportLayerInfo>());
 		
 		Cast< ALandscape >( NewActor )->Import( FGuid::NewGuid(), 0, 0, SizeX - 1, SizeY - 1, OutSectionsPerComponent, OutQuadsPerSection,
-			HeightmapDataPerLayers, nullptr, MaterialLayerDataPerLayer, ELandscapeImportAlphamapType::Additive);
+			HeightmapDataPerLayers, nullptr, MaterialLayerDataPerLayer, ELandscapeImportAlphamapType::Additive, /*InImportLayers = */MakeArrayView<FLandscapeLayer>({}));
 	} ) );
 
 	LandscapeActorElement->SetTranslation( OriginalTranslation );
diff --git a/Engine/Plugins/Experimental/LandscapePatch/Source/LandscapePatch/Private/LandscapePatchComponent.cpp b/Engine/Plugins/Experimental/LandscapePatch/Source/LandscapePatch/Private/LandscapePatchComponent.cpp
--- a/Engine/Plugins/Experimental/LandscapePatch/Source/LandscapePatch/Private/LandscapePatchComponent.cpp
+++ b/Engine/Plugins/Experimental/LandscapePatch/Source/LandscapePatch/Private/LandscapePatchComponent.cpp
@@ -26,7 +26,7 @@ namespace LandscapePatchComponentLocals
 
 	ALandscapePatchManager* FindExistingPatchManagerForLandscape(ALandscape* Landscape)
 	{
-		for (int LayerIndex = 0; LayerIndex < Landscape->LandscapeLayers.Num(); ++LayerIndex)
+		for (int LayerIndex = 0; LayerIndex < Landscape->GetLayers().Num(); ++LayerIndex)
 		{
 			TArray<ALandscapeBlueprintBrushBase*> LayerBrushes = Landscape->GetBrushesForLayer(LayerIndex);
 			ALandscapeBlueprintBrushBase** Found = LayerBrushes.FindByPredicate(
diff --git a/Engine/Plugins/Experimental/Water/Source/Editor/Private/WaterEditorModule.cpp b/Engine/Plugins/Experimental/Water/Source/Editor/Private/WaterEditorModule.cpp
--- a/Engine/Plugins/Experimental/Water/Source/Editor/Private/WaterEditorModule.cpp
+++ b/Engine/Plugins/Experimental/Water/Source/Editor/Private/WaterEditorModule.cpp
@@ -203,12 +203,19 @@ void FWaterEditorModule::OnLevelActorAddedToWorld(AActor* Actor)
 					check(IsValid(FoundLandscape));
 
 					bool bHasWaterManager = false;
-					FoundLandscape->ForEachLayer([&bHasWaterManager](FLandscapeLayer& Layer)
+					FoundLandscape->ForEachLayerConst([&bHasWaterManager](const FLandscapeLayer& Layer)
 					{
 						for (const FLandscapeLayerBrush& Brush : Layer.Brushes)
 						{
 							bHasWaterManager |= Cast<AWaterLandscapeBrush>(Brush.GetBrush()) != nullptr;
+							if (bHasWaterManager)
+							{
+								// Stop iterating the moment we've found the brush manager : 
+								return false;
+							}
 						}
+
+						return true;
 					});
 
 					if (!bHasWaterManager)
diff --git a/Engine/Source/Editor/DetailCustomizations/Private/LandscapeUIDetails.cpp b/Engine/Source/Editor/DetailCustomizations/Private/LandscapeUIDetails.cpp
--- a/Engine/Source/Editor/DetailCustomizations/Private/LandscapeUIDetails.cpp
+++ b/Engine/Source/Editor/DetailCustomizations/Private/LandscapeUIDetails.cpp
@@ -109,12 +109,9 @@ void FLandscapeUIDetails::ToggleCanHaveLayersContent(ALandscape* Landscape)
 	if (Landscape->bCanHaveLayersContent)
 	{
 		bool bHasHiddenLayers = false;
-		for (int32 i = 0; i < Landscape->GetLayerCount(); ++i)
+		for (const FLandscapeLayer& Layer : Landscape->GetLayers())
 		{
-			const FLandscapeLayer* Layer = Landscape->GetLayer(i);
-			check(Layer != nullptr);
-
-			if (!Layer->bVisible)
+			if (!Layer.bVisible)
 			{
 				bHasHiddenLayers = true;
 				break;
@@ -122,7 +119,6 @@ void FLandscapeUIDetails::ToggleCanHaveLayersContent(ALandscape* Landscape)
 		}
 				
 		FText Reason;
-
 		if (bHasHiddenLayers)
 		{
 			Reason = LOCTEXT("LandscapeDisableLayers_HiddenLayers", "Are you sure you want to disable the edit layers on this Landscape?\n\nDoing so, will result in losing the data stored for each edit layer, but the current visual output will be kept. Be aware that some edit layers are currently hidden, continuing will result in their data being lost. Undo/redo buffer will also be cleared.");
diff --git a/Engine/Source/Editor/EditorStyle/Private/StarshipStyle.cpp b/Engine/Source/Editor/EditorStyle/Private/StarshipStyle.cpp
--- a/Engine/Source/Editor/EditorStyle/Private/StarshipStyle.cpp
+++ b/Engine/Source/Editor/EditorStyle/Private/StarshipStyle.cpp
@@ -6012,6 +6012,9 @@ void FStarshipEditorStyle::FStyle::SetupClassIconsAndThumbnails()
 			{TEXT("Interface")},
 			{TEXT("KillZVolume")},
 			{TEXT("Landscape")},
+			{TEXT("LandscapeEditLayer")},
+			{TEXT("LandscapeEditLayerBase")},
+			{TEXT("LandscapeEditLayerSplines")},
 			{TEXT("LevelBounds")},
 			{TEXT("LevelInstance")},
 			{TEXT("LevelInstancePivot")},
diff --git a/Engine/Source/Editor/LandscapeEditor/Private/LandscapeEdMode.cpp b/Engine/Source/Editor/LandscapeEditor/Private/LandscapeEdMode.cpp
--- a/Engine/Source/Editor/LandscapeEditor/Private/LandscapeEdMode.cpp
+++ b/Engine/Source/Editor/LandscapeEditor/Private/LandscapeEdMode.cpp
@@ -21,6 +21,7 @@
 #include "LandscapeSubsystem.h"
 #include "LandscapeSettings.h"
 #include "LandscapeTiledImage.h"
+#include "LandscapeEditLayer.h"
 
 #include "EditorSupportDelegates.h"
 #include "ScopedTransaction.h"
@@ -505,10 +506,7 @@ void FEdModeLandscape::Enter()
 			Landscape->OnBlueprintBrushChangedDelegate().AddRaw(this, &FEdModeLandscape::RefreshDetailPanel);
 			if (Landscape->HasLayersContent())
 			{
-				if (Landscape->GetLandscapeSplinesReservedLayer())
-				{
-					Landscape->RequestSplineLayerUpdate();
-				}
+				Landscape->RequestSplineLayerUpdate();
 				Landscape->RequestLayersContentUpdateForceAll();
 			}
 		}
@@ -1605,11 +1603,11 @@ bool FEdModeLandscape::ProcessEditPaste()
 	bool Result = false;
 
 	
-	FLandscapeLayer* SplinesLayer = nullptr;
+	const FLandscapeLayer* SplinesLayer = nullptr;
 	if (CurrentTool == (FLandscapeTool*)SplinesTool)
 	{
 		ALandscape* Landscape = GetLandscape();
-		SplinesLayer = Landscape ? Landscape->GetLandscapeSplinesReservedLayer() : nullptr;
+		SplinesLayer = Landscape ? Landscape->FindLayerOfType(ULandscapeEditLayerSplines::StaticClass()) : nullptr;
 	}
 	FText Reason;
 	if (!CanEditLayer(&Reason, SplinesLayer))
@@ -3619,7 +3617,7 @@ ALandscape* FEdModeLandscape::ChangeComponentSetting(int32 NumComponentsX, int32
 				{
 					int32 HeightCount = 0;
 
-					for (const FLandscapeLayer& OldLayer : OldLandscape->LandscapeLayers)
+					for (const FLandscapeLayer& OldLayer : OldLandscape->GetLayers())
 					{
 						FScopedSetLandscapeEditingLayer Scope(OldLandscape, OldLayer.Guid);
 
@@ -3722,17 +3720,12 @@ ALandscape* FEdModeLandscape::ChangeComponentSetting(int32 NumComponentsX, int32
 			NewLandscape->bUsedForNavigation = OldLandscape->bUsedForNavigation;
 			NewLandscape->MaxPaintedLayersPerComponent = OldLandscape->MaxPaintedLayersPerComponent;
 
-			TArray<FLandscapeLayer>* LandscapeLayers = CanHaveLandscapeLayersContent() ? &OldLandscape->LandscapeLayers : nullptr;
-
-			NewLandscape->Import(FGuid::NewGuid(), NewMinX, NewMinY, NewMaxX, NewMaxY, NumSubsections, SubsectionSizeQuads, HeightDataPerLayers, *OldLandscape->ReimportHeightmapFilePath, ImportMaterialLayerInfosPerLayers, ELandscapeImportAlphamapType::Additive, LandscapeLayers);
-
-			// Find the new layer that corresponds to the splines reserved layer in the original, if any, and setup the new Guid : 
-			if (const FLandscapeLayer* OldSplinesReservedLayer = OldLandscape->GetLandscapeSplinesReservedLayer())
+			TArrayView<const FLandscapeLayer> LandscapeLayers;
+			if (CanHaveLandscapeLayersContent())
 			{
-				const FLandscapeLayer* NewSplinesReservedLayer = NewLandscape->GetLayer(OldSplinesReservedLayer->Name);
-				check(NewSplinesReservedLayer != nullptr);
-				NewLandscape->LandscapeSplinesTargetLayerGuid = NewSplinesReservedLayer->Guid;
-			}
+				LandscapeLayers = OldLandscape->GetLayers();
+			} 
+			NewLandscape->Import(FGuid::NewGuid(), NewMinX, NewMinY, NewMaxX, NewMaxY, NumSubsections, SubsectionSizeQuads, HeightDataPerLayers, *OldLandscape->ReimportHeightmapFilePath, ImportMaterialLayerInfosPerLayers, ELandscapeImportAlphamapType::Additive, LandscapeLayers);
 
 			ULandscapeInfo* NewLandscapeInfo = NewLandscape->GetLandscapeInfo();
 			check(NewLandscapeInfo);
@@ -3935,8 +3928,9 @@ void FEdModeLandscape::SetCurrentLayer(int32 InLayerIndex)
 
 	if (ALandscape* Landscape = GetLandscape())
 	{
-		const FLandscapeLayer* SplineLayer = Landscape->GetLandscapeSplinesReservedLayer();
-		if (SplineLayer != nullptr && SplineLayer == Landscape->GetLayer(InLayerIndex))
+		const FLandscapeLayer* Layer = Landscape->GetLayerConst(InLayerIndex);
+		check((Layer != nullptr) && (Layer->EditLayer != nullptr));
+		if (Layer->EditLayer->IsA<ULandscapeEditLayerSplines>())
 		{
 			SetCurrentToolMode("ToolMode_Manage", false);
 			SetCurrentTool(FName("Splines"));
@@ -3957,15 +3951,15 @@ ALandscape* FEdModeLandscape::GetLandscape() const
 	return CurrentToolTarget.LandscapeInfo.IsValid() ? CurrentToolTarget.LandscapeInfo->LandscapeActor.Get() : nullptr;
 }
 
-FLandscapeLayer* FEdModeLandscape::GetLayer(int32 InLayerIndex) const
+const FLandscapeLayer* FEdModeLandscape::GetLayer(int32 InLayerIndex) const
 {
-	ALandscape* Landscape = GetLandscape();
-	return Landscape ? Landscape->GetLayer(InLayerIndex) : nullptr;
+	const ALandscape* Landscape = GetLandscape();
+	return Landscape ? Landscape->GetLayerConst(InLayerIndex) : nullptr;
 }
 
 FName FEdModeLandscape::GetLayerName(int32 InLayerIndex) const
 {
-	FLandscapeLayer* Layer = GetLayer(InLayerIndex);
+	const FLandscapeLayer* Layer = GetLayer(InLayerIndex);
 	return Layer ? Layer->Name : NAME_None;
 }
 
@@ -4026,6 +4020,20 @@ float FEdModeLandscape::GetLayerAlpha(int32 InLayerIndex) const
 	return 1.0f;
 }
 
+const ULandscapeEditLayerBase* FEdModeLandscape::GetEditLayer(int32 InLayerIndex) const
+{
+	const ALandscape* Landscape = GetLandscape();
+	if (Landscape)
+	{
+		if (const FLandscapeLayer* EditLayer = Landscape->GetLayerConst(InLayerIndex))
+		{
+			return EditLayer->EditLayer;
+		}
+	}
+	return nullptr;
+}
+
+
 void FEdModeLandscape::SetLayerAlpha(int32 InLayerIndex, float InAlpha)
 {
 	ALandscape* Landscape = GetLandscape();
@@ -4040,7 +4048,7 @@ void FEdModeLandscape::SetLayerAlpha(int32 InLayerIndex, float InAlpha)
 
 bool FEdModeLandscape::IsLayerVisible(int32 InLayerIndex) const
 {
-	FLandscapeLayer* Layer = GetLayer(InLayerIndex);
+	const FLandscapeLayer* Layer = GetLayer(InLayerIndex);
 	return Layer ? Layer->bVisible : false;
 }
 
@@ -4171,7 +4179,7 @@ void FEdModeLandscape::SetCurrentLayerSubstractiveBlendStatus(bool InStatus, con
 	}
 }
 
-FLandscapeLayer* FEdModeLandscape::GetCurrentLayer() const
+const FLandscapeLayer* FEdModeLandscape::GetCurrentLayer() const
 {
 	return GetLayer(GetCurrentLayerIndex());
 }
@@ -4182,7 +4190,7 @@ void FEdModeLandscape::AutoUpdateDirtyLandscapeSplines()
 	{
 		// Only auto-update if a layer is reserved for landscape splines
 		ALandscape* Landscape = GetLandscape();
-		if (Landscape && Landscape->GetLandscapeSplinesReservedLayer())
+		if (Landscape)
 		{
 			// TODO : Only update dirty regions
 			Landscape->RequestSplineLayerUpdate();
@@ -4190,12 +4198,12 @@ void FEdModeLandscape::AutoUpdateDirtyLandscapeSplines()
 	}
 }
 
-bool FEdModeLandscape::CanEditLayer(FText* Reason /*=nullptr*/, FLandscapeLayer* InLayer /*= nullptr*/)
+bool FEdModeLandscape::CanEditLayer(FText* Reason /*=nullptr*/, const FLandscapeLayer* InLayer /*= nullptr*/)
 {
 	if (CanHaveLandscapeLayersContent())
 	{
 		ALandscape* Landscape = GetLandscape();
-		FLandscapeLayer* TargetLayer = InLayer ? InLayer : GetCurrentLayer();
+		const FLandscapeLayer* TargetLayer = InLayer ? InLayer : GetCurrentLayer();
 		if (!TargetLayer)
 		{
 			if (Reason)
@@ -4204,7 +4212,9 @@ bool FEdModeLandscape::CanEditLayer(FText* Reason /*=nullptr*/, FLandscapeLayer*
 			}
 			return false;
 		}
-		else if (!TargetLayer->bVisible)
+
+		check(TargetLayer->EditLayer != nullptr);
+		if (!TargetLayer->bVisible)
 		{
 			if (Reason)
 			{
@@ -4222,21 +4232,32 @@ bool FEdModeLandscape::CanEditLayer(FText* Reason /*=nullptr*/, FLandscapeLayer*
 		}
 		else if (CurrentTool)
 		{
-			int32 TargetLayerIndex = Landscape ? Landscape->LandscapeLayers.IndexOfByPredicate([TargetLayeyGuid = TargetLayer->Guid](const FLandscapeLayer& OtherLayer) { return OtherLayer.Guid == TargetLayeyGuid; }) : INDEX_NONE;
-
-			if ((CurrentTool != (FLandscapeTool*)SplinesTool) && Landscape && (TargetLayer == Landscape->GetLandscapeSplinesReservedLayer()))
+			// Special case for the splines tool, which is supported on the splines layer : 
+			if (CurrentTool != (FLandscapeTool*)SplinesTool)
 			{
-				if (Reason)
+				if (Landscape && (TargetLayer == Landscape->FindLayerOfType(ULandscapeEditLayerSplines::StaticClass())))
 				{
-					*Reason = NSLOCTEXT("UnrealEd", "LandscapeLayerReservedForSplines", "This layer is reserved for Landscape Splines.");
+					if (Reason)
+					{
+						*Reason = NSLOCTEXT("UnrealEd", "LandscapeLayerReservedForSplines", "This layer is reserved for Landscape Splines.");
+					}
+					return false;
+				}
+
+				if (!TargetLayer->EditLayer->SupportsEditingTools())
+				{
+					if (Reason)
+					{
+						*Reason = FText::Format(NSLOCTEXT("UnrealEd", "LandscapeLayerEditLayerDoesntSupportEditing", "This layer's type ({0}) doesn't support direct editing."), TargetLayer->EditLayer->GetClass()->GetDisplayNameText());
+					}
+					return false;
 				}
-				return false;
 			}
-			else if (CurrentTool->GetToolName() == FName("Retopologize"))
+			if (CurrentTool->GetToolName() == FName("Retopologize"))
 			{
 				if (Reason)
 				{
-					*Reason = FText::Format(NSLOCTEXT("UnrealEd", "LandscapeLayersNoSupportForRetopologize", "{0} Tool is not available with the Landscape Layer System."), CurrentTool->GetDisplayName());
+					*Reason = FText::Format(NSLOCTEXT("UnrealEd", "LandscapeLayersNoSupportForRetopologize", "{0} Tool is not available with the Landscape Edit Layer System."), CurrentTool->GetDisplayName());
 				}
 				return false;
 			}
@@ -4280,7 +4301,7 @@ void FEdModeLandscape::UpdateLandscapeSplines(bool bUpdateOnlySelected /* = fals
 
 FGuid FEdModeLandscape::GetCurrentLayerGuid() const
 {
-	FLandscapeLayer* CurrentLayer = GetCurrentLayer();
+	const FLandscapeLayer* CurrentLayer = GetCurrentLayer();
 	return CurrentLayer ? CurrentLayer->Guid : FGuid();
 }
 
@@ -4299,7 +4320,7 @@ bool FEdModeLandscape::NeedToFillEmptyMaterialLayers() const
 
 		if (Landscape != nullptr)
 		{
-			for (FLandscapeLayer& Layer : Landscape->LandscapeLayers)
+			for (const FLandscapeLayer& Layer : Landscape->GetLayers())
 			{
 				for (ULandscapeComponent* Component : Proxy->LandscapeComponents)
 				{
diff --git a/Engine/Source/Editor/LandscapeEditor/Private/LandscapeEdMode.h b/Engine/Source/Editor/LandscapeEditor/Private/LandscapeEdMode.h
--- a/Engine/Source/Editor/LandscapeEditor/Private/LandscapeEdMode.h
+++ b/Engine/Source/Editor/LandscapeEditor/Private/LandscapeEdMode.h
@@ -28,6 +28,7 @@ class FViewport;
 class ULandscapeComponent;
 class ULandscapeEditorObject;
 class UViewportInteractor;
+class ULandscapeEditLayerBase;
 struct FHeightmapToolTarget;
 struct FViewportActionKeyInput;
 struct FViewportClick;
@@ -549,25 +550,26 @@ public:
 	void SetCurrentLayer(int32 InLayerIndex);
 	int32 GetCurrentLayerIndex() const;
 	ALandscape* GetLandscape() const;
-	struct FLandscapeLayer* GetLayer(int32 InLayerIndex) const;
+	const struct FLandscapeLayer* GetLayer(int32 InLayerIndex) const;
 	FName GetLayerName(int32 InLayerIndex) const;
 	void SetLayerName(int32 InLayerIndex, const FName& InName);
 	bool CanRenameLayerTo(int32 InLayerIndex, const FName& InNewName);
 	void SetLayerAlpha(int32 InLayerIndex, float InAlpha);
 	float GetLayerAlpha(int32 InLayerIndex) const;
+	const ULandscapeEditLayerBase* GetEditLayer(int32 InLayerIndex) const;
 	float GetClampedLayerAlpha(float InLayerAlpha) const;
 	void SetLayerVisibility(bool InVisible, int32 InLayerIndex);
 	bool IsLayerVisible(int32 InLayerIndex) const;
 	bool IsLayerLocked(int32 InLayerIndex) const;
 	bool IsLayerAlphaVisible(int32 InLayerIndex) const;
 	void SetLayerLocked(int32 InLayerIndex, bool bInLocked);
-	struct FLandscapeLayer* GetCurrentLayer() const;
+	const struct FLandscapeLayer* GetCurrentLayer() const;
 	FGuid GetCurrentLayerGuid() const;
 	bool IsCurrentLayerBlendSubstractive(const TWeakObjectPtr<ULandscapeLayerInfoObject>& InLayerInfoObj) const;
 	void SetCurrentLayerSubstractiveBlendStatus(bool InStatus, const TWeakObjectPtr<ULandscapeLayerInfoObject>& InLayerInfoObj);
 	void UpdateLandscapeSplines(bool bUpdateOnlySelected = false);
 	void AutoUpdateDirtyLandscapeSplines();
-	bool CanEditLayer(FText* Reason = nullptr, FLandscapeLayer* InLayer = nullptr);
+	bool CanEditLayer(FText* Reason = nullptr, const FLandscapeLayer* InLayer = nullptr);
 
 	void AddBrushToCurrentLayer(class ALandscapeBlueprintBrushBase* InBrush);
 	void RemoveBrushFromCurrentLayer(int32 InBrushIndex);
diff --git a/Engine/Source/Editor/LandscapeEditor/Private/LandscapeEdModeComponentTools.cpp b/Engine/Source/Editor/LandscapeEditor/Private/LandscapeEdModeComponentTools.cpp
--- a/Engine/Source/Editor/LandscapeEditor/Private/LandscapeEdModeComponentTools.cpp
+++ b/Engine/Source/Editor/LandscapeEditor/Private/LandscapeEdModeComponentTools.cpp
@@ -603,7 +603,7 @@ namespace
 				TArray<ULandscapeComponent*> ComponentsUsingHeightmap;
 				ComponentsUsingHeightmap.Add(NewComponent);
 
-				for (const FLandscapeLayer& Layer : Landscape->LandscapeLayers)
+				for (const FLandscapeLayer& Layer : Landscape->GetLayers())
 				{
 					// Since we do not share heightmap when adding new component, we will provided the required array, but they will only be used for 1 component
 					TMap<UTexture2D*, UTexture2D*> CreatedHeightmapTextures;
@@ -764,7 +764,7 @@ public:
 					TArray<ULandscapeComponent*> ComponentsUsingHeightmap;
 					ComponentsUsingHeightmap.Add(NewComponent);
 
-					for (const FLandscapeLayer& Layer : Landscape->LandscapeLayers)
+					for (const FLandscapeLayer& Layer : Landscape->GetLayers())
 					{
 						// Since we do not share heightmap when adding new component, we will provided the required array, but they will only be used for 1 component
 						TMap<UTexture2D*, UTexture2D*> CreatedHeightmapTextures;
@@ -790,7 +790,7 @@ public:
 				TMap<ULandscapeLayerInfoObject*, int32> NeighbourLayerInfoObjectCount;
 
 				{
-					FLandscapeLayer* LandscapeLayer = Landscape ? Landscape->GetLayer(0) : nullptr;
+					const FLandscapeLayer* LandscapeLayer = Landscape ? Landscape->GetLayerConst(0) : nullptr;
 					FScopedSetLandscapeEditingLayer Scope(Landscape, LandscapeLayer ? LandscapeLayer->Guid : FGuid(), [=] { });
 
 					// Cover 9 tiles around us to determine which object should we use by default
diff --git a/Engine/Source/Editor/LandscapeEditor/Private/LandscapeEdModeTools.h b/Engine/Source/Editor/LandscapeEditor/Private/LandscapeEdModeTools.h
--- a/Engine/Source/Editor/LandscapeEditor/Private/LandscapeEdModeTools.h
+++ b/Engine/Source/Editor/LandscapeEditor/Private/LandscapeEdModeTools.h
@@ -949,14 +949,15 @@ public:
 			bCombinedLayerOperation = Landscape && Landscape->HasLayersContent() && InCombinedLayerOperation && bTargetIsHeightmap;
 			if (bCombinedLayerOperation)
 			{
-				for (uint8 i = 0; i < Landscape->GetLayerCount(); ++i)
+				int32 I = 0;
+				for (const FLandscapeLayer& CurrentLayer : Landscape->GetLayers())
 				{
-					FLandscapeLayer* CurrentLayer = Landscape->GetLayer(i);
-					BackupLayerVisibility.Add(CurrentLayer->bVisible);
-					if (CurrentLayer->Guid == EditingLayerGuid.GetValue())
+					BackupLayerVisibility.Add(CurrentLayer.bVisible);
+					if (CurrentLayer.Guid == EditingLayerGuid.GetValue())
 					{
-						EditingLayerIndex = i;
+						EditingLayerIndex = I;
 					}
+					++I;
 				}
 				check(EditingLayerIndex < Landscape->GetLayerCount());
 			}
@@ -971,10 +972,11 @@ public:
 		if (bCombinedLayerOperation)
 		{
 			TArray<bool> NewLayerVisibility;
-			for (int i = 0; i < Landscape->GetLayerCount(); ++i)
+			int32 I = 0;
+			for (const FLandscapeLayer& CurrentLayer : Landscape->GetLayers())
 			{
-				FLandscapeLayer* CurrentLayer = Landscape->GetLayer(i);
-				NewLayerVisibility.Add((i > EditingLayerIndex) ? false : CurrentLayer->bVisible);
+				NewLayerVisibility.Add((I > EditingLayerIndex) ? false : CurrentLayer.bVisible);
+				++I;
 			}
 
 			auto OnCacheUpdating = [&](const FIntRect& NewCacheBounds) -> FIntRect
@@ -1043,8 +1045,8 @@ public:
 
 			FGuid CacheAccessorLayerGuid = CacheUpToEditingLayer.DataAccess.GetEditLayer();
  			checkf(EditingLayerGuid.GetValue() == CacheAccessorLayerGuid, TEXT("Editing Layer has changed between Initialize and Write. Was: %s (%s). Is now: %s (%s)"),
- 				Landscape->GetLayer(*EditingLayerGuid) ? *(Landscape->GetLayer(*EditingLayerGuid)->Name.ToString()) : TEXT("<unknown>"), *EditingLayerGuid->ToString(),
- 				Landscape->GetLayer(CacheAccessorLayerGuid) ? *(Landscape->GetLayer(CacheAccessorLayerGuid)->Name.ToString()) : TEXT("<unknown>"), *CacheAccessorLayerGuid.ToString());
+ 				Landscape->GetLayerConst(*EditingLayerGuid) ? *(Landscape->GetLayerConst(*EditingLayerGuid)->Name.ToString()) : TEXT("<unknown>"), *EditingLayerGuid->ToString(),
+ 				Landscape->GetLayerConst(CacheAccessorLayerGuid) ? *(Landscape->GetLayerConst(CacheAccessorLayerGuid)->Name.ToString()) : TEXT("<unknown>"), *CacheAccessorLayerGuid.ToString());
 
 			// Restore layers visibility
 			SetLayersVisibility(BackupLayerVisibility);
@@ -1086,11 +1088,13 @@ private:
 		check(InLayerVisibility.Num() == Landscape->GetLayerCount());
 		for (int i = 0; i < InLayerVisibility.Num(); ++i)
 		{
-			if (FLandscapeLayer* Layer = Landscape->GetLayer(i))
+			if (const FLandscapeLayer* Layer = Landscape->GetLayerConst(i))
 			{
 				if (Layer->bVisible != InLayerVisibility[i])
 				{
-					Layer->bVisible = InLayerVisibility[i];
+					// Pass bInForIntermediateRender = true here, because we don't want to call Modify on the landscape, we don't actually want to truly change 
+					//  the layer visibility, only temporarily for the purpose of the intermediate render : 
+					Landscape->SetLayerVisibility(i, InLayerVisibility[i], /*bInForIntermediateRender = */true);
 					bVisibilityChanged = true;
 				}
 			}
diff --git a/Engine/Source/Editor/LandscapeEditor/Private/LandscapeEditorDetailCustomization_Layers.cpp b/Engine/Source/Editor/LandscapeEditor/Private/LandscapeEditorDetailCustomization_Layers.cpp
--- a/Engine/Source/Editor/LandscapeEditor/Private/LandscapeEditorDetailCustomization_Layers.cpp
+++ b/Engine/Source/Editor/LandscapeEditor/Private/LandscapeEditorDetailCustomization_Layers.cpp
@@ -21,6 +21,7 @@
 #include "Widgets/Input/SMultiLineEditableTextBox.h"
 #include "EditorModeManager.h"
 #include "EditorModes.h"
+#include "EditorClassUtils.h"
 #include "LandscapeEditorDetailCustomization_LayersBrushStack.h" // FLandscapeBrushDragDropOp
 #include "LandscapeEditorModule.h"
 #include "LandscapeEditorObject.h"
@@ -30,6 +31,10 @@
 #include "IDetailPropertyRow.h"
 #include "DetailCategoryBuilder.h"
 #include "PropertyCustomizationHelpers.h"
+#include "Styling/SlateIconFinder.h"
+#include "ClassViewerFilter.h"
+#include "ClassViewerModule.h"
+#include "Kismet2/SClassPickerDialog.h"
 
 #include "SLandscapeEditor.h"
 #include "Dialogs/DlgPickAssetPath.h"
@@ -128,7 +133,38 @@ void FLandscapeEditorCustomNodeBuilder_Layers::SetOnRebuildChildren(FSimpleDeleg
 
 void FLandscapeEditorCustomNodeBuilder_Layers::GenerateHeaderRowContent(FDetailWidgetRow& NodeRow)
 {
+	NodeRow.NameWidget
+		[
+			SNew(STextBlock)
+				.Font(IDetailLayoutBuilder::GetDetailFont())
+				.Text(LOCTEXT("LayersLabel", "Layers"))
+		];
+
 
+	TSharedRef<SWidget> AddButton = PropertyCustomizationHelpers::MakeAddButton(
+		FSimpleDelegate::CreateSP(this, &FLandscapeEditorCustomNodeBuilder_Layers::CreateLayer),
+		MakeAttributeSPLambda(this, [this] { FText Reason; CanCreateLayer(Reason); return Reason; }),
+		MakeAttributeSPLambda(this, [this] { FText Reason; return CanCreateLayer(Reason); }));
+
+	NodeRow.ValueWidget
+		[
+			SNew(SHorizontalBox)
+				+ SHorizontalBox::Slot()
+				.FillWidth(1) // Fill the entire width if possible
+				.VAlign(VAlign_Center)
+				[
+					SNew(STextBlock)
+						.Text(this, &FLandscapeEditorCustomNodeBuilder_Layers::GetNumLayersText)
+				]
+				+ SHorizontalBox::Slot()
+				.AutoWidth()
+				.HAlign(HAlign_Center)
+				.VAlign(VAlign_Center)
+				.Padding(2.0f, 0.0f)
+				[
+					AddButton
+				]
+		];
 }
 
 BEGIN_SLATE_FUNCTION_BUILD_OPTIMIZATION
@@ -153,7 +189,7 @@ void FLandscapeEditorCustomNodeBuilder_Layers::GenerateChildContent(IDetailChild
 
 		InlineTextBlocks.Empty();
 		InlineTextBlocks.AddDefaulted(LandscapeEdMode->GetLayerCount());
-		// Slots are displayed in the opposite order of LandscapeLayers
+		// Slots are displayed in the opposite order of LandscapeEditLayers
 		for (int32 i = LandscapeEdMode->GetLayerCount()-1; i >= 0 ; --i)
 		{
 			TSharedPtr<SWidget> GeneratedRowWidget = GenerateRow(i);
@@ -174,9 +210,14 @@ END_SLATE_FUNCTION_BUILD_OPTIMIZATION
 BEGIN_SLATE_FUNCTION_BUILD_OPTIMIZATION
 TSharedPtr<SWidget> FLandscapeEditorCustomNodeBuilder_Layers::GenerateRow(int32 InLayerIndex)
 {
+	TSharedRef<SWidget> DeleteButton = PropertyCustomizationHelpers::MakeDeleteButton(
+		FSimpleDelegate::CreateSP(this, &FLandscapeEditorCustomNodeBuilder_Layers::DeleteLayer, InLayerIndex),
+		MakeAttributeSPLambda(this, [this, InLayerIndex] { FText Reason; CanDeleteLayer(InLayerIndex, Reason); return Reason; }),
+		MakeAttributeSPLambda(this, [this, InLayerIndex] { FText Reason; return CanDeleteLayer(InLayerIndex, Reason); }));
+
 	FEdModeLandscape* LandscapeEdMode = GetEditorMode();
 	TSharedPtr<SWidget> RowWidget = SNew(SLandscapeEditorSelectableBorder)
-		.Padding(FMargin(8.f, 0.f))
+		.Padding(0.f)
 		.VAlign(VAlign_Center)
 		.OnContextMenuOpening(this, &FLandscapeEditorCustomNodeBuilder_Layers::OnLayerContextMenuOpening, InLayerIndex)
 		.OnSelected(this, &FLandscapeEditorCustomNodeBuilder_Layers::OnLayerSelectionChanged, InLayerIndex)
@@ -184,107 +225,153 @@ TSharedPtr<SWidget> FLandscapeEditorCustomNodeBuilder_Layers::GenerateRow(int32
 		.Visibility(EVisibility::Visible)
 		[
 			SNew(SHorizontalBox)
-			
-			+ SHorizontalBox::Slot()
-			.AutoWidth()
-			.VAlign(VAlign_Center)
-			[
-				SNew(SButton)
-				.ButtonStyle(FAppStyle::Get(), "NoBorder")
-				.OnClicked(this, &FLandscapeEditorCustomNodeBuilder_Layers::OnToggleLock, InLayerIndex)
-				.ToolTipText(LOCTEXT("LandscapeLayerLock", "Locks the current layer"))
+
+				+ SHorizontalBox::Slot()
+				.AutoWidth()
+				.VAlign(VAlign_Center)
 				[
-					SNew(SImage)
-					.Image(this, &FLandscapeEditorCustomNodeBuilder_Layers::GetLockBrushForLayer, InLayerIndex)
+					SNew(SBox)
+						.Padding(FMargin(2.0f, 0.0f, 2.0f, 0.0f))
+						[
+							SNew(SImage)
+								.Image(FCoreStyle::Get().GetBrush("VerticalBoxDragIndicatorShort"))
+						]
 				]
-			]
 
-			+SHorizontalBox::Slot()
-			.AutoWidth()
-			.Padding(4, 0)
-			.VAlign(VAlign_Center)
-			[
-				SNew(SButton)
-				.ContentPadding(0.0f)
-				.ButtonStyle(FAppStyle::Get(), "NoBorder")
-				.OnClicked(this, &FLandscapeEditorCustomNodeBuilder_Layers::OnToggleVisibility, InLayerIndex)
-				.ToolTipText(LOCTEXT("LandscapeLayerVisibility", "Toggle Layer Visibility"))
-				.HAlign(HAlign_Center)
+				+ SHorizontalBox::Slot()
+				.AutoWidth()
 				.VAlign(VAlign_Center)
-				.Content()
 				[
-					SNew(SImage)
-					.Image(this, &FLandscapeEditorCustomNodeBuilder_Layers::GetVisibilityBrushForLayer, InLayerIndex)
+					SNew(SBox)
+						.Padding(FMargin(2.0f, 0.0f, 2.0f, 0.0f))
+						[
+							SNew(SImage)
+								.Image(this, &FLandscapeEditorCustomNodeBuilder_Layers::GetEditLayerIconBrush, InLayerIndex)
+								// TODO [jonathan.bard] : investigate why this doesn't work : the tooltip just doesn't show up
+								.ToolTip(this, &FLandscapeEditorCustomNodeBuilder_Layers::GetEditLayerTypeTooltip, InLayerIndex)
+						]
 				]
-			]
 
-			+ SHorizontalBox::Slot()
-			.FillWidth(1.0)
-			.VAlign(VAlign_Center)
-			.Padding(4, 0)
-			[
-				SAssignNew(InlineTextBlocks[InLayerIndex], SInlineEditableTextBlock)
-				.IsEnabled(this, &FLandscapeEditorCustomNodeBuilder_Layers::IsLayerEditionEnabled, InLayerIndex)
-				.Text(this, &FLandscapeEditorCustomNodeBuilder_Layers::GetLayerDisplayName, InLayerIndex)
-				.ColorAndOpacity(TAttribute<FSlateColor>::Create(TAttribute<FSlateColor>::FGetter::CreateSP(this, &FLandscapeEditorCustomNodeBuilder_Layers::GetLayerTextColor, InLayerIndex)))
-				.ToolTipText(LOCTEXT("LandscapeLayers_tooltip", "Name of the Layer"))
-				.OnVerifyTextChanged(FOnVerifyTextChanged::CreateSP(this, &FLandscapeEditorCustomNodeBuilder_Layers::CanRenameLayerTo, InLayerIndex))
-				.OnEnterEditingMode(this, &FLandscapeEditorCustomNodeBuilder_Layers::OnBeginNameTextEdit)
-				.OnExitEditingMode(this, &FLandscapeEditorCustomNodeBuilder_Layers::OnEndNameTextEdit)
-				.OnTextCommitted(FOnTextCommitted::CreateSP(this, &FLandscapeEditorCustomNodeBuilder_Layers::SetLayerName, InLayerIndex))
-			]
-
-			+ SHorizontalBox::Slot()
-			.FillWidth(1.0f)
-			.VAlign(VAlign_Center)
-			.Padding(0, 2)
-			.HAlign(HAlign_Right)
-			[
-				SNew(SHorizontalBox)
 				+ SHorizontalBox::Slot()
-				.Padding(0)
-				.FillWidth(1.0f)
+				.AutoWidth()
 				.VAlign(VAlign_Center)
-				.HAlign(HAlign_Left)
 				[
-					SNew(STextBlock)
-					.IsEnabled(this, &FLandscapeEditorCustomNodeBuilder_Layers::IsLayerEditionEnabled, InLayerIndex)
-					.Visibility(this, &FLandscapeEditorCustomNodeBuilder_Layers::GetLayerAlphaVisibility, InLayerIndex)
-					.Text(LOCTEXT("LandscapeLayerAlpha", "Alpha"))
-					.ColorAndOpacity(TAttribute<FSlateColor>::Create(TAttribute<FSlateColor>::FGetter::CreateSP(this, &FLandscapeEditorCustomNodeBuilder_Layers::GetLayerTextColor, InLayerIndex)))
+					SNew(SButton)
+						.ButtonStyle(FAppStyle::Get(), "NoBorder")
+						.OnClicked(this, &FLandscapeEditorCustomNodeBuilder_Layers::OnToggleLock, InLayerIndex)
+						.ToolTipText(LOCTEXT("LandscapeLayerLock", "Locks the current layer"))
+						[
+							SNew(SImage)
+								.Image(this, &FLandscapeEditorCustomNodeBuilder_Layers::GetLockBrushForLayer, InLayerIndex)
+						]
 				]
+
 				+ SHorizontalBox::Slot()
+				.AutoWidth()
+				.Padding(4, 0)
 				.VAlign(VAlign_Center)
-				.Padding(0, 2)
-				.HAlign(HAlign_Left)
+				[
+					SNew(SButton)
+						.ContentPadding(0.0f)
+						.ButtonStyle(FAppStyle::Get(), "NoBorder")
+						.IsEnabled(MakeAttributeSPLambda(this, [this, InLayerIndex] { FText Reason; return CanToggleVisibility(InLayerIndex, Reason); }))
+						.OnClicked(this, &FLandscapeEditorCustomNodeBuilder_Layers::OnToggleVisibility, InLayerIndex)
+						.ToolTipText(MakeAttributeSPLambda(this, [this, InLayerIndex] { FText Reason; CanToggleVisibility(InLayerIndex, Reason); return Reason; }))
+						.HAlign(HAlign_Center)
+						.VAlign(VAlign_Center)
+						.Content()
+						[
+							SNew(SImage)
+								.Image(this, &FLandscapeEditorCustomNodeBuilder_Layers::GetVisibilityBrushForLayer, InLayerIndex)
+						]
+				]
+
+				+ SHorizontalBox::Slot()
+				.FillWidth(1.0)
+				.VAlign(VAlign_Center)
+				.Padding(4, 0)
+				[
+					SNew(SHorizontalBox)
+						+ SHorizontalBox::Slot()
+						.Padding(0)
+						.AutoWidth()
+						.VAlign(VAlign_Center)
+						.HAlign(HAlign_Left)
+						[
+							SAssignNew(InlineTextBlocks[InLayerIndex], SInlineEditableTextBlock)
+								.IsEnabled(MakeAttributeSPLambda(this, [this, InLayerIndex] { FText Reason; return CanRenameLayer(InLayerIndex, Reason); }))
+								.Text(this, &FLandscapeEditorCustomNodeBuilder_Layers::GetLayerDisplayName, InLayerIndex)
+								.ColorAndOpacity(TAttribute<FSlateColor>::Create(TAttribute<FSlateColor>::FGetter::CreateSP(this, &FLandscapeEditorCustomNodeBuilder_Layers::GetLayerTextColor, InLayerIndex)))
+								.ToolTipText(MakeAttributeSPLambda(this, [this, InLayerIndex] { FText Reason; CanRenameLayer(InLayerIndex, Reason); return Reason; }))
+								.OnVerifyTextChanged(FOnVerifyTextChanged::CreateSP(this, &FLandscapeEditorCustomNodeBuilder_Layers::CanRenameLayerTo, InLayerIndex))
+								.OnEnterEditingMode(this, &FLandscapeEditorCustomNodeBuilder_Layers::OnBeginNameTextEdit)
+								.OnExitEditingMode(this, &FLandscapeEditorCustomNodeBuilder_Layers::OnEndNameTextEdit)
+								.OnTextCommitted(FOnTextCommitted::CreateSP(this, &FLandscapeEditorCustomNodeBuilder_Layers::SetLayerName, InLayerIndex))
+						]
+				]
+
+				+ SHorizontalBox::Slot()
 				.FillWidth(1.0f)
+				.VAlign(VAlign_Center)
+				.Padding(0, 2)
+				.HAlign(HAlign_Right)
 				[
-					SNew(SNumericEntryBox<float>)
-					.AllowSpin(true)
-					.MinValue(-1.0f)
-					.MaxValue(1.0f)
-					.MinSliderValue(-1.0f)
-					.MaxSliderValue(1.0f)
-					.Delta(0.01f)
-					.MinDesiredValueWidth(60.0f)
-					.IsEnabled(this, &FLandscapeEditorCustomNodeBuilder_Layers::IsLayerEditionEnabled, InLayerIndex)
-					.Visibility(this, &FLandscapeEditorCustomNodeBuilder_Layers::GetLayerAlphaVisibility, InLayerIndex)
-					.Value(this, &FLandscapeEditorCustomNodeBuilder_Layers::GetLayerAlpha, InLayerIndex)
-					.OnValueChanged_Lambda([this, InLayerIndex](float InValue) { SetLayerAlpha(InValue, InLayerIndex, false); })
-					.OnValueCommitted_Lambda([this, InLayerIndex](float InValue, ETextCommit::Type InCommitType) { SetLayerAlpha(InValue, InLayerIndex, true); })
-					.OnBeginSliderMovement_Lambda([this, InLayerIndex]()
-					{
-						CurrentSlider = InLayerIndex;
-						GEditor->BeginTransaction(LOCTEXT("Landscape_Layers_SetAlpha", "Set Layer Alpha"));
-					})
-					.OnEndSliderMovement_Lambda([this](double)
-					{
-						GEditor->EndTransaction();
-						CurrentSlider = INDEX_NONE;
-					})
-				]		
-			]
-		];	
+					SNew(SHorizontalBox)
+						+ SHorizontalBox::Slot()
+						.Padding(0)
+						.FillWidth(1.0f)
+						.VAlign(VAlign_Center)
+						.HAlign(HAlign_Left)
+						[
+							SNew(STextBlock)
+								.IsEnabled(MakeAttributeSPLambda(this, [this, InLayerIndex] { FText Reason; return CanSetLayerAlpha(InLayerIndex, Reason); }))
+								.Visibility(this, &FLandscapeEditorCustomNodeBuilder_Layers::GetLayerAlphaVisibility, InLayerIndex)
+								.Text(LOCTEXT("LandscapeLayerAlpha", "Alpha"))
+								.ToolTipText(MakeAttributeSPLambda(this, [this, InLayerIndex] { FText Reason; CanSetLayerAlpha(InLayerIndex, Reason); return Reason; }))
+								.ColorAndOpacity(TAttribute<FSlateColor>::Create(TAttribute<FSlateColor>::FGetter::CreateSP(this, &FLandscapeEditorCustomNodeBuilder_Layers::GetLayerTextColor, InLayerIndex)))
+						]
+						+ SHorizontalBox::Slot()
+						.VAlign(VAlign_Center)
+						.Padding(0, 2)
+						.HAlign(HAlign_Left)
+						.FillWidth(1.0f)
+						[
+							SNew(SNumericEntryBox<float>)
+								.AllowSpin(true)
+								.MinValue(-1.0f)
+								.MaxValue(1.0f)
+								.MinSliderValue(-1.0f)
+								.MaxSliderValue(1.0f)
+								.Delta(0.01f)
+								.MinDesiredValueWidth(60.0f)
+								.IsEnabled(MakeAttributeSPLambda(this, [this, InLayerIndex] { FText Reason; return CanSetLayerAlpha(InLayerIndex, Reason); }))
+								.Visibility(this, &FLandscapeEditorCustomNodeBuilder_Layers::GetLayerAlphaVisibility, InLayerIndex)
+								.ToolTipText(MakeAttributeSPLambda(this, [this, InLayerIndex] { FText Reason; CanSetLayerAlpha(InLayerIndex, Reason); return Reason; }))
+								.Value(this, &FLandscapeEditorCustomNodeBuilder_Layers::GetLayerAlpha, InLayerIndex)
+								.OnValueChanged_Lambda([this, InLayerIndex](float InValue) { SetLayerAlpha(InValue, InLayerIndex, false); })
+								.OnValueCommitted_Lambda([this, InLayerIndex](float InValue, ETextCommit::Type InCommitType) { SetLayerAlpha(InValue, InLayerIndex, true); })
+								.OnBeginSliderMovement_Lambda([this, InLayerIndex]()
+							{
+								CurrentSlider = InLayerIndex;
+								GEditor->BeginTransaction(LOCTEXT("Landscape_Layers_SetAlpha", "Set Layer Alpha"));
+							})
+								.OnEndSliderMovement_Lambda([this](double)
+							{
+								GEditor->EndTransaction();
+								CurrentSlider = INDEX_NONE;
+							})
+						]
+				]
+
+				+ SHorizontalBox::Slot()
+				.AutoWidth()
+				.HAlign(HAlign_Center)
+				.VAlign(VAlign_Center)
+				.Padding(2.0f, 0.0f)
+				[
+					DeleteButton
+				]
+		];
 
 	return RowWidget;
 }
@@ -296,10 +383,7 @@ FText FLandscapeEditorCustomNodeBuilder_Layers::GetLayerDisplayName(int32 InLaye
 	ALandscape* Landscape = LandscapeEdMode ? LandscapeEdMode->GetLandscape() : nullptr;
 	if (LandscapeEdMode && Landscape)
 	{
-		const FLandscapeLayer* Layer = LandscapeEdMode->GetLayer(InLayerIndex);
-		const FLandscapeLayer* ReservedLayer = Landscape->GetLandscapeSplinesReservedLayer();
-		bool bIsReserved = Layer && Layer == ReservedLayer && InLayerIndex != CurrentEditingInlineTextBlock;
-		return FText::Format(LOCTEXT("LayerDisplayName", "{0}{1}"), FText::FromName(LandscapeEdMode->GetLayerName(InLayerIndex)), bIsReserved ? LOCTEXT("ReservedForSplines", " (Reserved for Splines)") : FText::GetEmpty());
+		return FText::FromName(LandscapeEdMode->GetLayerName(InLayerIndex));
 	}
 
 	return FText::FromString(TEXT("None"));
@@ -334,7 +418,7 @@ bool FLandscapeEditorCustomNodeBuilder_Layers::CanRenameLayerTo(const FText& InN
 	{
 		if (!LandscapeEdMode->CanRenameLayerTo(InLayerIndex, *InNewText.ToString()))
 		{
-			OutErrorMessage = LOCTEXT("Landscape_Layers_RenameFailed_AlreadyExists", "This layer already exists");
+			OutErrorMessage = LOCTEXT("Landscape_Layers_RenameFailed_AlreadyExists", "This edit layer name already exists");
 			return false;
 		}
 	}
@@ -361,31 +445,31 @@ void FLandscapeEditorCustomNodeBuilder_Layers::FillAddBrushMenu(FMenuBuilder& Me
 {
 	for (ALandscapeBlueprintBrushBase* Brush : Brushes)
 	{
-		TSharedRef<FLandscapeEditorCustomNodeBuilder_Layers> SharedThis = AsShared();
-		FUIAction AddAction = FUIAction(FExecuteAction::CreateLambda([SharedThis, Brush]() { SharedThis->AddBrushToCurrentLayer(Brush); }));
+		FUIAction AddAction = FUIAction(FExecuteAction::CreateSP(this, &FLandscapeEditorCustomNodeBuilder_Layers::AddBrushToCurrentLayer, Brush));
 		MenuBuilder.AddMenuEntry(FText::FromString(Brush->GetActorLabel()), FText(), FSlateIcon(), AddAction);
 	}
 }
 
-void FLandscapeEditorCustomNodeBuilder_Layers::FillClearPaintLayerMenu(FMenuBuilder& MenuBuilder, int32 InLayerIndex, TArray<ULandscapeLayerInfoObject*> InUsedLayerInfos)
+void FLandscapeEditorCustomNodeBuilder_Layers::FillClearTargetLayerMenu(FMenuBuilder& MenuBuilder, int32 InLayerIndex, TArray<ULandscapeLayerInfoObject*> InUsedLayerInfos)
 {
 	// Clear All Weightmap Data
-	TSharedRef<FLandscapeEditorCustomNodeBuilder_Layers> SharedThis = AsShared();
-	FUIAction ClearAction = FUIAction(FExecuteAction::CreateLambda([SharedThis, InLayerIndex]() { SharedThis->ClearLayer(InLayerIndex, ELandscapeClearMode::Clear_Weightmap); }));
+	FUIAction ClearAction = FUIAction(FExecuteAction::CreateSP(this, &FLandscapeEditorCustomNodeBuilder_Layers::ClearTargetLayersOnLayer, InLayerIndex, ELandscapeClearMode::Clear_Weightmap));
 	MenuBuilder.AddMenuEntry(LOCTEXT("LandscapeClearAllWeightmap", "All"), FText(), FSlateIcon(), ClearAction);
 	MenuBuilder.AddMenuSeparator();
 
 	// Clear Per LayerInfo
 	for (ULandscapeLayerInfoObject* LayerInfo : InUsedLayerInfos)
 	{
-		FUIAction ClearLayerInfoAction = FUIAction(FExecuteAction::CreateLambda([SharedThis, InLayerIndex, LayerInfo]() { SharedThis->ClearPaintLayer(InLayerIndex, LayerInfo); }));
-		MenuBuilder.AddMenuEntry(FText::FromName(LayerInfo->LayerName), FText(), FSlateIcon(), ClearLayerInfoAction);
+		FUIAction ClearLayerInfoAction = FUIAction(
+			FExecuteAction::CreateSP(this, &FLandscapeEditorCustomNodeBuilder_Layers::ClearTargetLayerOnLayer,InLayerIndex, LayerInfo),
+			FCanExecuteAction::CreateSPLambda(this, [this, InLayerIndex, LayerInfo]() { FText Reason; return CanClearTargetLayerOnLayer(InLayerIndex, LayerInfo, Reason); }));
+		TAttribute<FText> CanClearTargetLayerOnLayerTooltipText = MakeAttributeSPLambda(this, [this, InLayerIndex, LayerInfo] { FText Reason; CanClearTargetLayerOnLayer(InLayerIndex, LayerInfo, Reason); return Reason; });
+		MenuBuilder.AddMenuEntry(FText::FromName(LayerInfo->LayerName), CanClearTargetLayerOnLayerTooltipText, FSlateIcon(), ClearLayerInfoAction);
 	}
 }
 
 void FLandscapeEditorCustomNodeBuilder_Layers::FillClearLayerMenu(FMenuBuilder& MenuBuilder, int32 InLayerIndex)
 {
-	TSharedRef<FLandscapeEditorCustomNodeBuilder_Layers> SharedThis = AsShared();
 	const UEnum* EnumPtr = FindObject<UEnum>(nullptr, TEXT("/Script/Landscape.ELandscapeClearMode"), true);
 	if (ensure(EnumPtr != nullptr))
 	{
@@ -404,19 +488,26 @@ void FLandscapeEditorCustomNodeBuilder_Layers::FillClearLayerMenu(FMenuBuilder&
 				Landscape->GetUsedPaintLayers(InLayerIndex, UsedLayerInfos);
 				if (UsedLayerInfos.Num() > 0)
 				{
+					FUIAction ClearLayerAction = FUIAction(
+						FExecuteAction(),
+						FCanExecuteAction::CreateSPLambda(this, [this, InLayerIndex, EnumValue] { FText Reason; return CanClearTargetLayersOnLayer(InLayerIndex, EnumValue, Reason); }));
+					TAttribute<FText> CanClearLayerTooltipText = MakeAttributeSPLambda(this, [this, InLayerIndex, EnumValue] { FText Reason; CanClearTargetLayersOnLayer(InLayerIndex, EnumValue, Reason); return Reason; });
 					MenuBuilder.AddSubMenu(
 						EnumPtr->GetDisplayNameTextByIndex(i),
-						FText(),
-						FNewMenuDelegate::CreateSP(this, &FLandscapeEditorCustomNodeBuilder_Layers::FillClearPaintLayerMenu, InLayerIndex, UsedLayerInfos),
-						false,
-						FSlateIcon()
-					);
+						CanClearLayerTooltipText,
+						FNewMenuDelegate::CreateSP(this, &FLandscapeEditorCustomNodeBuilder_Layers::FillClearTargetLayerMenu, InLayerIndex, UsedLayerInfos),
+						ClearLayerAction,
+						/*InExtensionHook = */NAME_None,
+						EUserInterfaceActionType::None);
 				}
 			}
 			else
 			{
-				FUIAction ClearAction = FUIAction(FExecuteAction::CreateLambda([SharedThis, InLayerIndex, EnumValue]() { SharedThis->ClearLayer(InLayerIndex, EnumValue); }));
-				MenuBuilder.AddMenuEntry(EnumPtr->GetDisplayNameTextByIndex(i), FText(), FSlateIcon(), ClearAction);
+				FUIAction ClearLayerAction = FUIAction(
+					FExecuteAction::CreateSP(this, &FLandscapeEditorCustomNodeBuilder_Layers::ClearTargetLayersOnLayer, InLayerIndex, EnumValue),
+					FCanExecuteAction::CreateSPLambda(this, [this, InLayerIndex, EnumValue]() { FText Reason; return CanClearTargetLayersOnLayer(InLayerIndex, EnumValue, Reason); }));
+				TAttribute<FText> CanClearLayerActionTooltipText = MakeAttributeSPLambda(this, [this, InLayerIndex, EnumValue] { FText Reason; CanClearTargetLayersOnLayer(InLayerIndex, EnumValue, Reason); return Reason; });
+				MenuBuilder.AddMenuEntry(EnumPtr->GetDisplayNameTextByIndex(i), CanClearLayerActionTooltipText, FSlateIcon(), ClearLayerAction);
 			}
 		}
 	}
@@ -434,60 +525,63 @@ TSharedPtr<SWidget> FLandscapeEditorCustomNodeBuilder_Layers::OnLayerContextMenu
 	ALandscape* Landscape = LandscapeEdMode ? LandscapeEdMode->GetLandscape() : nullptr;
 	if (Landscape && LandscapeEdMode && LandscapeEdMode->DoesCurrentToolAffectEditLayers())
 	{
-		FLandscapeLayer* Layer = LandscapeEdMode->GetLayer(InLayerIndex);
-		TSharedRef<FLandscapeEditorCustomNodeBuilder_Layers> SharedThis = AsShared();
+		const FLandscapeLayer* Layer = LandscapeEdMode->GetLayer(InLayerIndex);
 		FMenuBuilder MenuBuilder(true, nullptr);
 		MenuBuilder.BeginSection("LandscapeEditorLayerActions", LOCTEXT("LandscapeEditorLayerActions.Heading", "Edit Layers"));
 		{
-			// Create Layer
-			FUIAction CreateLayerAction = FUIAction(FExecuteAction::CreateLambda([SharedThis] { SharedThis->CreateLayer(); }), FCanExecuteAction::CreateLambda([Landscape] { return !Landscape->IsMaxLayersReached(); }));
-			TAttribute<FText> CreateLayerText(Landscape->IsMaxLayersReached() ? LOCTEXT("MaxLayersReached", "Create (Max layers reached)") : LOCTEXT("CreateLayer", "Create"));
-			MenuBuilder.AddMenuEntry(CreateLayerText, LOCTEXT("CreateLayerTooltip", "Create Layer"), FSlateIcon(), CreateLayerAction);
-	
 			if (Layer)
 			{
-				// Rename Layer
-				FUIAction RenameLayerAction = FUIAction(FExecuteAction::CreateLambda([SharedThis, InLayerIndex] { SharedThis->RenameLayer(InLayerIndex); }));
-				MenuBuilder.AddMenuEntry(LOCTEXT("RenameLayer", "Rename..."), LOCTEXT("RenameLayerTooltip", "Rename Layer"), FSlateIcon(), RenameLayerAction);
+				const ULandscapeEditLayerBase* EditLayer = Layer->EditLayer;
+				check(EditLayer != nullptr);
 
-				if (!Layer->bLocked)
+				// Rename Layer
+				FUIAction RenameLayerAction = FUIAction(
+					FExecuteAction::CreateSP(this, &FLandscapeEditorCustomNodeBuilder_Layers::RenameLayer, InLayerIndex),
+					FCanExecuteAction::CreateSPLambda(this, [this, InLayerIndex] { FText Reason; return CanRenameLayer(InLayerIndex, Reason); }));
+				TAttribute<FText> CanRenameLayerTooltipText = MakeAttributeSPLambda(this, [this, InLayerIndex] { FText Reason; CanRenameLayer(InLayerIndex, Reason); return Reason; });
+				MenuBuilder.AddMenuEntry(LOCTEXT("RenameLayer", "Rename..."), CanRenameLayerTooltipText, FSlateIcon(), RenameLayerAction);
+
+				// Clear Layer
+				FUIAction ClearLayerAction = FUIAction(
+					FExecuteAction(),
+					FCanExecuteAction::CreateSPLambda(this, [this, InLayerIndex] { FText Reason; return CanClearLayer(InLayerIndex, Reason); }));
+				TAttribute<FText> CanClearLayerTooltipText = MakeAttributeSPLambda(this, [this, InLayerIndex] { FText Reason; CanClearLayer(InLayerIndex, Reason); return Reason; });
+				MenuBuilder.AddSubMenu(
+					LOCTEXT("LandscapeEditorClearLayerSubMenu", "Clear"),
+					CanClearLayerTooltipText,
+					FNewMenuDelegate::CreateSP(this, &FLandscapeEditorCustomNodeBuilder_Layers::FillClearLayerMenu, InLayerIndex),
+					ClearLayerAction, 
+					/*InExtensionHook = */NAME_None,
+					EUserInterfaceActionType::None);
+
+				// Delete Layer
+				FUIAction DeleteLayerAction = FUIAction(
+					FExecuteAction::CreateSP(this, &FLandscapeEditorCustomNodeBuilder_Layers::DeleteLayer, InLayerIndex),
+					FCanExecuteAction::CreateSPLambda(this, [this, InLayerIndex] { FText Reason; return CanDeleteLayer(InLayerIndex, Reason); }));
+				TAttribute<FText> DeleteLayerTooltipText = MakeAttributeSPLambda(this, [this, InLayerIndex] { FText Reason; CanDeleteLayer(InLayerIndex, Reason); return Reason; });
+				MenuBuilder.AddMenuEntry(LOCTEXT("DeleteLayer", "Delete..."), DeleteLayerTooltipText, FSlateIcon(), DeleteLayerAction);
+
+				// Collapse Layer
+				FUIAction CollapseLayerAction = FUIAction(
+					FExecuteAction::CreateSP(this, &FLandscapeEditorCustomNodeBuilder_Layers::CollapseLayer, InLayerIndex),
+					FCanExecuteAction::CreateSPLambda(this, [this, InLayerIndex] { FText Reason; return CanCollapseLayer(InLayerIndex, Reason); }));
+				TAttribute<FText> CollapseLayerTooltipText = MakeAttributeSPLambda(this, [this, InLayerIndex] { FText Reason; CanCollapseLayer(InLayerIndex, Reason); return Reason; });
+				MenuBuilder.AddMenuEntry(LOCTEXT("CollapseLayer", "Collapse..."), CollapseLayerTooltipText, FSlateIcon(), CollapseLayerAction);
+
+				// Add custom actions : 
+				if (!EditLayer->GetActions().IsEmpty())
 				{
-					// Clear Layer
-					MenuBuilder.AddSubMenu(
-						LOCTEXT("LandscapeEditorClearLayerSubMenu", "Clear"),
-						FText(),
-						FNewMenuDelegate::CreateSP(this, &FLandscapeEditorCustomNodeBuilder_Layers::FillClearLayerMenu, InLayerIndex),
-						false,
-						FSlateIcon()
-					);
-
-					// Delete Layer
-					FUIAction DeleteLayerAction = FUIAction(FExecuteAction::CreateLambda([SharedThis, InLayerIndex] { SharedThis->DeleteLayer(InLayerIndex); }), FCanExecuteAction::CreateLambda([Landscape] { return Landscape->LandscapeLayers.Num() > 1; }));
-					TAttribute<FText> DeleteLayerText(Landscape->LandscapeLayers.Num() == 1 ? LOCTEXT("CantDeleteLastLayer", "Delete (Last layer)") : LOCTEXT("DeleteLayer", "Delete..."));
-					MenuBuilder.AddMenuEntry(DeleteLayerText, LOCTEXT("DeleteLayerTooltip", "Delete Layer"), FSlateIcon(), DeleteLayerAction);
-
-					// Collapse Layer
-					if (LandscapeEdMode->CurrentToolMode->ToolModeName == "ToolMode_Manage")
+					MenuBuilder.AddMenuSeparator();
+
+					for (const ULandscapeEditLayerBase::FEditLayerAction& LayerAction : EditLayer->GetActions())
 					{
-						FUIAction CollapseLayerAction = FUIAction(FExecuteAction::CreateLambda([SharedThis, InLayerIndex] { SharedThis->CollapseLayer(InLayerIndex); }));
-						MenuBuilder.AddMenuEntry(LOCTEXT("CollapseLayer", "Collapse..."), LOCTEXT("CollapseLayerTooltip", "Collapse layer into top neighbor layer"), FSlateIcon(), CollapseLayerAction);
+						FUIAction LayerUIAction = FUIAction(
+							FExecuteAction::CreateSPLambda(this, [this, InLayerIndex, LayerAction] { ExecuteCustomLayerAction(InLayerIndex, LayerAction); }),
+							FCanExecuteAction::CreateSPLambda(this, [this, InLayerIndex, LayerAction] { FText Reason; return CanExecuteCustomLayerAction(InLayerIndex, LayerAction, Reason); }));
+						TAttribute<FText> LayerActionTooltipText = MakeAttributeSPLambda(this, [this, InLayerIndex, LayerAction] { FText Reason; CanExecuteCustomLayerAction(InLayerIndex, LayerAction, Reason); return Reason; });
+						MenuBuilder.AddMenuEntry(LayerAction.GetLabel(), LayerActionTooltipText, FSlateIcon(), LayerUIAction);
 					}
 				}
-
-				if (Landscape->GetLandscapeSplinesReservedLayer() != Landscape->GetLayer(InLayerIndex))
-				{
-					// Reserve for Landscape Splines
-					FUIAction ReserveLayerAction = FUIAction(FExecuteAction::CreateLambda([SharedThis, InLayerIndex] { SharedThis->SetLandscapeSplinesReservedLayer(InLayerIndex); }));
-					MenuBuilder.AddMenuEntry(LOCTEXT("ReserveLayerForSplines", "Reserve for Splines"), LOCTEXT("ReserveLayerForSplinesTooltip", "Reserve Layer for Landscape Splines"), FSlateIcon(), ReserveLayerAction);
-				}
-				else
-				{
-					FUIAction RemoveReserveLayerAction = FUIAction(FExecuteAction::CreateLambda([SharedThis, InLayerIndex] { SharedThis->SetLandscapeSplinesReservedLayer(INDEX_NONE); }));
-					MenuBuilder.AddMenuEntry(LOCTEXT("RemoveReserveLayerForSplines", "Remove Reserve for Splines"), LOCTEXT("RemoveReserveLayerForSplinesTooltip", "Remove reservation of Layer for Landscape Splines"), FSlateIcon(), RemoveReserveLayerAction);
-
-					FUIAction ForceUpdateSplinesAction = FUIAction(FExecuteAction::CreateLambda([SharedThis, InLayerIndex] { SharedThis->ForceUpdateSplines(); }));
-					MenuBuilder.AddMenuEntry(LOCTEXT("UpdateSplines", "Update Splines"), LOCTEXT("UpdateSplinesTooltip", "Update Landscape Splines"), FSlateIcon(), ForceUpdateSplinesAction);
-				}
 			}
 		}
 		MenuBuilder.EndSection();
@@ -499,23 +593,27 @@ TSharedPtr<SWidget> FLandscapeEditorCustomNodeBuilder_Layers::OnLayerContextMenu
 				if (Layer->bVisible)
 				{
 					// Hide Selected Layer
-					FUIAction HideSelectedLayerAction = FUIAction(FExecuteAction::CreateLambda([SharedThis, InLayerIndex] { SharedThis->OnToggleVisibility(InLayerIndex); }));
-					MenuBuilder.AddMenuEntry(LOCTEXT("HideSelected", "Hide Selected"), LOCTEXT("HideSelectedLayerTooltip", "Hide Selected Layer"), FSlateIcon(), HideSelectedLayerAction);
+					FUIAction HideSelectedLayerAction = FUIAction(FExecuteAction::CreateSPLambda(this, [this, InLayerIndex] { OnToggleVisibility(InLayerIndex); }),
+						FCanExecuteAction::CreateSPLambda(this, [this, InLayerIndex] { FText Reason; return CanToggleVisibility(InLayerIndex, Reason); }));
+					TAttribute<FText> HideSelectedLayerActionTooltipText = MakeAttributeSPLambda(this, [this, InLayerIndex] { FText Reason; CanToggleVisibility(InLayerIndex, Reason); return Reason; });
+					MenuBuilder.AddMenuEntry(LOCTEXT("HideSelected", "Hide Selected"), HideSelectedLayerActionTooltipText, FSlateIcon(), HideSelectedLayerAction);
 				}
 				else
 				{
 					// Show Selected Layer
-					FUIAction ShowSelectedLayerAction = FUIAction(FExecuteAction::CreateLambda([SharedThis, InLayerIndex] { SharedThis->OnToggleVisibility(InLayerIndex); }));
-					MenuBuilder.AddMenuEntry(LOCTEXT("ShowSelected", "Show Selected"), LOCTEXT("ShowSelectedLayerTooltip", "Show Selected Layer"), FSlateIcon(), ShowSelectedLayerAction);
+					FUIAction ShowSelectedLayerAction = FUIAction(FExecuteAction::CreateSPLambda(this, [this, InLayerIndex] { OnToggleVisibility(InLayerIndex); }),
+						FCanExecuteAction::CreateSPLambda(this, [this, InLayerIndex] { FText Reason; return CanToggleVisibility(InLayerIndex, Reason); }));
+					TAttribute<FText> ShowSelectedLayerActionTooltipText = MakeAttributeSPLambda(this, [this, InLayerIndex] { FText Reason; CanToggleVisibility(InLayerIndex, Reason); return Reason; });
+					MenuBuilder.AddMenuEntry(LOCTEXT("ShowSelected", "Show Selected"), ShowSelectedLayerActionTooltipText, FSlateIcon(), ShowSelectedLayerAction);
 				}
 
 				// Show Only Selected Layer
-				FUIAction ShowOnlySelectedLayerAction = FUIAction(FExecuteAction::CreateLambda([SharedThis, InLayerIndex] { SharedThis->ShowOnlySelectedLayer(InLayerIndex); }));
+				FUIAction ShowOnlySelectedLayerAction = FUIAction(FExecuteAction::CreateSPLambda(this, [this, InLayerIndex] { ShowOnlySelectedLayer(InLayerIndex); }));
 				MenuBuilder.AddMenuEntry(LOCTEXT("ShowOnlySelected", "Show Only Selected"), LOCTEXT("ShowOnlySelectedLayerTooltip", "Show Only Selected Layer"), FSlateIcon(), ShowOnlySelectedLayerAction);
 			}
 
 			// Show All Layers
-			FUIAction ShowAllLayersAction = FUIAction(FExecuteAction::CreateLambda([SharedThis] { SharedThis->ShowAllLayers(); }));
+			FUIAction ShowAllLayersAction = FUIAction(FExecuteAction::CreateSP(this, &FLandscapeEditorCustomNodeBuilder_Layers::ShowAllLayers));
 			MenuBuilder.AddMenuEntry(LOCTEXT("ShowAllLayers", "Show All Layers"), LOCTEXT("ShowAllLayersTooltip", "Show All Layers"), FSlateIcon(), ShowAllLayersAction);
 		}
 		MenuBuilder.EndSection();
@@ -542,71 +640,74 @@ TSharedPtr<SWidget> FLandscapeEditorCustomNodeBuilder_Layers::OnLayerContextMenu
 	return nullptr;
 }
 
-void FLandscapeEditorCustomNodeBuilder_Layers::ForceUpdateSplines()
+bool FLandscapeEditorCustomNodeBuilder_Layers::CanRenameLayer(int32 InLayerIndex, FText& OutReason) const
 {
 	FEdModeLandscape* LandscapeEdMode = GetEditorMode();
-	ALandscape* Landscape = LandscapeEdMode ? LandscapeEdMode->GetLandscape() : nullptr;
-	if (Landscape && Landscape->HasLayersContent() && Landscape->GetLandscapeSplinesReservedLayer())
+	const FLandscapeLayer* Layer = LandscapeEdMode->GetLayer(InLayerIndex);
+	check(Layer != nullptr);
+
+	if (Layer->bLocked)
 	{
-		const bool bUpdateOnlySelection = false;
-		const bool bForceUpdate = true;
-		Landscape->UpdateLandscapeSplines(FGuid(), bUpdateOnlySelection, bForceUpdate);
+		OutReason = LOCTEXT("RenameLayer_CantRenameLocked", "Cannot rename a locked edit layer");
+		return false;
 	}
+
+	OutReason = LOCTEXT("RenameLayer_CanRename", "Rename the edit layer");
+	return true;
 }
 
-void FLandscapeEditorCustomNodeBuilder_Layers::SetLandscapeSplinesReservedLayer(int32 InLayerIndex)
+void FLandscapeEditorCustomNodeBuilder_Layers::RenameLayer(int32 InLayerIndex)
 {
-	FEdModeLandscape* LandscapeEdMode = GetEditorMode();
-	ALandscape* Landscape = LandscapeEdMode ? LandscapeEdMode->GetLandscape() : nullptr;
-	if (Landscape)
+	if (InlineTextBlocks.IsValidIndex(InLayerIndex) && InlineTextBlocks[InLayerIndex].IsValid())
 	{
-		const FLandscapeLayer* CurrenReservedLayer = Landscape->GetLandscapeSplinesReservedLayer();
-		const FLandscapeLayer* NewReservedLayer = Landscape->GetLayer(InLayerIndex);
-		if (NewReservedLayer != CurrenReservedLayer)
-		{
-			EAppReturnType::Type Result = EAppReturnType::No;
-			if (NewReservedLayer)
-			{
-				Result = FMessageDialog::Open(EAppMsgType::YesNo, FText::Format(LOCTEXT("Landscape_SetReservedForSplines_Message", "Reserving layer {0} for landscape splines will clear it from its content and no editing will be allowed.  Continue?"), FText::FromName(NewReservedLayer->Name)));
-			}
-			else if (CurrenReservedLayer)
-			{
-				Result = FMessageDialog::Open(EAppMsgType::YesNo, FText::Format(LOCTEXT("Landscape_UnsetReservedForSplines_Message", "Removing reservation of layer {0} for landscape splines will clear its content.  Continue?"), FText::FromName(CurrenReservedLayer->Name)));
-			}
-
-			if (Result == EAppReturnType::Yes)
-			{
-				const FScopedTransaction Transaction(LOCTEXT("ReserveForSplines", "Reserve Landscape Layer for Splines"));
-				Landscape->SetLandscapeSplinesReservedLayer(InLayerIndex);
-				LandscapeEdMode->RefreshDetailPanel();
-				LandscapeEdMode->AutoUpdateDirtyLandscapeSplines();
-				OnLayerSelectionChanged(InLayerIndex);
-			}
-		}
+		InlineTextBlocks[InLayerIndex]->EnterEditingMode();
 	}
 }
 
-void FLandscapeEditorCustomNodeBuilder_Layers::RenameLayer(int32 InLayerIndex)
+
+bool FLandscapeEditorCustomNodeBuilder_Layers::CanClearTargetLayerOnLayer(int32 InLayerIndex, ULandscapeLayerInfoObject* InLayerInfo, FText& OutReason) const
 {
-	if (InlineTextBlocks.IsValidIndex(InLayerIndex) && InlineTextBlocks[InLayerIndex].IsValid())
+	FEdModeLandscape* LandscapeEdMode = GetEditorMode();
+	const FLandscapeLayer* Layer = LandscapeEdMode->GetLayer(InLayerIndex);
+	check(Layer != nullptr);
+	const ULandscapeEditLayerBase* EditLayer = Layer->EditLayer;
+	check(EditLayer != nullptr);
+
+	// Visibility case : 
+	if (InLayerInfo->LayerName == UMaterialExpressionLandscapeVisibilityMask::ParameterName)
 	{
-		InlineTextBlocks[InLayerIndex]->EnterEditingMode();
+		if (!EditLayer->SupportsTargetType(ELandscapeToolTargetType::Visibility))
+		{
+			OutReason = FText::Format(LOCTEXT("ClearTargetLayer_CantClearVisibilityOnLayerType", "Cannot clear visibility : the type of layer {0} ({1}) doesn't support visibility"),
+				FText::FromName(Layer->Name), EditLayer->GetClass()->GetDisplayNameText());
+			return false;
+		}
 	}
+		
+	if (!EditLayer->SupportsTargetType(ELandscapeToolTargetType::Weightmap))
+	{
+		OutReason = FText::Format(LOCTEXT("ClearTargetLayer_CantClearWeightmapOnLayerType", "Cannot clear weightmap : the type of layer {0} ({1}) doesn't support weightmaps"),
+			FText::FromName(Layer->Name), EditLayer->GetClass()->GetDisplayNameText());
+		return false;
+	}
+
+	OutReason = LOCTEXT("ClearTargetLayer_CanClear", "Clear the target layer on this edit layer");
+	return true;
 }
 
-void FLandscapeEditorCustomNodeBuilder_Layers::ClearPaintLayer(int32 InLayerIndex, ULandscapeLayerInfoObject* InLayerInfo)
+void FLandscapeEditorCustomNodeBuilder_Layers::ClearTargetLayerOnLayer(int32 InLayerIndex, ULandscapeLayerInfoObject* InLayerInfo)
 {
 	FEdModeLandscape* LandscapeEdMode = GetEditorMode();
 	ALandscape* Landscape = LandscapeEdMode ? LandscapeEdMode->GetLandscape() : nullptr;
 	if (Landscape)
 	{
-		FLandscapeLayer* Layer = LandscapeEdMode->GetLayer(InLayerIndex);
+		const FLandscapeLayer* Layer = LandscapeEdMode->GetLayer(InLayerIndex);
 		if (Layer)
 		{
-			EAppReturnType::Type Result = FMessageDialog::Open(EAppMsgType::YesNo, FText::Format(LOCTEXT("Landscape_ClearPaintLayer_Message", "The layer {0} : {1} content will be completely cleared.  Continue?"), FText::FromName(Layer->Name), FText::FromName(InLayerInfo->LayerName)));
+			EAppReturnType::Type Result = FMessageDialog::Open(EAppMsgType::YesNo, FText::Format(LOCTEXT("Landscape_ClearTargetLayer_Message", "The layer {0} : {1} content will be completely cleared.  Continue?"), FText::FromName(Layer->Name), FText::FromName(InLayerInfo->LayerName)));
 			if (Result == EAppReturnType::Yes)
 			{
-				const FScopedTransaction Transaction(LOCTEXT("Landscape_Layers_PaintClear", "Clear Paint Layer"));
+				const FScopedTransaction Transaction(LOCTEXT("Landscape_Layers_TargetClear", "Clear Target Layer"));
 				Landscape->ClearPaintLayer(InLayerIndex, InLayerInfo);
 				LandscapeEdMode->RequestUpdateShownLayerList();
 			}
@@ -614,13 +715,77 @@ void FLandscapeEditorCustomNodeBuilder_Layers::ClearPaintLayer(int32 InLayerInde
 	}
 }
 
-void FLandscapeEditorCustomNodeBuilder_Layers::ClearLayer(int32 InLayerIndex, ELandscapeClearMode InClearMode)
+bool FLandscapeEditorCustomNodeBuilder_Layers::CanClearLayer(int32 InLayerIndex, FText& OutReason) const
+{
+	FEdModeLandscape* LandscapeEdMode = GetEditorMode();
+	const FLandscapeLayer* Layer = LandscapeEdMode->GetLayer(InLayerIndex);
+	check(Layer != nullptr);
+	const ULandscapeEditLayerBase* EditLayer = Layer->EditLayer;
+	check(EditLayer != nullptr);
+
+	if (Layer->bLocked)
+	{
+		OutReason = LOCTEXT("ClearLayer_CantClearLocked", "Cannot clear a locked edit layer");
+		return false;
+	}
+
+	if (!EditLayer->SupportsEditingTools())
+	{
+		OutReason = FText::Format(LOCTEXT("ClearLayer_CantClearLayerWithoutPersistentTextures", "Cannot clear an edit layer which doesn't have editable textures (procedural)"),
+			FText::FromName(Layer->Name), EditLayer->GetClass()->GetDisplayNameText());
+		return false;
+	}
+
+	OutReason = LOCTEXT("ClearLayer_CanClear", "Clear the edit layer");
+	return true;
+}
+
+bool FLandscapeEditorCustomNodeBuilder_Layers::CanClearTargetLayersOnLayer(int32 InLayerIndex, ELandscapeClearMode InClearMode, FText& OutReason) const
+{
+	FEdModeLandscape* LandscapeEdMode = GetEditorMode();
+	const FLandscapeLayer* Layer = LandscapeEdMode->GetLayer(InLayerIndex);
+	check(Layer != nullptr);
+	const ULandscapeEditLayerBase* EditLayer = Layer->EditLayer;
+	check(EditLayer != nullptr);
+
+	// TODO [jonathan.bard] : Ideally, ELandscapeClearMode should be deprecated and we would treat Visibility and Weightmap separately but for now just keep on treating visibility as a standard weightmap : 
+	if ((InClearMode == ELandscapeClearMode::Clear_Weightmap)
+		&& !EditLayer->SupportsTargetType(ELandscapeToolTargetType::Weightmap) 
+		&& !EditLayer->SupportsTargetType(ELandscapeToolTargetType::Visibility))
+	{
+			OutReason = FText::Format(LOCTEXT("ClearTargetLayers_CantClearWeightmapsOnLayerType", "Cannot clear weightmaps : the type of layer {0} ({1}) doesn't support weightmaps"), 
+				FText::FromName(Layer->Name), EditLayer->GetClass()->GetDisplayNameText());
+			return false;
+	}
+	else if ((InClearMode == ELandscapeClearMode::Clear_Heightmap)
+		&& !EditLayer->SupportsTargetType(ELandscapeToolTargetType::Heightmap))
+	{
+		OutReason = FText::Format(LOCTEXT("ClearTargetLayers_CantClearHeightmapOnLayerType", "Cannot clear heightmap : the type of layer {0} ({1}) doesn't support heightmaps"),
+			FText::FromName(Layer->Name), EditLayer->GetClass()->GetDisplayNameText());
+			return false;
+	}
+	if ((InClearMode == ELandscapeClearMode::Clear_All)
+		&& !EditLayer->SupportsTargetType(ELandscapeToolTargetType::Heightmap)
+		&& !EditLayer->SupportsTargetType(ELandscapeToolTargetType::Weightmap)
+		&& !EditLayer->SupportsTargetType(ELandscapeToolTargetType::Visibility))
+	if ((InClearMode == ELandscapeClearMode::Clear_Heightmap) && !EditLayer->SupportsTargetType(ELandscapeToolTargetType::Heightmap))
+	{
+		OutReason = FText::Format(LOCTEXT("ClearTargetLayers_CantClearOnLayerType", "Cannot clear : the type of layer {0} ({1}) doesn't support heightmaps or weightmaps"),
+			FText::FromName(Layer->Name), EditLayer->GetClass()->GetDisplayNameText());
+		return false;
+	}
+
+	OutReason = LOCTEXT("ClearTargetLayers_CanClear", "Clear the target layers on the edit layer");
+	return true;
+}
+
+void FLandscapeEditorCustomNodeBuilder_Layers::ClearTargetLayersOnLayer(int32 InLayerIndex, ELandscapeClearMode InClearMode)
 {
 	FEdModeLandscape* LandscapeEdMode = GetEditorMode();
 	ALandscape* Landscape = LandscapeEdMode ? LandscapeEdMode->GetLandscape() : nullptr;
 	if (Landscape)
 	{
-		FLandscapeLayer* Layer = LandscapeEdMode->GetLayer(InLayerIndex);
+		const FLandscapeLayer* Layer = LandscapeEdMode->GetLayer(InLayerIndex);
 		if (Layer)
 		{
 			EAppReturnType::Type Result = FMessageDialog::Open(EAppMsgType::YesNo, FText::Format(LOCTEXT("Landscape_ClearLayer_Message", "The layer {0} content will be completely cleared.  Continue?"), FText::FromName(Layer->Name)));
@@ -638,13 +803,37 @@ void FLandscapeEditorCustomNodeBuilder_Layers::ClearLayer(int32 InLayerIndex, EL
 	}
 }
 
+bool FLandscapeEditorCustomNodeBuilder_Layers::CanDeleteLayer(int32 InLayerIndex, FText& OutReason) const
+{
+	FEdModeLandscape* LandscapeEdMode = GetEditorMode();
+	ALandscape* Landscape = LandscapeEdMode ? LandscapeEdMode->GetLandscape() : nullptr;
+	if (!Landscape || Landscape->GetLayers().Num() <= 1)
+	{
+		OutReason = LOCTEXT("DeleteLayer_CantDeleteLastLayer", "The last layer cannot be deleted");
+		return false;
+	}
+
+	const FLandscapeLayer* Layer = LandscapeEdMode->GetLayer(InLayerIndex);
+	check(Layer != nullptr);
+
+	if (Layer->bLocked)
+	{
+		OutReason = LOCTEXT("DeleteLayer_CantDeleteLocked", "Cannot delete a locked edit layer");
+		return false;
+	}
+
+
+	OutReason = LOCTEXT("DeleteLayer_CanDelete", "Delete the edit layer");
+	return true;
+}
+
 void FLandscapeEditorCustomNodeBuilder_Layers::DeleteLayer(int32 InLayerIndex)
 {
 	FEdModeLandscape* LandscapeEdMode = GetEditorMode();
 	ALandscape* Landscape = LandscapeEdMode ? LandscapeEdMode->GetLandscape() : nullptr;
-	if (Landscape && Landscape->LandscapeLayers.Num() > 1)
+	if (Landscape && Landscape->GetLayers().Num() > 1)
 	{
-		FLandscapeLayer* Layer = LandscapeEdMode->GetLayer(InLayerIndex);
+		const FLandscapeLayer* Layer = LandscapeEdMode->GetLayer(InLayerIndex);
 		if (Layer)
 		{
 			EAppReturnType::Type Result = FMessageDialog::Open(EAppMsgType::YesNo, FText::Format(LOCTEXT("Landscape_DeleteLayer_Message", "The layer {0} will be deleted.  Continue?"), FText::FromName(Layer->Name)));
@@ -652,7 +841,7 @@ void FLandscapeEditorCustomNodeBuilder_Layers::DeleteLayer(int32 InLayerIndex)
 			{
 				const FScopedTransaction Transaction(LOCTEXT("Landscape_Layers_Delete", "Delete Layer"));
 				Landscape->DeleteLayer(InLayerIndex);
-				int32 NewLayerSelectionIndex = Landscape->GetLayer(InLayerIndex) ? InLayerIndex : 0;
+				int32 NewLayerSelectionIndex = Landscape->GetLayerConst(InLayerIndex) ? InLayerIndex : 0;
 				OnLayerSelectionChanged(NewLayerSelectionIndex);
 				LandscapeEdMode->RefreshDetailPanel();
 			}
@@ -664,75 +853,134 @@ bool FLandscapeEditorCustomNodeBuilder_Layers::CanCollapseLayer(int32 InLayerInd
 {
 	FEdModeLandscape* LandscapeEdMode = GetEditorMode();
 	ALandscape* Landscape = LandscapeEdMode ? LandscapeEdMode->GetLandscape() : nullptr;
-	if (!Landscape || Landscape->LandscapeLayers.Num() <= 1)
+	if (!Landscape || Landscape->GetLayers().Num() <= 1)
 	{
-		OutReason = LOCTEXT("Landscape_CollapseLayer_Reason_NotEnoughLayersToCollapse", "Not enough layers to do Collapse.");
+		OutReason = LOCTEXT("Landscape_CollapseLayer_Reason_NotEnoughLayersToCollapse", "Not enough layers to do collapse");
 		return false;
 	}
 
 	if (InLayerIndex < 1)
 	{
-		OutReason = LOCTEXT("Landscape_CollapseLayer_Reason_CantCollapseBaseLayer", "Can't Collapse first layer.");
+		OutReason = LOCTEXT("Landscape_CollapseLayer_Reason_CantCollapseBaseLayer", "Cannot collapse the first layer");
 		return false;
 	}
 	
-	FLandscapeLayer* SplineLayer = Landscape->GetLandscapeSplinesReservedLayer();
-	if (SplineLayer == LandscapeEdMode->GetLayer(InLayerIndex))
+	const FLandscapeLayer* TopLayer = LandscapeEdMode->GetLayer(InLayerIndex);
+	const FLandscapeLayer* BottomLayer = LandscapeEdMode->GetLayer(InLayerIndex - 1);
+	check((TopLayer != nullptr) && (BottomLayer != nullptr));
+
+	const ULandscapeEditLayerBase* TopEditLayer = TopLayer->EditLayer;
+	const ULandscapeEditLayerBase* BottomEditLayer = BottomLayer->EditLayer;
+	check((TopEditLayer != nullptr) && (BottomEditLayer != nullptr));
+
+	if (!TopEditLayer->SupportsCollapsingTo())
 	{
-		OutReason = LOCTEXT("Landscape_CollapseLayer_Reason_CantCollapseSplineLayer", "Can't Collapse reserved spline layer.");
+		OutReason = FText::Format(LOCTEXT("Landscape_CollapseLayer_Reason_TopLayerDoesntSupportCollapsing", "Cannot collapse layer '{0}' onto layer '{1}'. The type of layer '{0}' ({2}) doesn't support collapsing to another one"), 
+			FText::FromName(TopLayer->Name), FText::FromName(BottomLayer->Name), TopEditLayer->GetClass()->GetDisplayNameText());
 		return false;
 	}
 
-	FLandscapeLayer* BaseLayer = LandscapeEdMode->GetLayer(InLayerIndex - 1);
-	if (!BaseLayer)
+	if (!BottomEditLayer->SupportsBeingCollapsedAway())
 	{
-		OutReason = LOCTEXT("Landscape_CollapseLayer_Reason_InvalidBaseLayer", "Invalid base layer.");
+		OutReason = FText::Format(LOCTEXT("Landscape_CollapseLayer_Reason_BottomLayerDoesntSupportCollapsing", "Cannot collapse layer '{0}' onto layer '{1}'. The type of layer '{1}' ({2}) doesn't support being collapsed away"),
+			FText::FromName(TopLayer->Name), FText::FromName(BottomLayer->Name), BottomEditLayer->GetClass()->GetDisplayNameText());
 		return false;
 	}
 
-	if (SplineLayer == BaseLayer)
+	if (TopLayer->bLocked)
 	{
-		OutReason = LOCTEXT("Landscape_CollapseLayer_Reason_CantCollapseOnSplineLayer", "Can't Collapse on reserved spline layer.");
+		OutReason = FText::Format(LOCTEXT("Landscape_CollapseLayer_Reason_TopLayerIsLocked", "Cannot collapse layer '{0}' onto layer '{1}'. Layer '{0}' will be deleted in the operation but it is currently locked"),
+			FText::FromName(TopLayer->Name), FText::FromName(BottomLayer->Name));
 		return false;
 	}
-	
+
+	if (BottomLayer->bLocked)
+	{
+		OutReason = FText::Format(LOCTEXT("Landscape_CollapseLayer_Reason_BottomLayerIsLocked", "Cannot collapse layer '{0}' onto layer '{1}'. Destination layer '{1}' is currently locked"),
+			FText::FromName(TopLayer->Name), FText::FromName(BottomLayer->Name));
+		return false;
+	}
+
 	// Can't collapse on layer that has a Brush because result will change...
-	if (BaseLayer->Brushes.Num() > 0)
+	if (TopLayer->Brushes.Num() > 0)
 	{
-		OutReason = FText::Format(LOCTEXT("Landscape_CollapseLayer_Reason_CantCollapseOnLayerWithBrush", "Can't Collapse because base layer '{0}' contains Brush(es)."), FText::FromName(BaseLayer->Name));
+		OutReason = FText::Format(LOCTEXT("Landscape_CollapseLayer_Reason_TopLayerHasBrush", "Cannot collapse layer '{0}' onto layer '{1}'. Layer '{0}' contains brush(es)"),
+			FText::FromName(TopLayer->Name), FText::FromName(BottomLayer->Name));
 		return false;
 	}
-	
 
+	if (BottomLayer->Brushes.Num() > 0)
+	{
+		OutReason = FText::Format(LOCTEXT("Landscape_CollapseLayer_Reason_TopLayerHasBrush", "Cannot collapse layer '{0}' onto layer '{1}'. Layer '{1}' contains brush(es)"),
+			FText::FromName(TopLayer->Name), FText::FromName(BottomLayer->Name));
+		return false;
+	}
+	
+	OutReason = FText::Format(LOCTEXT("Landscape_CollapseLayer_Reason_Collapse", "Collapse layer '{0}' onto layer '{1}'"),
+		FText::FromName(TopLayer->Name), FText::FromName(BottomLayer->Name));
 	return true;
 }
 
 void FLandscapeEditorCustomNodeBuilder_Layers::CollapseLayer(int32 InLayerIndex)
 {
 	FText Reason;
-	if (CanCollapseLayer(InLayerIndex, Reason))
+	check(CanCollapseLayer(InLayerIndex, Reason));
+
+	if (FEdModeLandscape* LandscapeEdMode = GetEditorMode())
 	{
-		if (FEdModeLandscape* LandscapeEdMode = GetEditorMode())
+		ALandscape* Landscape = LandscapeEdMode->GetLandscape();
+		const FLandscapeLayer* Layer = LandscapeEdMode->GetLayer(InLayerIndex);
+		const FLandscapeLayer* BaseLayer = LandscapeEdMode->GetLayer(InLayerIndex - 1);
+		if (Landscape && Layer && BaseLayer)
 		{
-			ALandscape* Landscape = LandscapeEdMode->GetLandscape();
-			FLandscapeLayer* Layer = LandscapeEdMode->GetLayer(InLayerIndex);
-			FLandscapeLayer* BaseLayer = LandscapeEdMode->GetLayer(InLayerIndex - 1);
-			if (Landscape && Layer && BaseLayer)
+			EAppReturnType::Type Result = FMessageDialog::Open(EAppMsgType::YesNo, FText::Format(LOCTEXT("Landscape_CollapseLayer_Message", "The layer {0} will be collapsed into layer {1}.  Continue?"), FText::FromName(Layer->Name), FText::FromName(BaseLayer->Name)));
+			if (Result == EAppReturnType::Yes)
 			{
-				EAppReturnType::Type Result = FMessageDialog::Open(EAppMsgType::YesNo, FText::Format(LOCTEXT("Landscape_CollapseLayer_Message", "The layer {0} will be collapsed into layer {1}.  Continue?"), FText::FromName(Layer->Name), FText::FromName(BaseLayer->Name)));
-				if (Result == EAppReturnType::Yes)
-				{
-					const FScopedTransaction Transaction(LOCTEXT("Landscape_Layers_Collapse", "Collapse Layer"));
-					Landscape->CollapseLayer(InLayerIndex);
-					OnLayerSelectionChanged(InLayerIndex - 1);
-					LandscapeEdMode->RefreshDetailPanel();
-				}
+				const FScopedTransaction Transaction(LOCTEXT("Landscape_Layers_Collapse", "Collapse Layer"));
+				Landscape->CollapseLayer(InLayerIndex);
+				OnLayerSelectionChanged(InLayerIndex - 1);
+				LandscapeEdMode->RefreshDetailPanel();
 			}
 		}
 	}
-	else
+}
+
+bool FLandscapeEditorCustomNodeBuilder_Layers::CanExecuteCustomLayerAction(int32 InLayerIndex, const ULandscapeEditLayerBase::FEditLayerAction& InCustomLayerAction, FText& OutReason) const
+{
+	FEdModeLandscape* LandscapeEdMode = GetEditorMode();
+	ALandscape* Landscape = LandscapeEdMode ? LandscapeEdMode->GetLandscape() : nullptr;
+	if (Landscape == nullptr)
 	{
-		FMessageDialog::Open(EAppMsgType::Ok, Reason);
+		OutReason = LOCTEXT("Landscape_CannotExecuteCustomLayerAction_Invalid", "Invalid landscape");
+		return false;
+	}
+
+	const FLandscapeLayer* Layer = LandscapeEdMode->GetLayer(InLayerIndex);
+	check(Layer != nullptr);
+	const ULandscapeEditLayerBase* EditLayer = Layer->EditLayer;
+	check(EditLayer != nullptr);
+	ULandscapeEditLayerBase::FEditLayerAction::FExecuteParams ExecuteParams(Layer, Landscape);
+	return InCustomLayerAction.GetCanExecuteDelegate().Execute(ExecuteParams, OutReason);
+}
+
+void FLandscapeEditorCustomNodeBuilder_Layers::ExecuteCustomLayerAction(int32 InLayerIndex, const ULandscapeEditLayerBase::FEditLayerAction& InCustomLayerAction)
+{
+	FEdModeLandscape* LandscapeEdMode = GetEditorMode();
+	if (ALandscape* Landscape = LandscapeEdMode ? LandscapeEdMode->GetLandscape() : nullptr)
+	{
+		const FLandscapeLayer* Layer = LandscapeEdMode->GetLayer(InLayerIndex);
+		check(Layer != nullptr);
+		const ULandscapeEditLayerBase* EditLayer = Layer->EditLayer;
+		check(EditLayer != nullptr);
+
+		const FScopedTransaction Transaction(InCustomLayerAction.GetLabel());
+		ULandscapeEditLayerBase::FEditLayerAction::FExecuteParams ExecuteParams(Layer, Landscape);
+		ULandscapeEditLayerBase::FEditLayerAction::FExecuteResult Result = InCustomLayerAction.GetExecuteDelegate().Execute(ExecuteParams);
+		if (!Result.bSuccess)
+		{
+			// Indicate to the user that the action failed : 
+			FMessageDialog::Open(EAppMsgType::Ok, Result.Reason);
+		}
 	}
 }
 
@@ -759,19 +1007,109 @@ void FLandscapeEditorCustomNodeBuilder_Layers::ShowAllLayers()
 	}
 }
 
+TSubclassOf<ULandscapeEditLayerBase> FLandscapeEditorCustomNodeBuilder_Layers::PickEditLayerClass() const
+{
+	class FLandscapeEditLayerClassFilter : public IClassViewerFilter
+	{
+	public:
+		FLandscapeEditLayerClassFilter()
+		{
+			AllowedChildrenOfClasses.Add(ULandscapeEditLayerBase::StaticClass());
+			DisallowedClassFlags = CLASS_Abstract | CLASS_Deprecated | CLASS_NewerVersionExists;
+		};
+
+		virtual bool IsClassAllowed(const FClassViewerInitializationOptions& InInitOptions, const UClass* InClass, TSharedRef< class FClassViewerFilterFuncs > InFilterFuncs) override
+		{
+			bool bIsCorrectClass = InFilterFuncs->IfInChildOfClassesSet(AllowedChildrenOfClasses, InClass) != EFilterReturn::Failed;
+			bool bValidClassFlags = !InClass->HasAnyClassFlags(DisallowedClassFlags);
+
+			return (bIsCorrectClass && bValidClassFlags);
+		};
+
+		virtual bool IsUnloadedClassAllowed(const FClassViewerInitializationOptions& InInitOptions, const TSharedRef< const class IUnloadedBlueprintData > InUnloadedClassData, TSharedRef< class FClassViewerFilterFuncs > InFilterFuncs) override
+		{
+			return false;
+		};
+
+	private:
+		/** All children of these classes will be included unless filtered out by another setting. */
+		TSet<const UClass*> AllowedChildrenOfClasses;
+
+		/** Disallowed class flags. */
+		EClassFlags DisallowedClassFlags;
+	};
+
+	// Load the classviewer module to display a class picker
+	FClassViewerModule& ClassViewerModule = FModuleManager::LoadModuleChecked<FClassViewerModule>("ClassViewer");
+
+	// Fill in options
+	FClassViewerInitializationOptions Options;
+	Options.Mode = EClassViewerMode::ClassPicker;
+
+	TSharedPtr<FLandscapeEditLayerClassFilter> Filter = MakeShareable(new FLandscapeEditLayerClassFilter());
+	Options.ClassFilters.Add(Filter.ToSharedRef());
+
+	const FText TitleText = LOCTEXT("PickEditLayerClass", "Pick Landscape Edit Layer Class");
+	UClass* ChosenClass = nullptr;
+	SClassPickerDialog::PickClass(TitleText, Options, ChosenClass, ULandscapeEditLayerBase::StaticClass());
+	return TSubclassOf<ULandscapeEditLayerBase>(ChosenClass);
+}
+
 void FLandscapeEditorCustomNodeBuilder_Layers::CreateLayer()
 {
-FEdModeLandscape* LandscapeEdMode = GetEditorMode();
-ALandscape* Landscape = LandscapeEdMode ? LandscapeEdMode->GetLandscape() : nullptr;
-if (Landscape)
+	FEdModeLandscape* LandscapeEdMode = GetEditorMode();
+	ALandscape* Landscape = LandscapeEdMode ? LandscapeEdMode->GetLandscape() : nullptr;
+	if (Landscape)
+	{
+		if (TSubclassOf<ULandscapeEditLayerBase> EditLayerClass = PickEditLayerClass())
+		{
+			// Disallow multiple layers of certain types : 
+			if (!EditLayerClass.GetDefaultObject()->SupportsMultiple())
+			{
+				if (int32 NumLayersOfThisType = Landscape->GetLayersOfType(EditLayerClass).Num(); NumLayersOfThisType > 0)
+				{
+					FMessageDialog::Open(EAppMsgType::Ok, FText::Format(LOCTEXT("Landscape_CreateLayer_CannotCreateMultiple", "Cannot create layer of type {0} : {1} {1}|plural(one=layer, other=layers) of this type already {1}|plural(one=exists, other=exist) and only 1 is allowed"), 
+						EditLayerClass->GetDisplayNameText(), NumLayersOfThisType));
+					return;
+				}
+			}
+
+			const FScopedTransaction Transaction(LOCTEXT("Landscape_Layers_Create", "Create Layer"));
+			Landscape->CreateLayer(FName(EditLayerClass.GetDefaultObject()->GetDefaultName()), EditLayerClass);
+			OnLayerSelectionChanged(Landscape->GetLayerCount() - 1);
+
+			LandscapeEdMode->RefreshDetailPanel();
+		}
+	}
+}
+
+FText FLandscapeEditorCustomNodeBuilder_Layers::GetNumLayersText() const
 {
+	FEdModeLandscape* LandscapeEdMode = GetEditorMode();
+	ALandscape* Landscape = LandscapeEdMode ? LandscapeEdMode->GetLandscape() : nullptr;
+	if (Landscape)
 	{
-		const FScopedTransaction Transaction(LOCTEXT("Landscape_Layers_Create", "Create Layer"));
-		Landscape->CreateLayer();
-		OnLayerSelectionChanged(Landscape->GetLayerCount() - 1);
+		return FText::Format(LOCTEXT("NumEditLayersText", "{0} Edit {0}|plural(one=Layer, other=Layers)"), Landscape->GetLayers().Num());
 	}
-	LandscapeEdMode->RefreshDetailPanel();
+
+	return FText();
 }
+
+bool FLandscapeEditorCustomNodeBuilder_Layers::CanCreateLayer(FText& OutReason) const
+{
+	FEdModeLandscape* LandscapeEdMode = GetEditorMode();
+	ALandscape* Landscape = LandscapeEdMode ? LandscapeEdMode->GetLandscape() : nullptr;
+	if (Landscape)
+	{
+		if (Landscape->IsMaxLayersReached())
+		{
+			OutReason = LOCTEXT("CreateLayerTooltip_MaxLayersReached", "Creates a new edit layer.\nCurrently disabled as the max number of layers has been reached. This can be adjusted in the landscape project settings : MaxNumberOfLayers)");
+			return false;
+		}
+	}
+
+	OutReason = LOCTEXT("CreateLayerTooltip", "Creates a new edit layer");
+	return true;
 }
 
 void FLandscapeEditorCustomNodeBuilder_Layers::OnLayerSelectionChanged(int32 InLayerIndex)
@@ -797,6 +1135,22 @@ TOptional<float> FLandscapeEditorCustomNodeBuilder_Layers::GetLayerAlpha(int32 I
 	return 1.0f;
 }
 
+bool FLandscapeEditorCustomNodeBuilder_Layers::CanSetLayerAlpha(int32 InLayerIndex, FText& OutReason) const
+{
+	FEdModeLandscape* LandscapeEdMode = GetEditorMode();
+	const FLandscapeLayer* Layer = LandscapeEdMode->GetLayer(InLayerIndex);
+	check(Layer != nullptr);
+
+	if (Layer->bLocked)
+	{
+		OutReason = LOCTEXT("SetLayerAlpha_LayerIsLocked", "Cannot change the alpha of a locked edit layer");
+		return false;
+	}
+
+	OutReason = LOCTEXT("SetLayerAlpha_CanSet", "Set the edit layer's alpha");
+	return true;
+}
+
 void FLandscapeEditorCustomNodeBuilder_Layers::SetLayerAlpha(float InAlpha, int32 InLayerIndex, bool bCommit)
 {
 	FEdModeLandscape* LandscapeEdMode = GetEditorMode();
@@ -814,6 +1168,22 @@ void FLandscapeEditorCustomNodeBuilder_Layers::SetLayerAlpha(float InAlpha, int3
 	}
 }
 
+bool FLandscapeEditorCustomNodeBuilder_Layers::CanToggleVisibility(int32 InLayerIndex, FText& OutReason) const
+{
+	FEdModeLandscape* LandscapeEdMode = GetEditorMode();
+	const FLandscapeLayer* Layer = LandscapeEdMode->GetLayer(InLayerIndex);
+	check(Layer != nullptr);
+
+	if (Layer->bLocked)
+	{
+		OutReason = LOCTEXT("ToggleVisibility_CantToggleLocked", "Cannot change the visibility of a locked edit layer");
+		return false;
+	}
+
+	OutReason = LOCTEXT("ToggleVisibility_CanToggle", "Toggle the visibility of the edit layer");
+	return true;
+}
+
 FReply FLandscapeEditorCustomNodeBuilder_Layers::OnToggleVisibility(int32 InLayerIndex)
 {
 	FEdModeLandscape* LandscapeEdMode = GetEditorMode();
@@ -854,13 +1224,12 @@ EVisibility FLandscapeEditorCustomNodeBuilder_Layers::GetLayerAlphaVisibility(in
 	return bIsVisible ? EVisibility::Visible : EVisibility::Hidden;
 }
 
-bool FLandscapeEditorCustomNodeBuilder_Layers::IsLayerEditionEnabled(int32 InLayerIndex) const
+TSharedPtr<IToolTip> FLandscapeEditorCustomNodeBuilder_Layers::GetEditLayerTypeTooltip(int32 InLayerIndex) const
 {
 	FEdModeLandscape* LandscapeEdMode = GetEditorMode();
-	ALandscape* Landscape = LandscapeEdMode ? LandscapeEdMode->GetLandscape() : nullptr;
-	const FLandscapeLayer* Layer = LandscapeEdMode ? LandscapeEdMode->GetLayer(InLayerIndex) : nullptr;
-	const FLandscapeLayer* LayerReservedForSplines = Landscape ? Landscape->GetLandscapeSplinesReservedLayer() : nullptr;
-	return Layer && !Layer->bLocked && (Layer != LayerReservedForSplines) && LandscapeEdMode->DoesCurrentToolAffectEditLayers();
+	const ULandscapeEditLayerBase* EditLayer = (LandscapeEdMode != nullptr) ? LandscapeEdMode->GetEditLayer(InLayerIndex) : nullptr;
+	check(EditLayer != nullptr);
+	return FEditorClassUtils::GetTooltip(EditLayer->GetClass());
 }
 
 const FSlateBrush* FLandscapeEditorCustomNodeBuilder_Layers::GetLockBrushForLayer(int32 InLayerIndex) const
@@ -870,6 +1239,17 @@ const FSlateBrush* FLandscapeEditorCustomNodeBuilder_Layers::GetLockBrushForLaye
 	return bIsLocked ? FAppStyle::GetBrush(TEXT("PropertyWindow.Locked")) : FAppStyle::GetBrush(TEXT("PropertyWindow.Unlocked"));
 }
 
+const FSlateBrush* FLandscapeEditorCustomNodeBuilder_Layers::GetEditLayerIconBrush(int32 InLayerIndex) const
+{
+	if (FEdModeLandscape* LandscapeEdMode = GetEditorMode())
+	{
+		const ULandscapeEditLayerBase* EditLayer = LandscapeEdMode->GetEditLayer(InLayerIndex);
+		check(EditLayer != nullptr);
+		return FSlateIconFinder::FindIconBrushForClass(EditLayer->GetClass());
+	}
+	return nullptr;
+}
+
 int32 FLandscapeEditorCustomNodeBuilder_Layers::SlotIndexToLayerIndex(int32 SlotIndex)
 {
 	FEdModeLandscape* LandscapeEdMode = GetEditorMode();
@@ -879,8 +1259,8 @@ int32 FLandscapeEditorCustomNodeBuilder_Layers::SlotIndexToLayerIndex(int32 Slot
 		return INDEX_NONE;
 	}
 	
-	check(Landscape->LandscapeLayers.IsValidIndex(SlotIndex));
-	return Landscape->LandscapeLayers.Num() - SlotIndex - 1;
+	check(Landscape->GetLayers().IsValidIndex(SlotIndex));
+	return Landscape->GetLayers().Num() - SlotIndex - 1;
 }
 
 FReply FLandscapeEditorCustomNodeBuilder_Layers::HandleDragDetected(const FGeometry& MyGeometry, const FPointerEvent& MouseEvent, int32 SlotIndex, SVerticalBox::FSlot* Slot)
@@ -889,7 +1269,7 @@ FReply FLandscapeEditorCustomNodeBuilder_Layers::HandleDragDetected(const FGeome
 	if (LandscapeEdMode && LandscapeEdMode->DoesCurrentToolAffectEditLayers())
 	{
 		int32 LayerIndex = SlotIndexToLayerIndex(SlotIndex);
-		FLandscapeLayer* Layer = LandscapeEdMode->GetLayer(LayerIndex);
+		const FLandscapeLayer* Layer = LandscapeEdMode->GetLayer(LayerIndex);
 		if (Layer && !Layer->bLocked)
 		{
 			TSharedPtr<SWidget> Row = GenerateRow(LayerIndex);
@@ -921,8 +1301,8 @@ FReply FLandscapeEditorCustomNodeBuilder_Layers::HandleAcceptDrop(FDragDropEvent
 		return FReply::Unhandled();
 	}
 
-	FEdModeLandscape* LandscapeEdMode = GetEditorMode();
-	ALandscape* Landscape = LandscapeEdMode ? LandscapeEdMode->GetLandscape() : nullptr;
+		FEdModeLandscape* LandscapeEdMode = GetEditorMode();
+		ALandscape* Landscape = LandscapeEdMode ? LandscapeEdMode->GetLandscape() : nullptr;
 	if (!Landscape)
 	{
 		return FReply::Unhandled();
diff --git a/Engine/Source/Editor/LandscapeEditor/Private/LandscapeEditorDetailCustomization_Layers.h b/Engine/Source/Editor/LandscapeEditor/Private/LandscapeEditorDetailCustomization_Layers.h
--- a/Engine/Source/Editor/LandscapeEditor/Private/LandscapeEditorDetailCustomization_Layers.h
+++ b/Engine/Source/Editor/LandscapeEditor/Private/LandscapeEditorDetailCustomization_Layers.h
@@ -6,6 +6,7 @@
 #include "Misc/Attribute.h"
 #include "LandscapeEditorDetailCustomization_Base.h"
 #include "LandscapeEdMode.h"
+#include "LandscapeEditLayer.h"
 #include "Layout/Visibility.h"
 #include "Layout/Margin.h"
 #include "Styling/SlateColor.h"
@@ -76,30 +77,50 @@ protected:
 	bool IsLayerSelected(int32 LayerIndex) const;
 	void OnLayerSelectionChanged(int32 LayerIndex);
 	TSharedPtr<SWidget> OnLayerContextMenuOpening(int32 InLayerIndex);
+
+	FText GetNumLayersText() const;
+
+	bool CanCreateLayer(FText& OutReason) const;
 	void CreateLayer();
-	void ClearPaintLayer(int32 InLayerIndex, ULandscapeLayerInfoObject* InLayerInfo);
-	void ClearLayer(int32 InLayerIndex, ELandscapeClearMode InClearMode);
+
+	const FSlateBrush* GetEditLayerIconBrush(int32 InLayerIndex) const;
+
+	bool CanClearTargetLayerOnLayer(int32 InLayerIndex, ULandscapeLayerInfoObject* InLayerInfo, FText& OutReason) const;
+	void ClearTargetLayerOnLayer(int32 InLayerIndex, ULandscapeLayerInfoObject* InLayerInfo);
+	bool CanClearLayer(int32 InLayerIndex, FText& OutReason) const;
+	bool CanClearTargetLayersOnLayer(int32 InLayerIndex, ELandscapeClearMode InClearMode, FText& OutReason) const;
+	void ClearTargetLayersOnLayer(int32 InLayerIndex, ELandscapeClearMode InClearMode);
+
+	bool CanRenameLayerTo(const FText& NewText, FText& OutErrorMessage, int32 InLayerIndex);
+	bool CanRenameLayer(int32 InLayerIndex, FText& OutReason) const;
 	void RenameLayer(int32 InLayerIndex);
+
+	bool CanDeleteLayer(int32 InLayerIndex, FText& OutReason) const;
 	void DeleteLayer(int32 InLayerIndex);
+
 	bool CanCollapseLayer(int32 InLayerIndex, FText& OutReason) const;
 	void CollapseLayer(int32 InLayerIndex);
+
+	bool CanExecuteCustomLayerAction(int32 InLayerIndex, const ULandscapeEditLayerBase::FEditLayerAction& InCustomLayerAction, FText& OutReason) const;
+	void ExecuteCustomLayerAction(int32 InLayerIndex, const ULandscapeEditLayerBase::FEditLayerAction& InCustomLayerAction);
+
 	void ShowOnlySelectedLayer(int32 InLayerIndex);
 	void ShowAllLayers();
-	bool CanRenameLayerTo(const FText& NewText, FText& OutErrorMessage, int32 InLayerIndex);
 	void OnBeginNameTextEdit();
 	void OnEndNameTextEdit();
 	void SetLayerName(const FText& InText, ETextCommit::Type InCommitType, int32 InLayerIndex);
 	FText GetLayerText(int32 InLayerIndex) const;
 	FSlateColor GetLayerTextColor(int32 InLayerIndex) const;
-	void SetLandscapeSplinesReservedLayer(int32 InLayerIndex);
 	FText GetLayerDisplayName(int32 InLayerIndex) const;
-	bool IsLayerEditionEnabled(int32 InLayerIndex) const;
 	EVisibility GetLayerAlphaVisibility(int32 InLayerIndex) const;
-	void ForceUpdateSplines();
+	TSharedPtr<IToolTip> GetEditLayerTypeTooltip(int32 InLayerIndex) const;
+	TSubclassOf<ULandscapeEditLayerBase> PickEditLayerClass() const;
 
 	TOptional<float> GetLayerAlpha(int32 InLayerIndex) const;
+	bool CanSetLayerAlpha(int32 InLayerIndex, FText& OutReason) const;
 	void SetLayerAlpha(float InAlpha, int32 InLayerIndex, bool bCommit);
 	
+	bool CanToggleVisibility(int32 InLayerIndex, FText& OutReason) const;
 	FReply OnToggleVisibility(int32 InLayerIndex);
 	const FSlateBrush* GetVisibilityBrushForLayer(int32 InLayerIndex) const;
 	
@@ -108,7 +129,7 @@ protected:
 
 	void FillAddBrushMenu(FMenuBuilder& MenuBuilder, TArray<ALandscapeBlueprintBrushBase*> Brushes);
 	void AddBrushToCurrentLayer(ALandscapeBlueprintBrushBase* Brush);
-	void FillClearPaintLayerMenu(FMenuBuilder& MenuBuilder, int32 InLayerIndex, TArray<ULandscapeLayerInfoObject*> InUsedLayerInfos);
+	void FillClearTargetLayerMenu(FMenuBuilder& MenuBuilder, int32 InLayerIndex, TArray<ULandscapeLayerInfoObject*> InUsedLayerInfos);
 	void FillClearLayerMenu(FMenuBuilder& MenuBuilder, int32 InLayerIndex);
 
 private:
diff --git a/Engine/Source/Editor/LandscapeEditor/Private/LandscapeEditorDetailCustomization_LayersBrushStack.cpp b/Engine/Source/Editor/LandscapeEditor/Private/LandscapeEditorDetailCustomization_LayersBrushStack.cpp
--- a/Engine/Source/Editor/LandscapeEditor/Private/LandscapeEditorDetailCustomization_LayersBrushStack.cpp
+++ b/Engine/Source/Editor/LandscapeEditor/Private/LandscapeEditorDetailCustomization_LayersBrushStack.cpp
@@ -500,7 +500,7 @@ FReply FLandscapeEditorCustomNodeBuilder_LayersBrushStack::HandleDragDetected(co
 {
 	if (FEdModeLandscape* LandscapeEdMode = GetEditorMode())
 	{
-		FLandscapeLayer* Layer = LandscapeEdMode->GetCurrentLayer();
+		const FLandscapeLayer* Layer = LandscapeEdMode->GetCurrentLayer();
 		if (Layer && !Layer->bLocked)
 		{
 			const TArray<ALandscapeBlueprintBrushBase*>& BrushStack = LandscapeEdMode->GetBrushesForCurrentLayer();
diff --git a/Engine/Source/Editor/LandscapeEditor/Private/LandscapeEditorDetailCustomization_NewLandscape.cpp b/Engine/Source/Editor/LandscapeEditor/Private/LandscapeEditorDetailCustomization_NewLandscape.cpp
--- a/Engine/Source/Editor/LandscapeEditor/Private/LandscapeEditorDetailCustomization_NewLandscape.cpp
+++ b/Engine/Source/Editor/LandscapeEditor/Private/LandscapeEditorDetailCustomization_NewLandscape.cpp
@@ -987,7 +987,7 @@ void FLandscapeEditorDetailCustomization_NewLandscape::AddComponents(ULandscapeI
 			TArray<ULandscapeComponent*> ComponentsUsingHeightmap;
 			ComponentsUsingHeightmap.Add(NewComponent);
 
-			for (const FLandscapeLayer& Layer : Landscape->LandscapeLayers)
+			for (const FLandscapeLayer& Layer : Landscape->GetLayers())
 			{
 				// Since we do not share heightmap when adding new component, we will provided the required array, but they will only be used for 1 component
 				TMap<UTexture2D*, UTexture2D*> CreatedHeightmapTextures;
@@ -1119,7 +1119,7 @@ FReply FLandscapeEditorDetailCustomization_NewLandscape::OnCreateButtonClicked()
 		ReimportHeightmapFilePath = UISettings->ImportLandscape_HeightmapFilename;
 	}
 
-	Landscape->Import(FGuid::NewGuid(), 0, 0, SizeX - 1, SizeY - 1, UISettings->NewLandscape_SectionsPerComponent, QuadsPerSection, HeightDataPerLayers, *ReimportHeightmapFilePath, MaterialLayerDataPerLayers, UISettings->ImportLandscape_AlphamapType);
+	Landscape->Import(FGuid::NewGuid(), 0, 0, SizeX - 1, SizeY - 1, UISettings->NewLandscape_SectionsPerComponent, QuadsPerSection, HeightDataPerLayers, *ReimportHeightmapFilePath, MaterialLayerDataPerLayers, UISettings->ImportLandscape_AlphamapType, TArrayView<const FLandscapeLayer>());
 
 	ULandscapeInfo* LandscapeInfo = Landscape->GetLandscapeInfo();
 	check(LandscapeInfo);
diff --git a/Engine/Source/Editor/LandscapeEditor/Private/LandscapeEditorModule.cpp b/Engine/Source/Editor/LandscapeEditor/Private/LandscapeEditorModule.cpp
--- a/Engine/Source/Editor/LandscapeEditor/Private/LandscapeEditorModule.cpp
+++ b/Engine/Source/Editor/LandscapeEditor/Private/LandscapeEditorModule.cpp
@@ -335,7 +335,7 @@ public:
 	/**
 	* ILandscapeEditorServices implementation
 	*/
-	virtual int32 GetOrCreateEditLayer(FName InEditLayerName, ALandscape* InTargetLandscape) override;
+	virtual int32 GetOrCreateEditLayer(FName InEditLayerName, ALandscape* InTargetLandscape, const TSubclassOf<ULandscapeEditLayerBase>& InEditLayerClass) override;
 	virtual void RefreshDetailPanel() override;
 
 protected:
@@ -522,14 +522,14 @@ FLandscapeImageFileCache& FLandscapeEditorModule::GetImageFileCache() const
 	return *LandscapeImageFileCache;
 }
 
-int32 FLandscapeEditorModule::GetOrCreateEditLayer(FName InEditLayerName, ALandscape* InTargetLandscape)
+int32 FLandscapeEditorModule::GetOrCreateEditLayer(FName InEditLayerName, ALandscape* InTargetLandscape, const TSubclassOf<ULandscapeEditLayerBase>& InEditLayerClass)
 {
 	// Insertion logic is left to the user through modal drag + drop dialog : 
 	int32 ExistingLayerIndex = InTargetLandscape->GetLayerIndex(InEditLayerName);
 	if (ExistingLayerIndex == INDEX_NONE)
 	{
-		InTargetLandscape->CreateLayer(InEditLayerName);
-		TSharedPtr<SLandscapeLayerListDialog> Dialog = SNew(SLandscapeLayerListDialog, InTargetLandscape->LandscapeLayers);
+		InTargetLandscape->CreateLayer(InEditLayerName, InEditLayerClass);
+		TSharedPtr<SLandscapeLayerListDialog> Dialog = SNew(SLandscapeLayerListDialog, InTargetLandscape);
 		Dialog->ShowModal();
 		ExistingLayerIndex = Dialog->GetInsertedLayerIndex();
 	}
diff --git a/Engine/Source/Editor/LandscapeEditor/Private/SLandscapeLayerListDialog.cpp b/Engine/Source/Editor/LandscapeEditor/Private/SLandscapeLayerListDialog.cpp
--- a/Engine/Source/Editor/LandscapeEditor/Private/SLandscapeLayerListDialog.cpp
+++ b/Engine/Source/Editor/LandscapeEditor/Private/SLandscapeLayerListDialog.cpp
@@ -14,9 +14,9 @@
 
 struct FWidgetLayerListItem
 {
-	FWidgetLayerListItem(const FLandscapeLayer* InLayer, const TFunction<void(void)>& InOnLayerListUpdated, TArray<TSharedPtr<FWidgetLayerListItem>>* InWidgetLayerList)
-	: LayerName(InLayer->Name)
-	, LayerGuid(InLayer->Guid)
+	FWidgetLayerListItem(const FLandscapeLayer& InLayer, const TFunction<void(void)>& InOnLayerListUpdated, TArray<TSharedPtr<FWidgetLayerListItem>>* InWidgetLayerList)
+	: LayerName(InLayer.Name)
+	, LayerGuid(InLayer.Guid)
 	, OnLayerListUpdated(InOnLayerListUpdated)
 	, WidgetLayerList(InWidgetLayerList)
 	{}
@@ -174,17 +174,21 @@ private:
 
 /* SLandscapeLayerListDialog implementation */
 
-void SLandscapeLayerListDialog::Construct(const FArguments& InArgs, TArray<FLandscapeLayer>& InLayers)
+void SLandscapeLayerListDialog::Construct(const FArguments& InArgs, const TWeakObjectPtr<ALandscape>& InLandscape)
 {
-	LayerList = &InLayers;
+	Landscape = InLandscape;
+	check(Landscape.IsValid());
 
-	for(int LayerIndex = 0; LayerIndex < LayerList->Num(); ++LayerIndex)
+	int32 LayerIndex = 0;
+	TArrayView<const FLandscapeLayer> LandscapeLayers = Landscape->GetLayers();
+	for (const FLandscapeLayer& Layer : LandscapeLayers)
 	{
-		WidgetLayerList.Add(MakeShared<FWidgetLayerListItem>(&(*LayerList)[LayerList->Num() - 1 - LayerIndex], [this](){ OnLayerListUpdated(); }, &WidgetLayerList));
+		WidgetLayerList.Add(MakeShared<FWidgetLayerListItem>(LandscapeLayers[LandscapeLayers.Num() - 1 - LayerIndex], [this](){ OnLayerListUpdated(); }, &WidgetLayerList));
+		++LayerIndex;
 	}
 
 	WidgetLayerList[0]->bAllowedToDrag = true;
-	InsertedLayerIndex = InLayers.Num() - 1;
+	InsertedLayerIndex = LandscapeLayers.Num() - 1;
 	
 	// Construct list view
 	SAssignNew(LayerListView, SWidgetLayerListView)
@@ -253,6 +257,8 @@ TSharedRef<ITableRow> SLandscapeLayerListDialog::OnGenerateRow(TSharedPtr<FWidge
 
 void SLandscapeLayerListDialog::OnAccept()
 {
+	check(Landscape.IsValid());
+	TArrayView<const FLandscapeLayer> LandscapeLayers = Landscape->GetLayers();
 	// Find the draggable layer in WidgetLayerList
 	for (int WidgetLayerListIndex = 0; WidgetLayerListIndex < WidgetLayerList.Num(); ++WidgetLayerListIndex)
 	{
@@ -261,18 +267,15 @@ void SLandscapeLayerListDialog::OnAccept()
 		if (WidgetLayer->bAllowedToDrag)
 		{
 			// Find actual edit layer corresponding to draggable layer
-			for (int LayerIndex = 0; LayerIndex < LayerList->Num(); ++LayerIndex)
+			for (int32 LayerIndex = 0; LayerIndex < LandscapeLayers.Num(); ++LayerIndex)
 			{
-				const FLandscapeLayer& Layer = (*LayerList)[LayerIndex];
+				const FLandscapeLayer& Layer = LandscapeLayers[LayerIndex];
 
 				// Once found, remove and copy edit layer into correct spot in LayerList
 				if (WidgetLayer->LayerGuid == Layer.Guid)
 				{
-					const FLandscapeLayer LayerCopy = Layer;
-					
-					LayerList->RemoveAt(LayerIndex);
 					InsertedLayerIndex = WidgetLayerListIndex;
-					LayerList->Insert(LayerCopy, InsertedLayerIndex);
+					Landscape->ReorderLayer(LayerIndex, InsertedLayerIndex);
 					
 					return;
 				}
diff --git a/Engine/Source/Editor/LandscapeEditor/Public/SLandscapeLayerListDialog.h b/Engine/Source/Editor/LandscapeEditor/Public/SLandscapeLayerListDialog.h
--- a/Engine/Source/Editor/LandscapeEditor/Public/SLandscapeLayerListDialog.h
+++ b/Engine/Source/Editor/LandscapeEditor/Public/SLandscapeLayerListDialog.h
@@ -7,7 +7,7 @@
 #include "Dialog/SCustomDialog.h"		// SLandscapeLayerListDialog
 
 class FWidgetLayerListDragDropOp;
-struct FLandscapeLayer;
+class ALandscape;
 struct FWidgetLayerListItem;
 typedef SListView<TSharedPtr<FWidgetLayerListItem> > SWidgetLayerListView;
 
@@ -28,7 +28,7 @@ public:
 	 * @param InArgs Declaration used by the SNew() macro to construct this widget
 	 * @param InLayers The array of landscape layers which should be displayed
 	 */
-	void Construct(const FArguments& InArgs, TArray<FLandscapeLayer>& InLayers);
+	void Construct(const FArguments& InArgs, const TWeakObjectPtr<ALandscape>& InLandscape);
 
 	int32 GetInsertedLayerIndex() const { return InsertedLayerIndex; }
 
@@ -49,8 +49,7 @@ private:
 	 * Note: This array is a UI reflection of LayerList but the order is the reverse order of LayerList */
 	TArray<TSharedPtr<FWidgetLayerListItem>> WidgetLayerList;
 
-	/* The original list of layers passed in Construct. */
-	TArray<FLandscapeLayer>* LayerList = nullptr;
+	TWeakObjectPtr<ALandscape> Landscape;
 
 	/* Index of the layer which is being inserted into the layer stack */
 	int InsertedLayerIndex = -1;
diff --git a/Engine/Source/Editor/WorldBrowser/Private/Tiles/WorldTileModel.cpp b/Engine/Source/Editor/WorldBrowser/Private/Tiles/WorldTileModel.cpp
--- a/Engine/Source/Editor/WorldBrowser/Private/Tiles/WorldTileModel.cpp
+++ b/Engine/Source/Editor/WorldBrowser/Private/Tiles/WorldTileModel.cpp
@@ -1049,7 +1049,7 @@ ALandscapeProxy* FWorldTileModel::ImportLandscapeTile(const FLandscapeImportSett
 
 	// Create landscape components
 	LandscapeProxy->Import(	Settings.LandscapeGuid, 0, 0, Settings.SizeX - 1, Settings.SizeY - 1, Settings.SectionsPerComponent, Settings.QuadsPerSection, HeightmapDataPerLayers, *Settings.HeightmapFilename,	
-							MaterialLayerDataPerLayer,	Settings.ImportLayerType);
+							MaterialLayerDataPerLayer,	Settings.ImportLayerType, TArrayView<const FLandscapeLayer>());
 
 	for (const FLandscapeImportLayerInfo& ImportLayerInfo : Settings.ImportLayers)
 	{
diff --git a/Engine/Source/Runtime/Core/Public/UObject/FortniteMainBranchObjectVersion.h b/Engine/Source/Runtime/Core/Public/UObject/FortniteMainBranchObjectVersion.h
--- a/Engine/Source/Runtime/Core/Public/UObject/FortniteMainBranchObjectVersion.h
+++ b/Engine/Source/Runtime/Core/Public/UObject/FortniteMainBranchObjectVersion.h
@@ -432,6 +432,9 @@ struct FFortniteMainBranchObjectVersion
 		// Added FVolumetricLightMapGridDesc in MapBuildData
 		VolumetricLightMapGridDescSupport,
 
+		// Introduce new structure for customizing the landscape edit layer behavior
+		IntroduceLandscapeEditLayerClass,
+
 		// -----<new versions can be added above this line>-------------------------------------------------
 		VersionPlusOne,
 		LatestVersion = VersionPlusOne - 1
diff --git a/Engine/Source/Runtime/Landscape/Classes/Landscape.h b/Engine/Source/Runtime/Landscape/Classes/Landscape.h
--- a/Engine/Source/Runtime/Landscape/Classes/Landscape.h
+++ b/Engine/Source/Runtime/Landscape/Classes/Landscape.h
@@ -9,6 +9,8 @@
 #include "LandscapeBlueprintBrushBase.h"
 #include "LandscapeEditTypes.h"
 #include "Delegates/DelegateCombinations.h"
+#include "Templates/SubclassOf.h"
+#include "Algo/Transform.h"
 
 #include "Landscape.generated.h"
 
@@ -18,6 +20,7 @@ class ILandscapeEdModeInterface;
 class SNotificationItem;
 class UStreamableRenderAsset;
 class UTextureRenderTarget;
+class ULandscapeEditLayerBase;
 class FMaterialResource;
 struct FLandscapeEditLayerComponentReadbackResult;
 struct FLandscapeNotification;
@@ -191,6 +194,9 @@ struct FLandscapeLayer
 
 	UPROPERTY()
 	TMap<TObjectPtr<ULandscapeLayerInfoObject>, bool> WeightmapLayerAllocationBlend; // True -> Substractive, False -> Additive
+
+	UPROPERTY(Instanced)
+	TObjectPtr<ULandscapeEditLayerBase> EditLayer;
 };
 
 UCLASS(MinimalAPI, showcategories=(Display, Movement, Collision, Lighting, LOD, Input), hidecategories=(Mobility))
@@ -198,6 +204,8 @@ class ALandscape : public ALandscapeProxy
 {
 	GENERATED_BODY()
 
+	friend class FLandscapeConfigHelper; // For copying/manipulating internal data
+
 public:
 	ALandscape(const FObjectInitializer& ObjectInitializer);
 
@@ -315,8 +323,23 @@ public:
 	LANDSCAPE_API void RequestLayersContentUpdate(ELandscapeLayerUpdateMode InModeMask);
 	LANDSCAPE_API bool ReorderLayer(int32 InStartingLayerIndex, int32 InDestinationLayerIndex);
 	LANDSCAPE_API FLandscapeLayer* DuplicateLayerAndMoveBrushes(const FLandscapeLayer& InOtherLayer);
-	LANDSCAPE_API int32 CreateLayer(FName InName = NAME_None);
+
+	/** 
+	* Creates a new edit layer
+	* @Param InName is the name of the new edit layer
+	* @Param InEditLayerClass is the class of the edit layer to create. Passing null will create a standard layer of type ULandscapeEditLayer
+	* @return the index of the newly-created layer
+	*/
+	LANDSCAPE_API int32 CreateLayer(FName InName = NAME_None, const TSubclassOf<ULandscapeEditLayerBase>& InEditLayerClass = TSubclassOf<ULandscapeEditLayerBase>());
+
+	/** 
+	* Creates a new edit layer by copying all the settings from the layer passed in parameter
+	* @Param InLayer edit layer to copy from
+	* @return the index of the newly-created layer
+	*/
+	LANDSCAPE_API int32 CreateLayerFrom(const FLandscapeLayer& InLayer);
 	LANDSCAPE_API void CreateDefaultLayer();
+
 	LANDSCAPE_API void CopyOldDataToDefaultLayer();
 	LANDSCAPE_API void CopyOldDataToDefaultLayer(ALandscapeProxy* Proxy);
 	LANDSCAPE_API void AddLayersToProxy(ALandscapeProxy* InProxy);
@@ -326,15 +349,41 @@ public:
 	LANDSCAPE_API void SetLayerAlpha(int32 InLayerIndex, const float InAlpha, bool bInHeightmap);
 	LANDSCAPE_API float GetLayerAlpha(int32 InLayerIndex, bool bInHeightmap) const;
 	LANDSCAPE_API float GetClampedLayerAlpha(float InAlpha, bool bInHeightmap) const;
-	LANDSCAPE_API void SetLayerVisibility(int32 InLayerIndex, bool bInVisible);
+	LANDSCAPE_API void SetLayerVisibility(int32 InLayerIndex, bool bInVisible, bool bInForIntermediateRender = false);
 	LANDSCAPE_API void SetLayerLocked(int32 InLayerIndex, bool bLocked);
+	LANDSCAPE_API void SetLayerBlendMode(int32 InLayerIndex, ELandscapeBlendMode InBlendMode);
 	LANDSCAPE_API uint8 GetLayerCount() const;
-	LANDSCAPE_API struct FLandscapeLayer* GetLayer(int32 InLayerIndex);
-	LANDSCAPE_API const struct FLandscapeLayer* GetLayer(int32 InLayerIndex) const;
-	LANDSCAPE_API const struct FLandscapeLayer* GetLayer(const FGuid& InLayerGuid) const;
-	LANDSCAPE_API const struct FLandscapeLayer* GetLayer(const FName& InLayerName) const;
+	LANDSCAPE_API TArrayView<const FLandscapeLayer> GetLayers() const { return MakeArrayView(LandscapeEditLayers); }
+	LANDSCAPE_API const FLandscapeLayer* GetLayerConst(int32 InLayerIndex) const;
+	LANDSCAPE_API const FLandscapeLayer* GetLayerConst(const FGuid& InLayerGuid) const;
+	LANDSCAPE_API const FLandscapeLayer* GetLayerConst(const FName& InLayerName) const;
 	LANDSCAPE_API int32 GetLayerIndex(FName InLayerName) const;
-	LANDSCAPE_API void ForEachLayer(TFunctionRef<void(struct FLandscapeLayer&)> Fn);
+
+	UE_DEPRECATED(5.5, "This has moved to private (GetLayerInternal), use the appropriate setters to mutate the edit layer")
+	LANDSCAPE_API FLandscapeLayer* GetLayer(int32 InLayerIndex) { return nullptr; }
+	UE_DEPRECATED(5.5, "Use GetLayerConst")
+	LANDSCAPE_API const FLandscapeLayer* GetLayer(int32 InLayerIndex) const;
+	UE_DEPRECATED(5.5, "Use GetLayerConst")
+	LANDSCAPE_API const FLandscapeLayer* GetLayer(const FGuid& InLayerGuid) const;
+	UE_DEPRECATED(5.5, "Use GetLayerConst")
+	LANDSCAPE_API const FLandscapeLayer* GetLayer(const FName& InLayerName) const;
+
+	UE_DEPRECATED(5.5, "Use ForEachLayerConst")
+	LANDSCAPE_API void ForEachLayer(TFunctionRef<void(FLandscapeLayer&)> Fn);
+	
+	/**
+	 * Runs the given function on each edit layer, with the possibility of early exit
+	 * Most easily used with a lambda as follows:
+	 * ForEachLayerConst([](const FLandscapeLayer& InLayer) -> bool
+	 * {
+	 *     return continueLoop ? true : false;
+	 * });
+	 */
+	LANDSCAPE_API void ForEachLayerConst(TFunctionRef<bool(const FLandscapeLayer&)> Fn);
+
+	LANDSCAPE_API const FLandscapeLayer* FindLayerOfType(const TSubclassOf<ULandscapeEditLayerBase>& InLayerClass) const;
+	LANDSCAPE_API TArray<const FLandscapeLayer*> GetLayersOfType(const TSubclassOf<ULandscapeEditLayerBase>& InLayerClass) const;
+
 	LANDSCAPE_API void GetUsedPaintLayers(int32 InLayerIndex, TArray<ULandscapeLayerInfoObject*>& OutUsedLayerInfos) const;
 	LANDSCAPE_API void GetUsedPaintLayers(const FGuid& InLayerGuid, TArray<ULandscapeLayerInfoObject*>& OutUsedLayerInfos) const;
 	LANDSCAPE_API void ClearPaintLayer(int32 InLayerIndex, ULandscapeLayerInfoObject* InLayerInfo);
@@ -351,13 +400,19 @@ public:
 	LANDSCAPE_API void ShowOnlySelectedLayer(int32 InLayerIndex);
 	LANDSCAPE_API void ShowAllLayers();
 	LANDSCAPE_API void UpdateLandscapeSplines(const FGuid& InLayerGuid = FGuid(), bool bInUpdateOnlySelected = false, bool bInForceUpdateAllCompoments = false);
+
+	UE_DEPRECATED(5.5, "Use CreateLayer with a ULandscapeEditLayerSplines instead")
 	LANDSCAPE_API void SetLandscapeSplinesReservedLayer(int32 InLayerIndex);
+	UE_DEPRECATED(5.5, "Use FindLayerOfType(ULandscapeEditLayerSplines::StaticClass()) instead")
 	LANDSCAPE_API struct FLandscapeLayer* GetLandscapeSplinesReservedLayer();
+	UE_DEPRECATED(5.5, "Use FindLayerOfType(ULandscapeEditLayerSplines::StaticClass()) instead")
 	LANDSCAPE_API const struct FLandscapeLayer* GetLandscapeSplinesReservedLayer() const;
+	UE_DEPRECATED(5.5, "Use FindLayerOfType(ULandscapeEditLayerSplines::StaticClass()).Guid == GetEditingLayer() instead")
 	LANDSCAPE_API bool IsEditingLayerReservedForSplines() const;
 
 	LANDSCAPE_API bool IsLayerBlendSubstractive(int32 InLayerIndex, const TWeakObjectPtr<ULandscapeLayerInfoObject>& InLayerInfoObj) const;
 	LANDSCAPE_API void SetLayerSubstractiveBlendStatus(int32 InLayerIndex, bool InStatus, const TWeakObjectPtr<ULandscapeLayerInfoObject>& InLayerInfoObj);
+	LANDSCAPE_API void ReplaceLayerSubstractiveBlendStatus(ULandscapeLayerInfoObject* InFromLayerInfo, ULandscapeLayerInfoObject* InToLayerInfo, bool bInShouldDirtyPackage);
 
 	LANDSCAPE_API int32 GetBrushLayer(class ALandscapeBlueprintBrushBase* InBrush) const;
 	LANDSCAPE_API void AddBrushToLayer(int32 InLayerIndex, class ALandscapeBlueprintBrushBase* InBrush);
@@ -377,8 +432,10 @@ public:
 	
 	LANDSCAPE_API void ToggleCanHaveLayersContent();
 	LANDSCAPE_API void ForceUpdateLayersContent(bool bIntermediateRender = false);
+
 	UFUNCTION(BlueprintCallable, Category = "Landscape")
 	LANDSCAPE_API void ForceLayersFullUpdate();
+
 	LANDSCAPE_API void InitializeLandscapeLayersWeightmapUsage();
 
 	LANDSCAPE_API bool ComputeLandscapeLayerBrushInfo(FTransform& OutLandscapeTransform, FIntPoint& OutLandscapeSize, FIntPoint& OutLandscapeRenderTargetSize);
@@ -404,6 +461,9 @@ protected:
 	FName GenerateUniqueLayerName(FName InName = NAME_None) const;
 
 private:
+	FLandscapeLayer* GetLayerInternal(int32 InLayerIndex);
+	void OnLayerCreatedInternal(FLandscapeLayer& Layer);
+
 	bool SupportsEditLayersLocalMerge();
 	bool HasNormalCaptureBPBrushLayer();
 
@@ -412,7 +472,7 @@ private:
 	void UpdateLayersContent(bool bInWaitForStreaming = false, bool bInSkipMonitorLandscapeEdModeChanges = false, bool bIntermediateRender = false, bool bFlushRender = false);
 	void MonitorShaderCompilation();
 	void MonitorLandscapeEdModeChanges();
-
+	
 	int32 RegenerateLayersHeightmaps(const FUpdateLayersContentContext& InUpdateLayersContentContext);
 	int32 PerformLayersHeightmapsLocalMerge(const FUpdateLayersContentContext& InUpdateLayersContentContext, const FEditLayersHeightmapMergeParams& InMergeParams);
 	int32 PerformLayersHeightmapsGlobalMerge(const FUpdateLayersContentContext& InUpdateLayersContentContext, const FEditLayersHeightmapMergeParams& InMergeParams);
@@ -519,8 +579,9 @@ public:
 	FLandscapeFullHeightmapRenderDoneDelegate& OnFullHeightmapRenderDoneDelegate() { return LandscapeFullHeightmapRenderDoneDelegate; }
 
 	/** Target Landscape Layer for Landscape Splines */
+	UE_DEPRECATED(5.5, "This has been refactored into the generic ULandscapeEditLayerBase system. Please check for the presence of a ULandscapeEditLayerSplines layer instead")
 	UPROPERTY()
-	FGuid LandscapeSplinesTargetLayerGuid;
+	FGuid LandscapeSplinesTargetLayerGuid_DEPRECATED;
 	
 	/** Current Editing Landscape Layer*/
 	FGuid EditingLayer;
@@ -534,8 +595,9 @@ public:
 	UPROPERTY(Transient, DuplicateTransient, TextExportTransient, NonPIEDuplicateTransient)
 	bool bWarnedGlobalMergeDimensionsExceeded = false;
 
+	UE_DEPRECATED(5.5, "This property has moved to private. Use the public accessors instead")
 	UPROPERTY()
-	TArray<FLandscapeLayer> LandscapeLayers;
+	TArray<FLandscapeLayer> LandscapeLayers_DEPRECATED;
 
 	UPROPERTY(Transient)
 	TArray<TObjectPtr<UTextureRenderTarget2D>> HeightmapRTList;
@@ -548,6 +610,9 @@ public:
 	TArray<TWeakObjectPtr<UTexture2D>> TrackedStreamingInTextures;
 
 private:
+	UPROPERTY()
+	TArray<FLandscapeLayer> LandscapeEditLayers;
+
 	FLandscapeBlueprintBrushChangedDelegate LandscapeBlueprintBrushChangedDelegate;
 	FLandscapeFullHeightmapRenderDoneDelegate LandscapeFullHeightmapRenderDoneDelegate;
 
@@ -573,7 +638,7 @@ private:
 
 	UPROPERTY(Transient)
 	bool bLandscapeLayersAreInitialized;
-
+	
 	UPROPERTY(Transient)
 	bool bLandscapeLayersAreInitializedForNormalCapture;
 
diff --git a/Engine/Source/Runtime/Landscape/Classes/LandscapeComponent.h b/Engine/Source/Runtime/Landscape/Classes/LandscapeComponent.h
--- a/Engine/Source/Runtime/Landscape/Classes/LandscapeComponent.h
+++ b/Engine/Source/Runtime/Landscape/Classes/LandscapeComponent.h
@@ -36,7 +36,6 @@ class UMaterialInterface;
 class UTexture2D;
 struct FConvexVolume;
 struct FEngineShowFlags;
-struct FLandscapeEditDataInterface;
 struct FLandscapeTextureDataInfo;
 struct FStaticLightingPrimitiveInfo;
 struct FLandscapeEditDataInterface;
@@ -389,6 +388,7 @@ static const uint32 DefaultSplineHash = 0xFFFFFFFF;
 
 #endif
 
+// TODO [jonathan.bard] : Deprecate : this is now kinda redundant with ELandscapeToolTargetType
 UENUM()
 enum ELandscapeClearMode : int
 {
diff --git a/Engine/Source/Runtime/Landscape/Classes/LandscapeEditLayer.h b/Engine/Source/Runtime/Landscape/Classes/LandscapeEditLayer.h
new file mode 100644
--- /dev/null
+++ b/Engine/Source/Runtime/Landscape/Classes/LandscapeEditLayer.h
@@ -0,0 +1,189 @@
+// Copyright Epic Games, Inc. All Rights Reserved.
+
+#pragma once
+
+#include "LandscapeEditLayer.generated.h"
+
+enum class ELandscapeToolTargetType : uint8;
+struct FLandscapeLayer;
+class ALandscape;
+
+UCLASS(MinimalAPI, Abstract)
+class ULandscapeEditLayerBase : public UObject
+{
+	GENERATED_BODY()
+
+public:
+	// Little wrapper class to let edit layers expose some actions generically : 
+	class FEditLayerAction
+	{
+		FEditLayerAction() = delete;
+
+	public:
+		class FExecuteParams
+		{
+			FExecuteParams() = delete;
+
+		public:
+			LANDSCAPE_API FExecuteParams(const FLandscapeLayer* InLayer, ALandscape* InLandscape)
+				: Layer(InLayer)
+				, Landscape(InLandscape)
+			{
+				check((InLayer != nullptr) && (InLandscape != nullptr));
+			}
+
+			inline const FLandscapeLayer* GetLayer() const { return Layer; }
+			inline ALandscape* GetLandscape() const { return Landscape; }
+
+		private:
+			const FLandscapeLayer* Layer = nullptr;
+			ALandscape* Landscape = nullptr;
+		};
+
+		struct FExecuteResult
+		{
+			LANDSCAPE_API FExecuteResult() = default;
+			LANDSCAPE_API FExecuteResult(bool bInSuccess, const FText& InReason = FText())
+				: bSuccess(bInSuccess)
+				, Reason(InReason)
+			{}
+
+			bool bSuccess = true;
+			FText Reason;
+		};
+
+		DECLARE_DELEGATE_RetVal_OneParam(FExecuteResult, FExecuteDelegate, const FExecuteParams& /*InParams*/);
+		DECLARE_DELEGATE_RetVal_TwoParams(bool, FCanExecuteDelegate, const FExecuteParams& /*InParams*/, FText& /*OutReason*/);
+
+		LANDSCAPE_API FEditLayerAction(const FText& InLabel, const FExecuteDelegate& InExecuteDelegate, const FCanExecuteDelegate& InCanExecuteDelegate)
+			: Label(InLabel)
+			, ExecuteDelegate(InExecuteDelegate)
+			, CanExecuteDelegate(InCanExecuteDelegate)
+		{}
+
+		inline const FText& GetLabel() const { return Label; }
+		inline const FExecuteDelegate& GetExecuteDelegate() const { return ExecuteDelegate; }
+		inline const FCanExecuteDelegate& GetCanExecuteDelegate() const { return CanExecuteDelegate; }
+
+	private:
+		FText Label;
+		FExecuteDelegate ExecuteDelegate;
+		FCanExecuteDelegate CanExecuteDelegate;
+	};
+
+public:
+	/**
+	 * @return true if the this edit layer has support for the target type (heightmap, weightmap, visibility)
+	 */
+	virtual bool SupportsTargetType(ELandscapeToolTargetType InType) const
+		PURE_VIRTUAL(ULandscapeEditLayerBase::SupportsTargetType, return true; );
+
+	/**
+	 * @return true if the edit layer can store heightmaps/weightmaps in the ALandscapeProxy (e.g. should return false for purely procedural layers, to avoid allocating textures)
+	 */
+	virtual bool NeedsPersistentTextures() const
+		PURE_VIRTUAL(ULandscapeEditLayerBase::NeedsPersistentTextures, return false; );
+
+	/**
+	* @return true if the edit layer can be manually edited via the landscape editing tools :
+	*/
+	virtual bool SupportsEditingTools() const
+		PURE_VIRTUAL(ULandscapeEditLayerBase::SupportsEditingTools, return true; );
+
+	/**
+	 * @return true if it's allowed to have more than one edit layer of this type at a time
+	 */
+	virtual bool SupportsMultiple() const
+		PURE_VIRTUAL(ULandscapeEditLayerBase::SupportsMultiple, return true; );
+
+	/**
+	 * @return true if the layer supports a layer above being collapsed onto it
+	 */
+	virtual bool SupportsBeingCollapsedAway() const
+		PURE_VIRTUAL(ULandscapeEditLayerBase::SupportsBeingCollapsedAway, return true; );
+
+	/**
+	 * @return true if the layer supports being collapsed onto a layer underneath
+	 */
+	virtual bool SupportsCollapsingTo() const
+		PURE_VIRTUAL(ULandscapeEditLayerBase::SupportsCollapsingTo, return true; );
+
+	/**
+	* @return the default name to use when creating a new layer of this type
+	*/
+	virtual FString GetDefaultName() const
+		PURE_VIRTUAL(ULandscapeEditLayerBase::GetDefaultName, return FString(); );
+
+	/**
+	* @return a list of actions that can be triggered on this edit layer
+	*/
+	virtual TArray<FEditLayerAction> GetActions() const { return {}; }
+
+	// TODO [jonathan.bard] : Remove : temporary method to give the edit layer an opportunity to change some settings on the old struct (FLandscapeLayer) upon creation. 
+	//  The better way would be to move most of the settings to ULandscapeEditLayer(Base?) and expose a "property changed" event that both the UI and runtime code 
+	//  could listen to in order to generically react to any change on the layer settings, including on the derived class (e.g. if there were some splines layer-specific 
+	//  settings, on property change, the event would be triggered and the landscape layers would be updated as a result), instead of providing ad-hoc functions on 
+	//  ALandscape, like we do currently (e.g. SetLayerAlpha, SetLayerVisibility, SetLayerName, etc.) on both the runtime code (ALandscape) and the UI code (FEdModeLandscape)
+	virtual void OnLayerCreated(FLandscapeLayer& Layer) {}
+};
+
+/** 
+	This is the standard type of edit layer. It can be manually authored (sculpted, painted, etc.) in the landscape editor 
+*/
+UCLASS(MinimalAPI)
+class ULandscapeEditLayer : public ULandscapeEditLayerBase
+{
+	GENERATED_BODY()
+
+public:
+	// Begin ULandscapeEditLayerBase implementation
+	virtual bool SupportsTargetType(ELandscapeToolTargetType InType) const override;
+	virtual bool SupportsEditingTools() const override { return true; }
+	virtual bool NeedsPersistentTextures() const override { return true; };
+	virtual bool SupportsMultiple() const override { return true; }
+	virtual bool SupportsBeingCollapsedAway() const override { return true; }
+	virtual bool SupportsCollapsingTo() const override { return true; }
+	virtual FString GetDefaultName() const { return TEXT("Layer"); }
+	// End ULandscapeEditLayerBase implementation
+
+protected:
+
+};
+
+/** 
+	Base class for procedural layers. Procedural layers cannot be edited through standard editing tools
+*/
+UCLASS(MinimalAPI, Abstract)
+class ULandscapeEditLayerProcedural : public ULandscapeEditLayerBase
+{
+	GENERATED_BODY()
+
+public:
+	// Begin ULandscapeEditLayerBase implementation
+	virtual bool SupportsEditingTools() const override { return false; } // procedural layers cannot be edited through standard editing tools
+	virtual bool SupportsBeingCollapsedAway() const override { return false; } // this is a procedural and therefore cannot be collapsed 
+	virtual bool SupportsCollapsingTo() const override { return NeedsPersistentTextures(); } // if the layer has persistent textures, it supports being collapsed into a layer underneath (this may become unneeded if we make the collapse happen on the GPU)
+	// End ULandscapeEditLayerBase implementation
+};
+
+/** 
+	Procedural edit layer that lets the user manipulate its content using landscape splines (Splines tool in the Manage panel) 
+*/
+UCLASS(MinimalAPI)
+class ULandscapeEditLayerSplines : public ULandscapeEditLayerProcedural
+{
+	GENERATED_BODY()
+
+public:
+	// Begin ULandscapeEditLayerBase implementation
+	virtual bool SupportsTargetType(ELandscapeToolTargetType InType) const override;
+	virtual bool NeedsPersistentTextures() const override { return true; }; // it's a layer computed on the CPU and outputting to persistent textures
+	virtual bool SupportsMultiple() const override { return false; } // only one layer of this type is allowed
+	virtual FString GetDefaultName() const override { return TEXT("Splines"); }
+	virtual void OnLayerCreated(FLandscapeLayer& Layer) override;
+	virtual TArray<FEditLayerAction> GetActions() const override;
+	// End ULandscapeEditLayerBase implementation
+
+protected:
+
+};
diff --git a/Engine/Source/Runtime/Landscape/Classes/LandscapeProxy.h b/Engine/Source/Runtime/Landscape/Classes/LandscapeProxy.h
--- a/Engine/Source/Runtime/Landscape/Classes/LandscapeProxy.h
+++ b/Engine/Source/Runtime/Landscape/Classes/LandscapeProxy.h
@@ -1291,8 +1291,14 @@ public:
 	LANDSCAPE_API static ULandscapeMaterialInstanceConstant* GetLayerThumbnailMIC(UMaterialInterface* LandscapeMaterial, FName LayerName, UTexture2D* ThumbnailWeightmap, UTexture2D* ThumbnailHeightmap, ALandscapeProxy* Proxy);
 
 	/** Import the given Height/Weight data into this landscape */
+	UE_DEPRECATED(5.5, "Use the new method prototype which takes a TArrayView parameter for import layers")
 	LANDSCAPE_API void Import(const FGuid& InGuid, int32 InMinX, int32 InMinY, int32 InMaxX, int32 InMaxY, int32 InNumSubsections, int32 InSubsectionSizeQuads, const TMap<FGuid, TArray<uint16>>& InImportHeightData,
-							  const TCHAR* const InHeightmapFileName, const TMap<FGuid, TArray<FLandscapeImportLayerInfo>>& InImportMaterialLayerInfos, ELandscapeImportAlphamapType InImportMaterialLayerType, const TArray<struct FLandscapeLayer>* InImportLayers = nullptr);
+		const TCHAR* const InHeightmapFileName, const TMap<FGuid, TArray<FLandscapeImportLayerInfo>>& InImportMaterialLayerInfos, ELandscapeImportAlphamapType InImportMaterialLayerType, 
+		const TArray<struct FLandscapeLayer>* InImportLayers = nullptr);
+
+	LANDSCAPE_API void Import(const FGuid& InGuid, int32 InMinX, int32 InMinY, int32 InMaxX, int32 InMaxY, int32 InNumSubsections, int32 InSubsectionSizeQuads, const TMap<FGuid, TArray<uint16>>& InImportHeightData,
+		const TCHAR* const InHeightmapFileName, const TMap<FGuid, TArray<FLandscapeImportLayerInfo>>& InImportMaterialLayerInfos, ELandscapeImportAlphamapType InImportMaterialLayerType,
+		const TArrayView<const struct FLandscapeLayer>& InImportLayers);
 
 	struct FRawMeshExportParams
 	{
diff --git a/Engine/Source/Runtime/Landscape/Private/Landscape.cpp b/Engine/Source/Runtime/Landscape/Private/Landscape.cpp
--- a/Engine/Source/Runtime/Landscape/Private/Landscape.cpp
+++ b/Engine/Source/Runtime/Landscape/Private/Landscape.cpp
@@ -77,6 +77,7 @@ Landscape.cpp: Terrain rendering
 #include "VisualLogger/VisualLogger.h"
 #include "NaniteSceneProxy.h"
 #include "Misc/ArchiveMD5.h"
+#include "LandscapeEditLayer.h"
 #include "LandscapeTextureStorageProvider.h"
 #include "LandscapeVersion.h"
 #include "UObject/FortniteMainBranchObjectVersion.h"
@@ -426,11 +427,11 @@ FGraphEventRef ALandscapeProxy::UpdateNaniteRepresentationAsync(const ITargetPla
 
 		const FGuid ComponentNaniteContentId = GetNaniteComponentContentId();
 		const bool bNaniteContentDirty = ComponentNaniteContentId != NaniteContentId;
-	
-		if(bNaniteContentDirty && IsRunningCookCommandlet())
+
+		if (bNaniteContentDirty && IsRunningCookCommandlet())
 		{
 			UE_LOG(LogLandscape, Display, TEXT("Landscape Nanite out of date. Map requires resaving. Actor: '%s' Package: '%s'"), *GetActorNameOrLabel(), *GetPackage()->GetName());
-		}		
+		}
 
 		TArray<ULandscapeComponent*> StableOrderComponents(LandscapeComponents);
 		ULandscapeSubsystem* Subsystem = GetWorld()->GetSubsystem<ULandscapeSubsystem>();
@@ -456,7 +457,7 @@ FGraphEventRef ALandscapeProxy::UpdateNaniteRepresentationAsync(const ITargetPla
 
 				TArrayView<ULandscapeComponent*> ComponentsToExport(StableOrderComponentsView.GetData(), NumComponents);
 				FGraphEventRef ComponentProcessTask = NaniteComponents[i]->InitializeForLandscapeAsync(this, NaniteContentId, Subsystem->IsMultithreadedNaniteBuildEnabled(), ComponentsToExport, i);
-				SingleProxyDependencies.Add(ComponentProcessTask);	
+				SingleProxyDependencies.Add(ComponentProcessTask);
 			}
 
 			// TODO: Add a flag that only initializes the platform if we called InitializeForLandscape during the PreSave for this or a previous platform
@@ -470,14 +471,14 @@ FGraphEventRef ALandscapeProxy::UpdateNaniteRepresentationAsync(const ITargetPla
 				}
 				WeakComponent->InitializePlatformForLandscape(WeakProxy.Get(), InTargetPlatform);
 				WeakComponent->UpdatedSharedPropertiesFromActor();
-				},
+			},
 				TStatId(),
 				&SingleProxyDependencies,
 				ENamedThreads::GameThread);
 
 			UpdateDependencies.Add(FinalizeEvent);
 		}
-			
+
 		BatchBuildEvent = FFunctionGraphTask::CreateAndDispatchWhenReady([] {}, TStatId(), &UpdateDependencies, ENamedThreads::GameThread);
 
 
@@ -508,7 +509,7 @@ void ALandscapeProxy::UpdateNaniteRepresentation(const ITargetPlatform* InTarget
 		{
 			ENamedThreads::Type CurrentThread = FTaskGraphInterface::Get().GetCurrentThreadIfKnown();
 			FTaskGraphInterface::Get().ProcessThreadUntilIdle(CurrentThread);
-			FAssetCompilingManager::Get().ProcessAsyncTasks();	
+			FAssetCompilingManager::Get().ProcessAsyncTasks();
 		}
 	}
 }
@@ -561,7 +562,7 @@ FGuid ALandscapeProxy::GetNaniteContentId() const
 				return true;
 			}
 			if (!B)
-			{
+		{
 				return false;
 			}
 			// Sort components based on their SectionBase (i.e. 2D index relative to the entire landscape) to ensure stable ID generation
@@ -948,8 +949,8 @@ PRAGMA_ENABLE_DEPRECATION_WARNINGS
 			}
 			else
 			{
-				Ar << GrassData.Get();
-			}
+			Ar << GrassData.Get();
+		}
 		}
 
 		// When loading or saving a component, validate that grass data is valid : 
@@ -1197,7 +1198,7 @@ void ULandscapeComponent::PostLoad()
 		{
 			VertexFactoryDataList.Add(FPSOPrecacheVertexFactoryData(&FLandscapeFixedGridVertexFactory::StaticType));
 		}
-		
+
 		if (Culling::UseCulling(GMaxRHIShaderPlatform))
 		{
 			VertexFactoryDataList.Add(FPSOPrecacheVertexFactoryData(Culling::GetTileVertexFactoryType()));
@@ -1368,7 +1369,7 @@ PRAGMA_ENABLE_DEPRECATION_WARNINGS
 		}
 
 		// Fixup missing/mismatching edit layer names :
-		if (const FLandscapeLayer* EditLayer = GetLandscapeActor() ? GetLandscapeActor()->GetLayer(ItPair.Key) : nullptr)
+		if (const FLandscapeLayer* EditLayer = GetLandscapeActor() ? GetLandscapeActor()->GetLayerConst(ItPair.Key) : nullptr)
 		{
 			if (LayerComponentData.DebugName != EditLayer->Name)
 			{
@@ -1394,7 +1395,7 @@ PRAGMA_ENABLE_DEPRECATION_WARNINGS
 		}
 	}
 
- 	for (UMaterialInstance* MobileCombinationMaterialInstance : MobileCombinationMaterialInstances)
+	for (UMaterialInstance* MobileCombinationMaterialInstance : MobileCombinationMaterialInstances)
 	{
 		while (ReparentObject(MobileCombinationMaterialInstance))
 		{
@@ -1932,7 +1933,7 @@ const FMeshMapBuildData* ULandscapeComponent::GetMeshMapBuildData() const
 		if (OwnerLevel && OwnerLevel->OwningWorld)
 		{
 			UMapBuildDataRegistry* MapBuildData = UMapBuildDataRegistry::Get(this);
-			
+
 			if (MapBuildData)
 			{
 				return MapBuildData->GetMeshBuildData(MapBuildDataId);
@@ -2422,7 +2423,7 @@ void ULandscapeComponent::AddDefaultLayerData(const FGuid& InLayerGuid, const TA
 
 	if (LayerData == nullptr || !LayerData->IsInitialized())
 	{
-		const FLandscapeLayer* EditLayer = GetLandscapeActor() ? GetLandscapeActor()->GetLayer(InLayerGuid) : nullptr;
+		const FLandscapeLayer* EditLayer = GetLandscapeActor() ? GetLandscapeActor()->GetLayerConst(InLayerGuid) : nullptr;
 		FLandscapeLayerComponentData NewData(EditLayer ? EditLayer->Name : FName());
 
 		// Setup Heightmap data
@@ -2836,10 +2837,31 @@ void ALandscape::PostLoad()
 	}
 
 #if WITH_EDITOR
-	for (FLandscapeLayer& Layer : LandscapeLayers)
+	if (GetLinkerCustomVersion(FFortniteMainBranchObjectVersion::GUID) < FFortniteMainBranchObjectVersion::IntroduceLandscapeEditLayerClass)
 	{
+PRAGMA_DISABLE_DEPRECATION_WARNINGS
+		LandscapeEditLayers = LandscapeLayers_DEPRECATED;
+
+		for (FLandscapeLayer& Layer : LandscapeEditLayers)
+		{
+			UClass* EditLayerClass = ULandscapeEditLayer::StaticClass();
+			if (Layer.Guid == LandscapeSplinesTargetLayerGuid_DEPRECATED)
+			{
+				EditLayerClass = ULandscapeEditLayerSplines::StaticClass();
+			}
+			check(Layer.EditLayer == nullptr);
+			Layer.EditLayer = NewObject<ULandscapeEditLayerBase>(this, EditLayerClass, MakeUniqueObjectName(this, EditLayerClass));
+			Layer.EditLayer->OnLayerCreated(Layer);
+		}
+PRAGMA_ENABLE_DEPRECATION_WARNINGS
+	}
+
+	for (FLandscapeLayer& Layer : LandscapeEditLayers)
+	{
+		check(Layer.EditLayer != nullptr);
 		// For now, only Layer reserved for Landscape Spline uses AlphaBlend
-		Layer.BlendMode = (Layer.Guid == LandscapeSplinesTargetLayerGuid) ? LSBM_AlphaBlend : LSBM_AdditiveBlend;
+		Layer.BlendMode = Layer.EditLayer->IsA<ULandscapeEditLayerSplines>() ? LSBM_AlphaBlend : LSBM_AdditiveBlend;
+
 		for (FLandscapeLayerBrush& Brush : Layer.Brushes)
 		{
 			Brush.SetOwner(this);
@@ -2871,7 +2893,7 @@ public:
 		SHADER_PARAMETER(int32, InSourceTextureChannel)
 		SHADER_PARAMETER_RDG_TEXTURE_SRV(Texture2D<float4>, InSourceTexture)
 		RENDER_TARGET_BINDING_SLOTS()
-	END_SHADER_PARAMETER_STRUCT()
+		END_SHADER_PARAMETER_STRUCT()
 
 	class FIsHeightmap : SHADER_PERMUTATION_BOOL("IS_HEIGHTMAP");
 
@@ -2884,7 +2906,7 @@ public:
 		return PermutationVector;
 	}
 
-	
+
 	static bool ShouldCompilePermutation(const FGlobalShaderPermutationParameters& InParameters)
 	{
 		return true;
@@ -2931,7 +2953,7 @@ public:
 		SHADER_PARAMETER_SAMPLER(SamplerState, InMergedTextureSampler)
 		SHADER_PARAMETER(FUintVector2, InRenderAreaSize)
 		RENDER_TARGET_BINDING_SLOTS()
-	END_SHADER_PARAMETER_STRUCT()
+		END_SHADER_PARAMETER_STRUCT()
 
 	class FIsHeightmap : SHADER_PERMUTATION_BOOL("IS_HEIGHTMAP");
 	class FCompressHeight : SHADER_PERMUTATION_BOOL("COMPRESS_HEIGHT");
@@ -3133,8 +3155,8 @@ namespace UE::Landscape::Private::RenderMergedTexture_RenderThread
 
 			// We now need to resample the atlas texture where the render area is : 
 			FLandscapeResampleMergedTexturePS::ResampleMergedTexture(GraphBuilder, ResampleMergedTexturePSParams, InRenderInfo.bIsHeightmap, InRenderInfo.bCompressHeight);
-		}
 	}
+}
 } // namespace UE::Landscape::Private::RenderMergedTexture_RenderThread
 
 bool ALandscape::IsValidRenderTargetFormatHeightmap(EPixelFormat InRenderTargetFormat, bool& bOutCompressHeight)
@@ -3236,7 +3258,7 @@ bool ALandscape::RenderMergedTextureInternal(const FTransform& InRenderAreaWorld
 		}
 
 		if (!IsValidRenderTargetFormatHeightmap(RenderTargetFormat, bCompressHeight))
-		{
+	{
 			UE_LOG(LogLandscape, Warning, TEXT("RenderMergedTexture : invalid render target format for rendering heightmap (%s)"), GetPixelFormatString(RenderTargetFormat));
 			return false;
 		}
@@ -3343,7 +3365,7 @@ bool ALandscape::RenderMergedTextureInternal(const FTransform& InRenderAreaWorld
 					SourceTextureChannel = AllocInfo->WeightmapTextureChannel;
 				}
 			}
-			
+
 			if (SourceTexture != nullptr)
 			{
 				// Get the subregion of the source texture that this component uses (differs due to texture sharing).
@@ -3359,7 +3381,7 @@ bool ALandscape::RenderMergedTextureInternal(const FTransform& InRenderAreaWorld
 						FMath::FloorToInt32(SourceTextureBias.X * SourceTextureResource->GetSizeX()),
 						FMath::FloorToInt32(SourceTextureBias.Y * SourceTextureResource->GetSizeY()));
 				}
-			
+
 				// When mips are partially loaded, we need to take that into consideration when merging the source texture :
 				uint32 MipBias = SourceTexture->GetNumMips() - SourceTexture->GetNumResidentMips();
 
@@ -3369,7 +3391,7 @@ bool ALandscape::RenderMergedTextureInternal(const FTransform& InRenderAreaWorld
 				SourceTextureOffset.Y >>= MipBias;
 				ComponentSize >>= MipBias;
 
-				// Effective area of the texture affecting this component (because of texture sharing):
+				// Effective area of the texture affecting this component (because of texture sharing) :
 				FIntRect SourceTextureSubregion(SourceTextureOffset, SourceTextureOffset + ComponentSize);
 				MergeTextureRenderInfo.ComponentTexturesToRender.Add(ComponentKey, FTexture2DResourceSubregion(SourceTextureResource->GetTexture2DResource(), SourceTextureSubregion, SourceTextureChannel));
 
@@ -3392,7 +3414,7 @@ bool ALandscape::RenderMergedTextureInternal(const FTransform& InRenderAreaWorld
 		FVector MergedTextureScale = (FVector(RenderTargetComponentIndicesBoundingRect.Max - RenderTargetComponentIndicesBoundingRect.Min) * static_cast<double>(ComponentSizeQuads) + 1)
 			* LandscapeTransform.GetScale3D();
 		MergedTextureScale.Z = 1.0f;
-		FVector MergedTextureUVOrigin = LandscapeTransform.TransformPosition(FVector(RenderTargetComponentIndicesBoundingRect.Min) * (double)ComponentSizeQuads - FVector(0.5,0.5,0));
+		FVector MergedTextureUVOrigin = LandscapeTransform.TransformPosition(FVector(RenderTargetComponentIndicesBoundingRect.Min) * (double)ComponentSizeQuads - FVector(0.5, 0.5, 0));
 		FTransform MergedTextureUVToWorld(LandscapeTransform.GetRotation(), MergedTextureUVOrigin, MergedTextureScale);
 
 		MergeTextureRenderInfo.OutputUVToMergedTextureUV = OutputUVToWorld.ToMatrixWithScale() * MergedTextureUVToWorld.ToInverseMatrixWithScale();
@@ -3637,7 +3659,7 @@ void ALandscapeProxy::PreSave(FObjectPreSaveContext ObjectSaveContext)
 	{
 		LandscapeInfo->UpdateNanite(ObjectSaveContext.GetTargetPlatform());
 	}
-	
+
 	if (ALandscape* Landscape = GetLandscapeActor())
 	{
 		for (ULandscapeComponent* LandscapeComponent : LandscapeComponents)
@@ -3647,7 +3669,7 @@ void ALandscapeProxy::PreSave(FObjectPreSaveContext ObjectSaveContext)
 			// Make sure edit layer debug names are synchronized upon save :
 			LandscapeComponent->ForEachLayer([&](const FGuid& LayerGuid, FLandscapeLayerComponentData& LayerData)
 			{
-				if (const FLandscapeLayer* EditLayer = Landscape->GetLayer(LayerGuid))
+				if (const FLandscapeLayer* EditLayer = Landscape->GetLayerConst(LayerGuid))
 				{
 					LayerData.DebugName = EditLayer->Name;
 				}
@@ -4043,9 +4065,9 @@ bool ULandscapeInfo::UpdateLayerInfoMap(ALandscapeProxy* Proxy /*= nullptr*/, bo
 /* if the outer world is instanced, we need to change our landscape guid(in a deterministic way)
  * this avoids guid collisions when you instance a world (and its landscapes) multiple times,
  * while maintaining the same GUID between landscape proxy objects within an instance
- */ 
- void ChangeLandscapeGuidIfObjectIsInstanced(FGuid& InOutGuid, UObject* InObject)
- {
+ */
+void ChangeLandscapeGuidIfObjectIsInstanced(FGuid& InOutGuid, UObject* InObject)
+{
 	// we shouldn't be dealing with any instanced landscapes in these cases, early out
 	if (InObject->IsTemplate() || IsRunningCookCommandlet())
 	{
@@ -4073,7 +4095,7 @@ bool ULandscapeInfo::UpdateLayerInfoMap(ALandscapeProxy* Proxy /*= nullptr*/, bo
 		InOutGuid = Ar.GetGuidFromHash();
 	}
 }
- 
+
 void ALandscapeProxy::PostLoadFixupLandscapeGuidsIfInstanced()
 {
 	// record the original value before modification
@@ -4150,47 +4172,47 @@ void ALandscapeProxy::PostLoad()
 	}
 
 #if WITH_EDITOR
-PRAGMA_DISABLE_DEPRECATION_WARNINGS
-	if (!LandscapeMaterialsOverride_DEPRECATED.IsEmpty())
-	{
-		PerLODOverrideMaterials.Reserve(LandscapeMaterialsOverride_DEPRECATED.Num());
-		for (const FLandscapeProxyMaterialOverride& LocalMaterialOverride : LandscapeMaterialsOverride_DEPRECATED)
+	PRAGMA_DISABLE_DEPRECATION_WARNINGS
+		if (!LandscapeMaterialsOverride_DEPRECATED.IsEmpty())
 		{
-			PerLODOverrideMaterials.Add({ LocalMaterialOverride.LODIndex.Default, LocalMaterialOverride.Material });
+			PerLODOverrideMaterials.Reserve(LandscapeMaterialsOverride_DEPRECATED.Num());
+			for (const FLandscapeProxyMaterialOverride& LocalMaterialOverride : LandscapeMaterialsOverride_DEPRECATED)
+			{
+				PerLODOverrideMaterials.Add({ LocalMaterialOverride.LODIndex.Default, LocalMaterialOverride.Material });
+			}
+			LandscapeMaterialsOverride_DEPRECATED.Reset();
 		}
-		LandscapeMaterialsOverride_DEPRECATED.Reset();
-	}
-PRAGMA_ENABLE_DEPRECATION_WARNINGS
+	PRAGMA_ENABLE_DEPRECATION_WARNINGS
 
-	if (GIsEditor)
-	{
-		// We may not have run PostLoad on LandscapeComponents yet
-		for (TObjectPtr<ULandscapeComponent>& LandscapeComponent : LandscapeComponents)
+		if (GIsEditor)
 		{
-			if (LandscapeComponent)
+			// We may not have run PostLoad on LandscapeComponents yet
+			for (TObjectPtr<ULandscapeComponent>& LandscapeComponent : LandscapeComponents)
 			{
-				LandscapeComponent->ConditionalPostLoad();
+				if (LandscapeComponent)
+				{
+					LandscapeComponent->ConditionalPostLoad();
+				}
 			}
-		}
-		
-		// We may not have run PostLoad on CollisionComponent yet
-		for (TObjectPtr<ULandscapeHeightfieldCollisionComponent>& CollisionComponent : CollisionComponents)
-		{
-			if (CollisionComponent)
+
+			// We may not have run PostLoad on CollisionComponent yet
+			for (TObjectPtr<ULandscapeHeightfieldCollisionComponent>& CollisionComponent : CollisionComponents)
 			{
-				CollisionComponent->ConditionalPostLoad();
+				if (CollisionComponent)
+				{
+					CollisionComponent->ConditionalPostLoad();
+				}
 			}
-		}
 
-		if ((GetLinker() && (GetLinker()->UEVer() < VER_UE4_LANDSCAPE_COMPONENT_LAZY_REFERENCES)) ||
-			LandscapeComponents.Num() != CollisionComponents.Num() ||
-			LandscapeComponents.ContainsByPredicate([](ULandscapeComponent* Comp) { return ((Comp != nullptr) && (Comp->GetCollisionComponent() == nullptr)); }) ||
-			CollisionComponents.ContainsByPredicate([](ULandscapeHeightfieldCollisionComponent* Comp) { return ((Comp != nullptr) && (Comp->GetRenderComponent() == nullptr)); }))
-		{
-			// Need to clean up invalid collision and render components
-			RecreateCollisionComponents();
+			if ((GetLinker() && (GetLinker()->UEVer() < VER_UE4_LANDSCAPE_COMPONENT_LAZY_REFERENCES)) ||
+				LandscapeComponents.Num() != CollisionComponents.Num() ||
+				LandscapeComponents.ContainsByPredicate([](ULandscapeComponent* Comp) { return ((Comp != nullptr) && (Comp->GetCollisionComponent() == nullptr)); }) ||
+				CollisionComponents.ContainsByPredicate([](ULandscapeHeightfieldCollisionComponent* Comp) { return ((Comp != nullptr) && (Comp->GetRenderComponent() == nullptr)); }))
+			{
+				// Need to clean up invalid collision and render components
+				RecreateCollisionComponents();
+			}
 		}
-	}
 
 	EditorLayerSettings.RemoveAll([](const FLandscapeEditorLayerSettings& Settings) { return Settings.LayerInfoObj == nullptr; });
 
@@ -4205,7 +4227,7 @@ PRAGMA_DISABLE_DEPRECATION_WARNINGS
 	}
 PRAGMA_ENABLE_DEPRECATION_WARNINGS
 
-	bool bFixedUpInvalidMaterialInstances = false;
+		bool bFixedUpInvalidMaterialInstances = false;
 
 	UMaterial* BaseLandscapeMaterial = GetLandscapeMaterial()->GetMaterial();
 
@@ -4268,17 +4290,17 @@ PRAGMA_ENABLE_DEPRECATION_WARNINGS
 	if (!IsNaniteMeshUpToDate() && !IsRunningCookCommandlet())
 	{
 		FFormatNamedArguments Arguments;
-		Arguments.Add(TEXT("LandscapeProxyName"), FText::FromString(GetActorNameOrLabel())); 
+		Arguments.Add(TEXT("LandscapeProxyName"), FText::FromString(GetActorNameOrLabel()));
 
 		auto CreateMapCheckMessage = [](FMessageLog& MessageLog)
 		{
 			if (CVarLandscapeSupressMapCheckWarnings_Nanite->GetBool())
 			{
-				return MessageLog.Info();	
+				return MessageLog.Info();
 			}
 			return MessageLog.Warning();
 		};
-		
+
 		TWeakObjectPtr<ALandscapeProxy> WeakLandscapeProxy(this);
 
 		FMessageLog MessageLog("MapCheck");
@@ -4286,40 +4308,40 @@ PRAGMA_ENABLE_DEPRECATION_WARNINGS
 			->AddToken(FTextToken::Create(FText::Format(LOCTEXT("MapCheck_Message_LandscapeRebuildNanite", "{LandscapeProxyName} : Landscape Nanite is enabled but saved mesh data is out of date. "), Arguments)))
 			->AddToken(FActionToken::Create(LOCTEXT("MapCheck_SaveFixedUpData", "Save Modified Landscapes"), LOCTEXT("MapCheck_SaveFixedUpData_Desc", "Saves the modified landscape proxy actors"),
 				FOnActionTokenExecuted::CreateLambda([WeakLandscapeProxy]()
+		{
+			if (!WeakLandscapeProxy.IsValid())
+			{
+				return;
+			}
+			ULandscapeInfo* Info = WeakLandscapeProxy->GetLandscapeInfo();
+			check(Info);
+
+			TSet<UPackage*> DirtyNanitePackages;
+			Info->ForEachLandscapeProxy([&DirtyNanitePackages](const ALandscapeProxy* Proxy)
+			{
+				if (!Proxy->IsNaniteMeshUpToDate())
 				{
-					if (!WeakLandscapeProxy.IsValid())
-					{
-						return;
-					}
-					ULandscapeInfo* Info = WeakLandscapeProxy->GetLandscapeInfo();
-					check(Info);
-					
-					TSet<UPackage*> DirtyNanitePackages;
-					Info->ForEachLandscapeProxy([&DirtyNanitePackages](const ALandscapeProxy* Proxy)
-					{
-						if (!Proxy->IsNaniteMeshUpToDate())
-						{
-							DirtyNanitePackages.Add(Proxy->GetOutermost());
-						}
-						return true;
-					});
-					
-					Info->UpdateNanite(nullptr);
-					
-					constexpr bool bPromptUserToSave = true;
-					constexpr bool bSaveMapPackages = true;
-					constexpr bool bSaveContentPackages = true;
-					constexpr bool bFastSave = false;
-					constexpr bool bNotifyNoPackagesSaved = false;
-					constexpr bool bCanBeDeclined = true;
-
-					FEditorFileUtils::SaveDirtyPackages(bPromptUserToSave, bSaveMapPackages, bSaveContentPackages, bFastSave, bNotifyNoPackagesSaved, bCanBeDeclined, nullptr,
-						[ &DirtyNanitePackages ](const UPackage* Package) { return !DirtyNanitePackages.Contains(Package);}  );
-				
-				}), FCanExecuteActionToken::CreateLambda([WeakLandscapeProxy]() { return WeakLandscapeProxy.IsValid() ? !WeakLandscapeProxy->IsNaniteMeshUpToDate() : false;} ))
-				);
+					DirtyNanitePackages.Add(Proxy->GetOutermost());
+				}
+				return true;
+			});
+
+			Info->UpdateNanite(nullptr);
+
+			constexpr bool bPromptUserToSave = true;
+			constexpr bool bSaveMapPackages = true;
+			constexpr bool bSaveContentPackages = true;
+			constexpr bool bFastSave = false;
+			constexpr bool bNotifyNoPackagesSaved = false;
+			constexpr bool bCanBeDeclined = true;
+
+			FEditorFileUtils::SaveDirtyPackages(bPromptUserToSave, bSaveMapPackages, bSaveContentPackages, bFastSave, bNotifyNoPackagesSaved, bCanBeDeclined, nullptr,
+				[&DirtyNanitePackages](const UPackage* Package) { return !DirtyNanitePackages.Contains(Package); });
+
+		}), FCanExecuteActionToken::CreateLambda([WeakLandscapeProxy]() { return WeakLandscapeProxy.IsValid() ? !WeakLandscapeProxy->IsNaniteMeshUpToDate() : false; }))
+			);
 	}
-	
+
 	UWorld* World = GetWorld();
 
 	// track feature level change to flush grass cache
@@ -4329,13 +4351,15 @@ PRAGMA_ENABLE_DEPRECATION_WARNINGS
 		FeatureLevelChangedDelegateHandle = World->AddOnFeatureLevelChangedHandler(FeatureLevelChangedDelegate);
 	}
 	RepairInvalidTextures();
-	PRAGMA_DISABLE_DEPRECATION_WARNINGS
+
+PRAGMA_DISABLE_DEPRECATION_WARNINGS
 	if (NaniteComponent_DEPRECATED)
 	{
 		NaniteComponents.Add(NaniteComponent_DEPRECATED);
 		NaniteComponent_DEPRECATED = nullptr;
 	}
-	PRAGMA_ENABLE_DEPRECATION_WARNINGS
+PRAGMA_ENABLE_DEPRECATION_WARNINGS
+
 	// Handle Nanite representation invalidation on load: 
 	if (!HasAnyFlags(RF_ClassDefaultObject) && !FPlatformProperties::RequiresCookedData())
 	{
@@ -4509,7 +4533,7 @@ namespace UE::Landscape::Private
 				FCanExecuteActionToken::CreateStatic(&HasModifiedLandscapes),
 				/*bInSingleUse = */false))
 			->AddToken(FTextToken::Create(FText::Format(LOCTEXT("MapCheck_Message_LandscapeProxy_FixupSharedData_SharedProperties", "The following properties were synchronized: {0}."), FText::FromString(SynchronizedPropertiesStringBuilder.ToString()))));
-			
+
 		if (bAddSilencingMessage)
 		{
 			Message->AddToken(FTextToken::Create(LOCTEXT("MapCheck_Message_LandscapeProxy_SilenceWarning", "You can silence this warning and perform the deprecation silently using the landscape.SilenceSharedPropertyDeprecationFixup CVar. ")));
@@ -4730,7 +4754,7 @@ void ALandscapeProxy::UpgradeSharedProperties(ALandscape* InParentLandscape)
 		FProperty* Property = *PropertyIterator;
 
 		if (Property == nullptr)
-		{
+{
 			continue;
 		}
 
@@ -4818,10 +4842,10 @@ void ALandscapeProxy::FixupSharedData(ALandscape* Landscape, const bool bMapChec
 		bool bUpdated = !SynchronizedProperties.IsEmpty();
 
 		TSet<FGuid> LayerGuids;
-		Algo::Transform(Landscape->LandscapeLayers, LayerGuids, [](const FLandscapeLayer& Layer) { return Layer.Guid; });
+		Algo::Transform(Landscape->GetLayers(), LayerGuids, [](const FLandscapeLayer& Layer) { return Layer.Guid; });
 		bUpdated |= RemoveObsoleteLayers(LayerGuids);
 
-		for (const FLandscapeLayer& Layer : Landscape->LandscapeLayers)
+		for (const FLandscapeLayer& Layer : Landscape->GetLayers())
 		{
 			bUpdated |= AddLayer(Layer.Guid);
 		}
@@ -4835,7 +4859,7 @@ void ALandscapeProxy::FixupSharedData(ALandscape* Landscape, const bool bMapChec
 			bool bNeedsManualResave = LandscapeInfo->MarkObjectDirty(/*InObject = */this, /*bInForceResave = */true, LandscapeActor);
 
 			if (bMapCheck && bNeedsManualResave)
-			{
+	{
 				UE::Landscape::Private::DisplaySynchronizedPropertiesMapcheckWarning(SynchronizedProperties, *this, *Landscape);
 			}
 		}
@@ -4852,7 +4876,7 @@ void ALandscapeProxy::SetAbsoluteSectionBase(FIntPoint InSectionBase)
 		FIntPoint AbsoluteSectionBase = Comp->GetSectionBase() + Difference;
 		Comp->SetSectionBase(AbsoluteSectionBase);
 	});
-
+	
 	for (int32 CompIdx = 0; CompIdx < CollisionComponents.Num(); CompIdx++)
 	{
 		ULandscapeHeightfieldCollisionComponent* Comp = CollisionComponents[CompIdx];
@@ -5006,7 +5030,7 @@ bool ULandscapeInfo::TryAddToModifiedPackages(UPackage* InPackage, const ALandsc
 bool ULandscapeInfo::MarkObjectDirty(UObject* InObject, bool bInForceResave, const ALandscape* InLandscapeOverride)
 {
 	check(InObject && (InObject->IsA<ALandscapeProxy>() || InObject->GetTypedOuter<ALandscapeProxy>() != nullptr));
-
+		
 	bool bWasAddedToModifiedPackages = false;
 	if (bInForceResave)
 	{
@@ -6327,13 +6351,13 @@ void ALandscapeProxy::ClearNaniteTransactional()
 }
 
 void ALandscapeProxy::UpdateNaniteSharedPropertiesFromActor()
-{	
+{
 	for (ULandscapeNaniteComponent* NaniteComponent : NaniteComponents)
 	{
 		if (NaniteComponent)
 		{
 			NaniteComponent->UpdatedSharedPropertiesFromActor();
-		}	
+		}
 	}
 }
 
@@ -6383,7 +6407,7 @@ void ALandscapeProxy::UpdatePhysicalMaterialTasksStatus(TSet<ULandscapeComponent
 	}
 
 	if (OutdatedComponentsCount)
-{
+	{
 		*OutdatedComponentsCount = OutdatedCount;
 	}
 }
@@ -6397,7 +6421,7 @@ void ALandscapeProxy::UpdatePhysicalMaterialTasks(bool bInShouldMarkDirty)
 	{
 		Component->UpdatePhysicalMaterialTasks();
 	}
-	if (bInShouldMarkDirty && PendingComponentsToBeSaved >0)
+	if (bInShouldMarkDirty && PendingComponentsToBeSaved > 0)
 	{
 		MarkPackageDirty();
 	}
@@ -6501,8 +6525,6 @@ bool ALandscapeProxy::AuditNaniteMaterials() const
 	return true;
 }
 
-
-
 void ALandscapeProxy::InvalidateGeneratedComponentData(bool bInvalidateLightingCache)
 {
 	InvalidateGeneratedComponentData(LandscapeComponents, bInvalidateLightingCache);
@@ -6526,7 +6548,7 @@ void ALandscapeProxy::InvalidateGeneratedComponentData(const TArray<ULandscapeCo
 		ALandscapeProxy* Proxy = Iter.Key();
 		Proxy->FlushGrassComponents(&Iter.Value());
 
-	#if WITH_EDITOR
+#if WITH_EDITOR
 		ULandscapeSubsystem* Subsystem = Proxy->GetWorld()->GetSubsystem<ULandscapeSubsystem>();
 		if (Subsystem->IsLiveNaniteRebuildEnabled())
 		{
@@ -6536,10 +6558,10 @@ void ALandscapeProxy::InvalidateGeneratedComponentData(const TArray<ULandscapeCo
 		{
 			Proxy->InvalidateOrUpdateNaniteRepresentation(/* bInCheckContentId = */true, /*InTargetPlatform = */nullptr);
 		}
-		
+
 		FLandscapeProxyComponentDataChangedParams ChangeParams(Iter.Value());
 		Proxy->OnComponentDataChanged.Broadcast(Iter.Key(), ChangeParams);
-	#endif
+#endif // WITH_EDITOR
 
 		Proxy->UpdateRenderingMethod();
 	}
diff --git a/Engine/Source/Runtime/Landscape/Private/LandscapeBlueprintSupport.cpp b/Engine/Source/Runtime/Landscape/Private/LandscapeBlueprintSupport.cpp
--- a/Engine/Source/Runtime/Landscape/Private/LandscapeBlueprintSupport.cpp
+++ b/Engine/Source/Runtime/Landscape/Private/LandscapeBlueprintSupport.cpp
@@ -23,7 +23,7 @@ void ALandscapeProxy::EditorApplySpline(USplineComponent* InSplineComponent, flo
 	{
 		if (ALandscape* Landscape = GetLandscapeInfo()->LandscapeActor.Get())
 		{
-			const FLandscapeLayer* Layer = Landscape->GetLayer(EditLayerName);
+			const FLandscapeLayer* Layer = Landscape->GetLayerConst(EditLayerName);
 			if (Landscape->HasLayersContent() && (Layer == nullptr))
 			{
 				UE_LOG(LogLandscape, Error, TEXT("Invalid landscape edit layer name (\"%s\") for Edit Layers-enabled landscape. Cannot apply spline. "), *EditLayerName.ToString());
diff --git a/Engine/Source/Runtime/Landscape/Private/LandscapeConfigHelper.cpp b/Engine/Source/Runtime/Landscape/Private/LandscapeConfigHelper.cpp
--- a/Engine/Source/Runtime/Landscape/Private/LandscapeConfigHelper.cpp
+++ b/Engine/Source/Runtime/Landscape/Private/LandscapeConfigHelper.cpp
@@ -16,6 +16,7 @@
 #include "LandscapeSplinesComponent.h"
 #include "LandscapeEdit.h"
 #include "LandscapeDataAccess.h"
+#include "LandscapeEditLayer.h"
 #include "LandscapeSubsystem.h"
 #include "ActorPartition/ActorPartitionSubsystem.h"
 
@@ -379,12 +380,19 @@ void FLandscapeConfigHelper::CopyRegionToComponent(ULandscapeInfo* InLandscapeIn
 		}
 		InComponent->InitHeightmapData(HeightInitData, false);
 
-		for (const FLandscapeLayer& LandscapeLayer : InLandscapeInfo->LandscapeActor->LandscapeLayers)
+		for (const FLandscapeLayer& LandscapeLayer : InLandscapeInfo->LandscapeActor->GetLayers())
 		{		
+			check(LandscapeLayer.EditLayer != nullptr);
+			if (!LandscapeLayer.EditLayer->NeedsPersistentTextures())
+			{
+				continue;
+			}
+
 			TMap<UTexture2D*, UTexture2D*> CreatedTextures;
 			InComponent->AddDefaultLayerData(LandscapeLayer.Guid, { InComponent }, CreatedTextures);
 
-			if (LandscapeLayer.Guid == InLandscapeInfo->LandscapeActor->LandscapeSplinesTargetLayerGuid)
+			// No need to copy the persistent textures if they are not manually edited (e.g. fully procedural layers with persistent textures : e.g. the splines layer) :
+			if (!LandscapeLayer.EditLayer->SupportsEditingTools())
 			{
 				continue;
 			}
@@ -565,11 +573,11 @@ ULandscapeInfo* FLandscapeConfigHelper::ChangeConfiguration(ULandscapeInfo* InLa
 	if (OldLandscape->HasLayersContent())
 	{
 		NewLandscape->bCanHaveLayersContent = true;
-		OldLandscape->ForEachLayer([NewLandscape](FLandscapeLayer& Layer)
+		for (const FLandscapeLayer& OldLayer : OldLandscape->GetLayers())
 		{
-			NewLandscape->LandscapeLayers.Add(Layer);
-		});
-		NewLandscape->LandscapeSplinesTargetLayerGuid = OldLandscape->LandscapeSplinesTargetLayerGuid;
+			FLandscapeLayer* NewLayer = NewLandscape->DuplicateLayerAndMoveBrushes(OldLayer);
+			check((NewLayer != nullptr) && (NewLayer->EditLayer != nullptr));
+		}
 	}
 		
 	ULandscapeInfo* NewLandscapeInfo = NewLandscape->CreateLandscapeInfo();
diff --git a/Engine/Source/Runtime/Landscape/Private/LandscapeEdit.cpp b/Engine/Source/Runtime/Landscape/Private/LandscapeEdit.cpp
--- a/Engine/Source/Runtime/Landscape/Private/LandscapeEdit.cpp
+++ b/Engine/Source/Runtime/Landscape/Private/LandscapeEdit.cpp
@@ -3005,15 +3005,26 @@ ULandscapeLayerInfoObject* ALandscapeProxy::CreateLayerInfo(const TCHAR* InLayer
 #define HEIGHTDATA(X,Y) (HeightData.Num() == 0 ? LandscapeDataAccess::GetTexHeight(0.0f) : HeightData[ FMath::Clamp<int32>(Y,0,VertsY) * VertsX + FMath::Clamp<int32>(X,0,VertsX) ])
 ENGINE_API extern bool GDisableAutomaticTextureMaterialUpdateDependencies;
 
-LANDSCAPE_API void ALandscapeProxy::Import(const FGuid& InGuid, int32 InMinX, int32 InMinY, int32 InMaxX, int32 InMaxY, int32 InNumSubsections, int32 InSubsectionSizeQuads, const TMap<FGuid, TArray<uint16>>& InImportHeightData, 
+void ALandscapeProxy::Import(const FGuid& InGuid, int32 InMinX, int32 InMinY, int32 InMaxX, int32 InMaxY, int32 InNumSubsections, int32 InSubsectionSizeQuads, const TMap<FGuid, TArray<uint16>>& InImportHeightData, 
 										   const TCHAR* const InHeightmapFileName, const TMap<FGuid, TArray<FLandscapeImportLayerInfo>>& InImportMaterialLayerInfos, ELandscapeImportAlphamapType InImportMaterialLayerType, const TArray<FLandscapeLayer>* InImportLayers)
+{
+	TArrayView<const FLandscapeLayer> ImportLayers;
+	if (InImportLayers != nullptr)
+	{
+		ImportLayers = MakeArrayView(static_cast<const FLandscapeLayer*>(InImportLayers->GetData()), InImportLayers->Num());
+	}
+	Import(InGuid, InMinX, InMinY, InMaxX, InMaxY, InNumSubsections, InSubsectionSizeQuads, InImportHeightData, InHeightmapFileName, InImportMaterialLayerInfos, InImportMaterialLayerType, ImportLayers);
+}
+
+void ALandscapeProxy::Import(const FGuid& InGuid, int32 InMinX, int32 InMinY, int32 InMaxX, int32 InMaxY, int32 InNumSubsections, int32 InSubsectionSizeQuads, const TMap<FGuid, TArray<uint16>>& InImportHeightData,
+	const TCHAR* const InHeightmapFileName, const TMap<FGuid, TArray<FLandscapeImportLayerInfo>>& InImportMaterialLayerInfos, ELandscapeImportAlphamapType InImportMaterialLayerType, const TArrayView<const FLandscapeLayer>& InImportLayers)
 {
 	TRACE_CPUPROFILER_EVENT_SCOPE(ALandscapeProxy::Import);
 	
 	check(InGuid.IsValid());
 	check(InImportHeightData.Num() == InImportMaterialLayerInfos.Num());
 
-	check(CanHaveLayersContent() || InImportLayers == nullptr);
+	check(CanHaveLayersContent() || InImportLayers.IsEmpty());
 
 	FScopedSlowTask SlowTask(2, LOCTEXT("BeingImportingLandscapeTask", "Importing Landscape"));
 	SlowTask.MakeDialog();
@@ -3639,7 +3650,7 @@ LANDSCAPE_API void ALandscapeProxy::Import(const FGuid& InGuid, int32 InMinX, in
 		// Create the default layer first
 		ALandscape* LandscapeActor = GetLandscapeActor();
 		check(LandscapeActor != nullptr);
-		if (LandscapeActor->GetLayerCount() == 0 && InImportLayers == nullptr)
+		if (LandscapeActor->GetLayerCount() == 0 && InImportLayers.IsEmpty())
 		{
 			LandscapeActor->CreateDefaultLayer();
 		}
@@ -3659,9 +3670,9 @@ LANDSCAPE_API void ALandscapeProxy::Import(const FGuid& InGuid, int32 InMinX, in
 		TArray<FLayerImportSettings> LayerImportSettings;		
 
 		// Only create Layers on main Landscape
-		if (LandscapeActor == this && InImportLayers != nullptr)
+		if (LandscapeActor == this && !InImportLayers.IsEmpty())
 		{
-			for (const FLandscapeLayer& OldLayer : *InImportLayers)
+			for (const FLandscapeLayer& OldLayer : InImportLayers)
 			{
 				FLandscapeLayer* NewLayer = LandscapeActor->DuplicateLayerAndMoveBrushes(OldLayer);
 				check(NewLayer != nullptr);
@@ -3683,7 +3694,7 @@ LANDSCAPE_API void ALandscapeProxy::Import(const FGuid& InGuid, int32 InMinX, in
 			}
 
 			// And we will fill all the landscape components with the provided final layer content put into the default layer (aka layer index 0)
-			const FLandscapeLayer* DefaultLayer = LandscapeActor->GetLayer(0);
+			const FLandscapeLayer* DefaultLayer = LandscapeActor->GetLayerConst(0);
 			check(DefaultLayer != nullptr);
 
 			FLayerImportSettings ImportSettings;
@@ -4015,7 +4026,7 @@ TSharedRef<UE::Landscape::Nanite::FAsyncBuildData> ALandscapeProxy::MakeAsyncNan
 	return AsyncBuildData;
 }
 
-LANDSCAPE_API void ALandscape::SetNanitePositionPrecision(int32 InPrecision, bool bInShouldDirtyPackage)
+void ALandscape::SetNanitePositionPrecision(int32 InPrecision, bool bInShouldDirtyPackage)
 {
 	NanitePositionPrecision = InPrecision;
 
@@ -4608,7 +4619,7 @@ bool ULandscapeInfo::GetLandscapeExtent(int32& MinX, int32& MinY, int32& MaxX, i
 	return (MinX != MAX_int32);
 }
 
-LANDSCAPE_API bool ULandscapeInfo::GetLandscapeXYComponentBounds(FIntRect& OutXYComponentBounds) const
+bool ULandscapeInfo::GetLandscapeXYComponentBounds(FIntRect& OutXYComponentBounds) const
 {
 	OutXYComponentBounds = XYComponentBounds;
 
@@ -4616,7 +4627,7 @@ LANDSCAPE_API bool ULandscapeInfo::GetLandscapeXYComponentBounds(FIntRect& OutXY
 		&& (OutXYComponentBounds.Max.X != MAX_int32) && (OutXYComponentBounds.Max.Y != MAX_int32);
 }
 
-LANDSCAPE_API void ULandscapeInfo::ForAllLandscapeComponents(TFunctionRef<void(ULandscapeComponent*)> Fn) const
+void ULandscapeInfo::ForAllLandscapeComponents(TFunctionRef<void(ULandscapeComponent*)> Fn) const
 {
 	ForEachLandscapeProxy([&](ALandscapeProxy* Proxy)
 	{
@@ -5070,7 +5081,7 @@ void ALandscape::PostRegisterAllComponents()
 
 	auto HasValidBrush = [this]()
 	{
-		for (const FLandscapeLayer& Layer : LandscapeLayers)
+		for (const FLandscapeLayer& Layer : LandscapeEditLayers)
 		{
 			for (const FLandscapeLayerBrush& Brush : Layer.Brushes)
 			{
@@ -5131,7 +5142,7 @@ void ALandscape::PostEditImport()
 	}
 
 	// We need to reparent brushes that may have been part of the copy/pasted actors :
-	for (FLandscapeLayer& Layer : LandscapeLayers)
+	for (FLandscapeLayer& Layer : LandscapeEditLayers)
 	{
 		for (FLandscapeLayerBrush& Brush : Layer.Brushes)
 		{
diff --git a/Engine/Source/Runtime/Landscape/Private/LandscapeEditInterface.cpp b/Engine/Source/Runtime/Landscape/Private/LandscapeEditInterface.cpp
--- a/Engine/Source/Runtime/Landscape/Private/LandscapeEditInterface.cpp
+++ b/Engine/Source/Runtime/Landscape/Private/LandscapeEditInterface.cpp
@@ -2346,15 +2346,11 @@ void FLandscapeEditDataInterface::ReplaceLayer(ULandscapeLayerInfoObject* FromLa
 	if (LandscapeInfo->LandscapeActor.IsValid() && LandscapeInfo->LandscapeActor->HasLayersContent())
 	{
 		LandscapeInfo->LandscapeActor->Modify(GetShouldDirtyPackage());
-		LandscapeInfo->LandscapeActor->ForEachLayer([&](FLandscapeLayer& CurrentLayer)
+		LandscapeInfo->LandscapeActor->ReplaceLayerSubstractiveBlendStatus(FromLayerInfo, ToLayerInfo, GetShouldDirtyPackage());
+		LandscapeInfo->LandscapeActor->ForEachLayerConst([DoReplace](const FLandscapeLayer& CurrentLayer)
 		{
-			bool OutValue;
-			if (CurrentLayer.WeightmapLayerAllocationBlend.RemoveAndCopyValue(FromLayerInfo, OutValue))
-			{
-				CurrentLayer.WeightmapLayerAllocationBlend.Add(ToLayerInfo, OutValue);
-			}
-
 			DoReplace(CurrentLayer.Guid);
+			return true;
 		});
 		LandscapeInfo->LandscapeActor->RequestLayersContentUpdateForceAll(ELandscapeLayerUpdateMode::Update_Weightmap_All);
 	}
diff --git a/Engine/Source/Runtime/Landscape/Private/LandscapeEditLayer.cpp b/Engine/Source/Runtime/Landscape/Private/LandscapeEditLayer.cpp
new file mode 100644
--- /dev/null
+++ b/Engine/Source/Runtime/Landscape/Private/LandscapeEditLayer.cpp
@@ -0,0 +1,60 @@
+// Copyright Epic Games, Inc. All Rights Reserved.
+
+#include "LandscapeEditLayer.h"
+#include "LandscapeEditTypes.h"
+#include "Landscape.h"
+
+#define LOCTEXT_NAMESPACE "LandscapeEditLayer"
+
+
+// ----------------------------------------------------------------------------------
+
+bool ULandscapeEditLayer::SupportsTargetType(ELandscapeToolTargetType InType) const
+{
+	return (InType == ELandscapeToolTargetType::Heightmap) || (InType == ELandscapeToolTargetType::Weightmap) || (InType == ELandscapeToolTargetType::Visibility);
+}
+
+
+// ----------------------------------------------------------------------------------
+
+bool ULandscapeEditLayerSplines::SupportsTargetType(ELandscapeToolTargetType InType) const
+{
+	return (InType == ELandscapeToolTargetType::Heightmap) || (InType == ELandscapeToolTargetType::Weightmap) || (InType == ELandscapeToolTargetType::Visibility);
+}
+
+void ULandscapeEditLayerSplines::OnLayerCreated(FLandscapeLayer& Layer)
+{
+	// Splines edit layer is always using alpha blend mode
+	Layer.BlendMode = LSBM_AlphaBlend;
+}
+
+TArray<ULandscapeEditLayerSplines::FEditLayerAction> ULandscapeEditLayerSplines::GetActions() const
+{
+	TArray<FEditLayerAction> Actions;
+
+#if WITH_EDITOR
+	// Register an "Update Splines" action :
+	Actions.Add(FEditLayerAction(
+		LOCTEXT("LandscapeEditLayerSplines_UpdateSplines", "Update Splines"),
+		FEditLayerAction::FExecuteDelegate::CreateWeakLambda(this, [](const FEditLayerAction::FExecuteParams& InParams)
+		{
+			InParams.GetLandscape()->UpdateLandscapeSplines(FGuid(), /*bUpdateOnlySelection = */false, /*bForceUpdate =*/true);
+			return FEditLayerAction::FExecuteResult(/*bInSuccess = */true);
+		}),
+		FEditLayerAction::FCanExecuteDelegate::CreateWeakLambda(this, [](const FEditLayerAction::FExecuteParams& InParams, FText& OutReason)
+		{
+			if (InParams.GetLayer()->bLocked)
+			{
+				OutReason = FText::Format(LOCTEXT("LandscapeEditLayerSplines_CannotUpdateSplinesOnLockedLayer", "Cannot update splines on layer '{0}' : the layer is currently locked"), FText::FromName(InParams.GetLayer()->Name));
+				return false;
+			}
+
+			OutReason = LOCTEXT("LandscapeEditLayerSplines_UpdateSplines_Tooltip", "Update Landscape Splines");
+			return true;
+		})));
+#endif // WITH_EDITOR
+
+	return Actions;
+}
+
+#undef LOCTEXT_NAMESPACE
\ No newline at end of file
diff --git a/Engine/Source/Runtime/Landscape/Private/LandscapeEditLayers.cpp b/Engine/Source/Runtime/Landscape/Private/LandscapeEditLayers.cpp
--- a/Engine/Source/Runtime/Landscape/Private/LandscapeEditLayers.cpp
+++ b/Engine/Source/Runtime/Landscape/Private/LandscapeEditLayers.cpp
@@ -6,6 +6,7 @@ LandscapeEditLayers.cpp: Landscape editing layers mode
 
 #include "LandscapeEdit.h"
 #include "Landscape.h"
+#include "LandscapeEditLayer.h"
 #include "LandscapeProxy.h"
 #include "LandscapeStreamingProxy.h"
 #include "LandscapeInfo.h"
@@ -1415,7 +1416,7 @@ public:
 		SHADER_PARAMETER_RDG_BUFFER_SRV(StructuredBuffer<uint>, InPaintLayerInfoIndices)
 		SHADER_PARAMETER_RDG_BUFFER_SRV(StructuredBuffer<FLandscapeWeightmapPaintLayerInfo>, InPaintLayerInfos)
 		RENDER_TARGET_BINDING_SLOTS()
-	END_SHADER_PARAMETER_STRUCT()
+		END_SHADER_PARAMETER_STRUCT()
 
 		static bool ShouldCompilePermutation(const FGlobalShaderPermutationParameters& Parameters)
 	{
@@ -1804,7 +1805,7 @@ bool ALandscape::SupportsEditLayersLocalMerge()
 	ALandscape* Landscape = GetLandscapeActor();
 	check(Landscape);
 
-	for (const FLandscapeLayer& Layer : Landscape->LandscapeLayers)
+	for (const FLandscapeLayer& Layer : Landscape->LandscapeEditLayers)
 	{
 		// No BP brush is supported for local merge of edit layers : 
 		if (Layer.bVisible && Algo::AnyOf(Layer.Brushes, [](const FLandscapeLayerBrush& Brush) { return (Brush.GetBrush() != nullptr) && Brush.GetBrush()->IsVisible(); }))
@@ -1821,7 +1822,7 @@ bool ALandscape::HasNormalCaptureBPBrushLayer()
 	ALandscape* Landscape = GetLandscapeActor();
 	check(Landscape);
 
-	for (const FLandscapeLayer& Layer : Landscape->LandscapeLayers)
+	for (const FLandscapeLayer& Layer : Landscape->LandscapeEditLayers)
 	{
 		if (Layer.bVisible && Algo::AnyOf(Layer.Brushes, [](const FLandscapeLayerBrush& Brush) { return (Brush.GetBrush() != nullptr) && Brush.GetBrush()->IsVisible() && Brush.GetBrush()->GetCaptureBoundaryNormals(); }))
 		{
@@ -1915,8 +1916,8 @@ bool ALandscape::CreateLayersRenderingResource(bool bUseNormalCapture)
 			}
 		};
 
-		for (int32 i = 0; i < (int32)EHeightmapRTType::HeightmapRT_Count; ++i)
-		{
+			for (int32 i = 0; i < (int32)EHeightmapRTType::HeightmapRT_Count; ++i)
+			{
 			if (bCreateFromScratch)
 			{
 				FText DisplayName = StaticEnum<EHeightmapRTType>()->GetDisplayValueAsText((EHeightmapRTType)i);
@@ -1929,23 +1930,23 @@ bool ALandscape::CreateLayersRenderingResource(bool bUseNormalCapture)
 				Landscape->HeightmapRTList[i]->ClearColor = FLinearColor(0.0f, 0.0f, 0.0f, 0.0f);
 			}
 
-			if (i < (int32)EHeightmapRTType::HeightmapRT_Mip1) // Landscape size RT
-			{
+				if (i < (int32)EHeightmapRTType::HeightmapRT_Mip1) // Landscape size RT
+				{
 				InitOrResizeRT(Landscape->HeightmapRTList[i], CurrentMipSizeX, CurrentMipSizeY, bCreateFromScratch);
-			}
-			else // Mips
-			{
-				CurrentMipSizeX >>= 1;
-				CurrentMipSizeY >>= 1;
+				}
+				else // Mips
+				{
+					CurrentMipSizeX >>= 1;
+					CurrentMipSizeY >>= 1;
 				InitOrResizeRT(Landscape->HeightmapRTList[i], CurrentMipSizeX, CurrentMipSizeY, bCreateFromScratch);
-			}
+				}
 
-			// Only generate required mips RT
-			if (CurrentMipSizeX == ComponentCounts.X && CurrentMipSizeY == ComponentCounts.Y)
-			{
-				break;
+				// Only generate required mips RT
+				if (CurrentMipSizeX == ComponentCounts.X && CurrentMipSizeY == ComponentCounts.Y)
+				{
+					break;
+				}
 			}
-		}
 
 		// Weightmap mip size
 		CurrentMipSizeX = ((SubsectionSizeQuads + 1) * NumSubsections) * ComponentCounts.X;
@@ -1953,7 +1954,7 @@ bool ALandscape::CreateLayersRenderingResource(bool bUseNormalCapture)
 		bCreateFromScratch = (Landscape->WeightmapRTList.Num() == 0);
 
 		if (bCreateFromScratch)
-		{
+				{
 			Landscape->WeightmapRTList.Init(nullptr, (int32)EWeightmapRTType::WeightmapRT_Count);
 		}
 
@@ -1977,26 +1978,26 @@ bool ALandscape::CreateLayersRenderingResource(bool bUseNormalCapture)
 				{
 					Landscape->WeightmapRTList[i]->RenderTargetFormat = RTF_R8;
 				}
-			}
+				}
 
 			if (i < (int32)EWeightmapRTType::WeightmapRT_Mip0)
-			{
+				{
 				InitOrResizeRT(Landscape->WeightmapRTList[i], CurrentMipSizeX, CurrentMipSizeY, bCreateFromScratch);
-			}
-			else // Mips
-			{
+				}
+				else // Mips
+				{
 				InitOrResizeRT(Landscape->WeightmapRTList[i], CurrentMipSizeX, CurrentMipSizeY, bCreateFromScratch);
 
-				CurrentMipSizeX >>= 1;
-				CurrentMipSizeY >>= 1;
-			}
+					CurrentMipSizeX >>= 1;
+					CurrentMipSizeY >>= 1;
+				}
 
-			// Only generate required mips RT
-			if (CurrentMipSizeX < ComponentCounts.X && CurrentMipSizeY < ComponentCounts.Y)
-			{
-				break;
+				// Only generate required mips RT
+				if (CurrentMipSizeX < ComponentCounts.X && CurrentMipSizeY < ComponentCounts.Y)
+				{
+					break;
+				}
 			}
-		}
 
 		InitializeLayersWeightmapResources();
 	}
@@ -2155,7 +2156,7 @@ void ALandscape::CopyOldDataToDefaultLayer(ALandscapeProxy* InProxy)
 
 	InProxy->Modify();
 
-	FLandscapeLayer* DefaultLayer = GetLayer(0);
+	const FLandscapeLayer* DefaultLayer = GetLayerConst(0);
 	check(DefaultLayer != nullptr);
 
 	struct FWeightmapTextureData
@@ -2349,8 +2350,7 @@ void ALandscapeProxy::InitializeProxyLayersWeightmapUsage()
 			// Reinitialize the weightmap usages for the base (final) paint layers allocations :
 			Component->InitializeLayersWeightmapUsage(FGuid());
 
-			const FLandscapeLayer* SplinesLayer = Landscape->GetLandscapeSplinesReservedLayer();
-			for (const FLandscapeLayer& Layer : Landscape->LandscapeLayers)
+			for (const FLandscapeLayer& Layer : Landscape->GetLayers())
 			{
 				// Reinitialize each edit layer's weightmap usages list :
 				Component->InitializeLayersWeightmapUsage(Layer.Guid);
@@ -2368,7 +2368,7 @@ void ULandscapeComponent::InitializeLayersWeightmapUsage(const FGuid& InLayerGui
 	check(Proxy);
 	ALandscape* Landscape = GetLandscapeActor();
 	check(Landscape);
-	const FLandscapeLayer* SplinesLayer = Landscape->GetLandscapeSplinesReservedLayer();
+	const FLandscapeLayer* SplinesLayer = Landscape->FindLayerOfType(ULandscapeEditLayerSplines::StaticClass());
 	FGuid SplinesEditLayerGuid = SplinesLayer ? SplinesLayer->Guid : FGuid();
 
 	// Don't consider invalid edit layers : 
@@ -2515,8 +2515,8 @@ void ALandscapeProxy::ValidateProxyLayersWeightmapUsage() const
 
 			// Validate edit layers weightmap allocations : 
 			{
-				const FLandscapeLayer* SplinesLayer = Landscape->GetLandscapeSplinesReservedLayer();
-				for (const FLandscapeLayer& Layer : Landscape->LandscapeLayers)
+				const FLandscapeLayer* SplinesLayer = Landscape->FindLayerOfType(ULandscapeEditLayerSplines::StaticClass());
+				for (const FLandscapeLayer& Layer : Landscape->GetLayers())
 				{
 					FLandscapeLayerComponentData* LayerData = Component->GetLayerData(Layer.Guid);
 
@@ -3594,14 +3594,14 @@ void ALandscape::InvalidateRVTForTextures(const TSet<UTexture2D*>& InTextures)
 						{
 							// Landscape only supports UPrimitiveComponent for the moment
 							if (UPrimitiveComponent* PrimitiveComponent = PrimitiveComponentInterface->GetUObject<UPrimitiveComponent>())
-							{
-								PrimitiveComponentsToInvalidate.Add(PrimitiveComponent);
-							}
+						{
+							PrimitiveComponentsToInvalidate.Add(PrimitiveComponent);
 						}
 					}
 				}
 			}
 		}
+		}
 
 		if (!PrimitiveComponentsToInvalidate.IsEmpty())
 		{
@@ -3623,7 +3623,7 @@ void ALandscape::InvalidateRVTForTextures(const TSet<UTexture2D*>& InTextures)
 
 bool ALandscape::PrepareLayersTextureResources(bool bInWaitForStreaming)
 {
-	return PrepareLayersTextureResources(LandscapeLayers, bInWaitForStreaming);
+	return PrepareLayersTextureResources(LandscapeEditLayers, bInWaitForStreaming);
 }
 
 bool ALandscape::PrepareLayersTextureResources(const TArray<FLandscapeLayer>& InLayers, bool bInWaitForStreaming)
@@ -3666,7 +3666,7 @@ bool ALandscape::PrepareLayersBrushResources(ERHIFeatureLevel::Type InFeatureLev
 {
 	TRACE_CPUPROFILER_EVENT_SCOPE(LandscapeLayers_PrepareLayersBrushTextureResources);
 	TSet<UObject*> Dependencies;
-	for (const FLandscapeLayer& Layer : LandscapeLayers)
+	for (const FLandscapeLayer& Layer : LandscapeEditLayers)
 	{
 		for (const FLandscapeLayerBrush& Brush : Layer.Brushes)
 		{
@@ -4521,8 +4521,8 @@ void ALandscape::PrepareLayersHeightmapsLocalMergeRenderThreadData(const FUpdate
 	OutRenderThreadData.NumSubsections = NumSubsections;
 
 	// Prepare landscape edit layers data common to all landscape components: 
-	OutRenderThreadData.VisibleEditLayerInfos.Reserve(LandscapeLayers.Num());
-	for (const FLandscapeLayer& Layer : LandscapeLayers)
+	OutRenderThreadData.VisibleEditLayerInfos.Reserve(LandscapeEditLayers.Num());
+	for (const FLandscapeLayer& Layer : LandscapeEditLayers)
 	{
 		if (Layer.bVisible && !InMergeParams.bSkipBrush)
 		{
@@ -4562,7 +4562,7 @@ void ALandscape::PrepareLayersHeightmapsLocalMergeRenderThreadData(const FUpdate
 			const FIntRect ComponentTextureSubregion(HeightmapOffset, HeightmapOffset + OutRenderThreadData.ComponentSizeVerts);
 
 			NewComponentRenderInfo.VisibleLayerHeightmapTextures.Reserve(OutRenderThreadData.VisibleEditLayerInfos.Num());
-			for (const FLandscapeLayer& Layer : LandscapeLayers)
+			for (const FLandscapeLayer& Layer : LandscapeEditLayers)
 			{
 				if (Layer.bVisible && !InMergeParams.bSkipBrush)
 				{
@@ -4863,7 +4863,7 @@ int32 ALandscape::PerformLayersHeightmapsGlobalMerge(const FUpdateLayersContentC
 	UTextureRenderTarget2D* LandscapeScratchRT2 = HeightmapRTList[(int32)EHeightmapRTType::HeightmapRT_Scratch2];
 	UTextureRenderTarget2D* LandscapeScratchRT3 = HeightmapRTList[(int32)EHeightmapRTType::HeightmapRT_Scratch3];
 
-	for (FLandscapeLayer& Layer : LandscapeLayers)
+	for (FLandscapeLayer& Layer : LandscapeEditLayers)
 	{
 		// Draw each Layer's heightmaps to a Combined RT Atlas in LandscapeScratchRT1
 		ShaderParams.ApplyLayerModifiers = false;
@@ -4875,8 +4875,7 @@ int32 ALandscape::PerformLayersHeightmapsGlobalMerge(const FUpdateLayersContentC
 		if (Layer.BlendMode == LSBM_AlphaBlend)
 		{
 			// For now, only Layer reserved for Landscape Splines will use the AlphaBlendMode
-			const FLandscapeLayer* SplinesReservedLayer = GetLandscapeSplinesReservedLayer();
-			check(&Layer == SplinesReservedLayer);
+			check((Layer.EditLayer != nullptr) && Layer.EditLayer->IsA<ULandscapeEditLayerSplines>());
 			ShaderParams.LayerAlpha = 1.0f;
 		}
 		else
@@ -4925,7 +4924,7 @@ int32 ALandscape::PerformLayersHeightmapsGlobalMerge(const FUpdateLayersContentC
 				FLandscapeBrushParameters BrushParameters(ELandscapeToolTargetType::Heightmap, CombinedHeightmapNonAtlasRT);
 
 				UTextureRenderTarget2D* BrushOutputNonAtlasRT = Brush.RenderLayer(LandscapeExtent, BrushParameters);
-				if ((BrushOutputNonAtlasRT == nullptr)
+				if ((BrushOutputNonAtlasRT == nullptr) 
 					|| (BrushOutputNonAtlasRT->SizeX != CombinedHeightmapNonAtlasRT->SizeX)
 					|| (BrushOutputNonAtlasRT->SizeY != CombinedHeightmapNonAtlasRT->SizeY))
 				{
@@ -5348,7 +5347,7 @@ void ALandscape::OnDirtyWeightmap(FTextureToComponentHelper const& MapHelper, UT
 										++NumDifferentPixels;
 									}
 								}
-							}
+					}
 							StrBuilder.Appendf(TEXT("----------------------------------------\n"));
 							StrBuilder.Appendf(TEXT("Num diffs = %u\n"), NumDifferentPixels);
 							StrBuilder.Appendf(TEXT("Max diff = %u (%1.3f%%)\n"), MaxDiff, 100.0 * static_cast<float>(MaxDiff) / MAX_uint8);
@@ -6010,7 +6009,7 @@ void ALandscape::ReallocateLayersWeightmaps(FUpdateLayersContentContext& InUpdat
 
 	// Build a map of all the allocation per components
 	TMap<ULandscapeComponent*, TArray<ULandscapeLayerInfoObject*>> LayerAllocsPerComponent;
-	for (FLandscapeLayer& Layer : LandscapeLayers)
+	for (FLandscapeLayer& Layer : LandscapeEditLayers)
 	{
 		for (ULandscapeComponent* Component : InUpdateLayersContentContext.LandscapeComponentsWeightmapsToResolve)
 		{
@@ -6026,15 +6025,15 @@ void ALandscape::ReallocateLayersWeightmaps(FUpdateLayersContentContext& InUpdat
 			{
 				if (FLandscapeLayerComponentData* LayerComponentData = Component->GetLayerData(Layer.Guid))
 				{
-					for (const FWeightmapLayerAllocationInfo& LayerWeightmapAllocation : LayerComponentData->WeightmapData.LayerAllocations)
+				for (const FWeightmapLayerAllocationInfo& LayerWeightmapAllocation : LayerComponentData->WeightmapData.LayerAllocations)
+				{
+					if (LayerWeightmapAllocation.LayerInfo != nullptr)
 					{
-						if (LayerWeightmapAllocation.LayerInfo != nullptr)
-						{
-							ComponentLayerAlloc->AddUnique(LayerWeightmapAllocation.LayerInfo);
-						}
+						ComponentLayerAlloc->AddUnique(LayerWeightmapAllocation.LayerInfo);
 					}
 				}
 			}
+			}
 
 			// Add the brush alloc also (only if !InMergeParams.bSkipBrush, but InBrushRequiredAllocations should be empty already if InMergeParams.bSkipBrush is true) :
 			for (ULandscapeLayerInfoObject* BrushLayerInfo : InBrushRequiredAllocations)
@@ -6791,7 +6790,7 @@ void ALandscape::PrepareLayersWeightmapsLocalMergeRenderThreadData(const FUpdate
 
 			// Build, for each of the paint layers, a list of the weightmaps from each of the edit layers that affect the final result :
 			TMap<ULandscapeLayerInfoObject*, int32> PaintLayerToComponentPaintLayerRenderInfoIndex;
-			for (const FLandscapeLayer& Layer : LandscapeLayers)
+			for (const FLandscapeLayer& Layer : LandscapeEditLayers)
 			{
 				if (Layer.bVisible && !InMergeParams.bSkipBrush)
 				{
@@ -7118,7 +7117,7 @@ int32 ALandscape::PerformLayersWeightmapsGlobalMerge(FUpdateLayersContentContext
 		bool bHasWeightmapData = false;
 		bool bFirstLayer = true;
 
-		for (FLandscapeLayer& Layer : LandscapeLayers)
+		for (FLandscapeLayer& Layer : LandscapeEditLayers)
 		{
 			int8 CurrentWeightmapToProcessIndex = 0;
 			bool HasFoundWeightmapToProcess = true; // try processing at least once
@@ -7243,7 +7242,7 @@ int32 ALandscape::PerformLayersWeightmapsGlobalMerge(FUpdateLayersContentContext
 					int32 LayerIndex = LayerInfoObject.Value;
 					ULandscapeLayerInfoObject* LayerInfoObj = LayerInfoObject.Key;
 
-					SCOPED_DRAW_EVENTF_GAMETHREAD(LandscapeLayers, TEXT("LS Weight: %s PaintLayer: %s"), Layer.Name, LayerInfoObj->LayerName);
+					SCOPED_DRAW_EVENTF_GAMETHREAD(LandscapeEditLayers, TEXT("LS Weight: %s PaintLayer: %s"), Layer.Name, LayerInfoObj->LayerName);
 
 					// Copy the layer we are working on
 					SourceDebugName = FString::Printf(TEXT("Weight: %s PaintLayer: %s, CurrentProcLayerWeightmapAllLayersResource"), *Layer.Name.ToString(), *LayerInfoObj->LayerName.ToString());
@@ -7894,7 +7893,7 @@ void ALandscape::ResolveLayersWeightmapTexture(
 
 bool ALandscape::HasLayersContent() const
 {
-	return LandscapeLayers.Num() > 0;
+	return LandscapeEditLayers.Num() > 0;
 }
 
 void ALandscape::UpdateCachedHasLayersContent(bool bInCheckComponentDataIntegrity)
@@ -7925,7 +7924,7 @@ void ALandscape::RequestLayersInitialization(bool bInRequestContentUpdate)
 
 void ALandscape::RequestSplineLayerUpdate()
 {
-	if (HasLayersContent() && GetLandscapeSplinesReservedLayer() != nullptr)
+	if (HasLayersContent() && FindLayerOfType(ULandscapeEditLayerSplines::StaticClass()) != nullptr)
 	{
 		bSplineLayerUpdateRequested = true;
 	}
@@ -8306,7 +8305,7 @@ void ALandscape::UpdateLayersContent(bool bInWaitForStreaming, bool bInSkipMonit
 	MonitorShaderCompilation();
 
 	// Make sure Brush get a chance to request an update of the landscape
-	for (FLandscapeLayer& Layer : LandscapeLayers)
+	for (FLandscapeLayer& Layer : LandscapeEditLayers)
 	{
 		for (FLandscapeLayerBrush& Brush : Layer.Brushes)
 		{
@@ -8322,7 +8321,7 @@ void ALandscape::UpdateLayersContent(bool bInWaitForStreaming, bool bInSkipMonit
 
 	if (bSplineLayerUpdateRequested)
 	{
-		if (FLandscapeLayer* SplinesLayer = GetLandscapeSplinesReservedLayer())
+		if (const FLandscapeLayer* SplinesLayer = FindLayerOfType(ULandscapeEditLayerSplines::StaticClass()))
 		{
 			// We need the spline layer resources to all be ready before updating it:
 			if (!PrepareLayersTextureResources({ *SplinesLayer }, bInWaitForStreaming))
@@ -8655,9 +8654,9 @@ void ALandscape::InitializeLayers(bool bUseNormalCapture)
 
 	if (CreateLayersRenderingResource(bUseNormalCapture))
 	{
-		InitializeLandscapeLayersWeightmapUsage();
-		bLandscapeLayersAreInitialized = true;
-	}
+	InitializeLandscapeLayersWeightmapUsage();
+	bLandscapeLayersAreInitialized = true;
+}
 }
 
 void ALandscape::OnPreSave()
@@ -8934,7 +8933,7 @@ bool ALandscapeProxy::AddLayer(const FGuid& InLayerGuid)
 	{
 		if ((Component != nullptr) && !Component->GetLayerData(InLayerGuid))
 		{
-			const FLandscapeLayer* EditLayer = GetLandscapeActor() ? GetLandscapeActor()->GetLayer(InLayerGuid) : nullptr;
+			const FLandscapeLayer* EditLayer = GetLandscapeActor() ? GetLandscapeActor()->GetLayerConst(InLayerGuid) : nullptr;
 			Component->AddLayerData(InLayerGuid, FLandscapeLayerComponentData(EditLayer ? EditLayer->Name : FName()));
 			bModified = true;
 		}
@@ -9118,12 +9117,12 @@ bool ALandscape::IsUpToDate() const
 #if WITH_EDITOR
 bool ALandscape::IsLayerNameUnique(const FName& InName) const
 {
-	return Algo::CountIf(LandscapeLayers, [InName](const FLandscapeLayer& Layer) { return (Layer.Name == InName); }) == 0;
+	return Algo::CountIf(LandscapeEditLayers, [InName](const FLandscapeLayer& Layer) { return (Layer.Name == InName); }) == 0;
 }
 
 void ALandscape::SetLayerName(int32 InLayerIndex, const FName& InName)
 {
-	FLandscapeLayer* Layer = GetLayer(InLayerIndex);
+	FLandscapeLayer* Layer = GetLayerInternal(InLayerIndex);
 	ULandscapeInfo* LandscapeInfo = GetLandscapeInfo();
 	if (!LandscapeInfo || !Layer || Layer->Name == InName)
 	{
@@ -9136,13 +9135,14 @@ void ALandscape::SetLayerName(int32 InLayerIndex, const FName& InName)
 	}
 
 	Modify();
-	LandscapeLayers[InLayerIndex].Name = InName;
+	Layer->Name = InName;
 }
 
 float ALandscape::GetLayerAlpha(int32 InLayerIndex, bool bInHeightmap) const
 {
-	const FLandscapeLayer* SplinesReservedLayer = GetLandscapeSplinesReservedLayer();
-	const FLandscapeLayer* Layer = GetLayer(InLayerIndex);
+	// TODO [jonathan.bard] : replace by proper getter virtual GetLayerAlpha implementation on ULandscapeEditLayerSplines
+	const FLandscapeLayer* SplinesReservedLayer = FindLayerOfType(ULandscapeEditLayerSplines::StaticClass());
+	const FLandscapeLayer* Layer = GetLayerConst(InLayerIndex);
 	if (Layer && SplinesReservedLayer != Layer)
 	{
 		return GetClampedLayerAlpha(bInHeightmap ? Layer->HeightmapAlpha : Layer->WeightmapAlpha, bInHeightmap);
@@ -9158,7 +9158,7 @@ float ALandscape::GetClampedLayerAlpha(float InAlpha, bool bInHeightmap) const
 
 void ALandscape::SetLayerAlpha(int32 InLayerIndex, float InAlpha, bool bInHeightmap)
 {
-	FLandscapeLayer* Layer = GetLayer(InLayerIndex);
+	FLandscapeLayer* Layer = GetLayerInternal(InLayerIndex);
 	ULandscapeInfo* LandscapeInfo = GetLandscapeInfo();
 	if (!LandscapeInfo || !Layer)
 	{
@@ -9176,23 +9176,31 @@ void ALandscape::SetLayerAlpha(int32 InLayerIndex, float InAlpha, bool bInHeight
 	RequestLayersContentUpdateForceAll();
 }
 
-void ALandscape::SetLayerVisibility(int32 InLayerIndex, bool bInVisible)
+void ALandscape::SetLayerVisibility(int32 InLayerIndex, bool bInVisible, bool bInForIntermediateRender)
 {
-	FLandscapeLayer* Layer = GetLayer(InLayerIndex);
+	FLandscapeLayer* Layer = GetLayerInternal(InLayerIndex);
 	ULandscapeInfo* LandscapeInfo = GetLandscapeInfo();
 	if (!LandscapeInfo || !Layer || Layer->bVisible == bInVisible)
 	{
 		return;
 	}
 
-	Modify();
+	if (!bInForIntermediateRender)
+	{
+		Modify();
+	}
+
 	Layer->bVisible = bInVisible;
-	RequestLayersContentUpdateForceAll();
+
+	if (!bInForIntermediateRender)
+	{
+		RequestLayersContentUpdateForceAll();
+	}
 }
 
 void ALandscape::SetLayerLocked(int32 InLayerIndex, bool bLocked)
 {
-	FLandscapeLayer* Layer = GetLayer(InLayerIndex);
+	FLandscapeLayer* Layer = GetLayerInternal(InLayerIndex);
 	if (!Layer || Layer->bLocked == bLocked)
 	{
 		return;
@@ -9202,55 +9210,112 @@ void ALandscape::SetLayerLocked(int32 InLayerIndex, bool bLocked)
 	Layer->bLocked = bLocked;
 }
 
+void ALandscape::SetLayerBlendMode(int32 InLayerIndex, ELandscapeBlendMode InBlendMode)
+{
+	FLandscapeLayer* Layer = GetLayerInternal(InLayerIndex);
+	if (!Layer || Layer->BlendMode == InBlendMode)
+	{
+		return;
+	}
+
+	Modify();
+	Layer->BlendMode = InBlendMode;
+}
+
 uint8 ALandscape::GetLayerCount() const
 {
-	return static_cast<uint8>(LandscapeLayers.Num());
+	return static_cast<uint8>(LandscapeEditLayers.Num());
 }
 
-FLandscapeLayer* ALandscape::GetLayer(int32 InLayerIndex)
+FLandscapeLayer* ALandscape::GetLayerInternal(int32 InLayerIndex)
 {
-	if (LandscapeLayers.IsValidIndex(InLayerIndex))
+	if (LandscapeEditLayers.IsValidIndex(InLayerIndex))
 	{
-		return &LandscapeLayers[InLayerIndex];
+		return &LandscapeEditLayers[InLayerIndex];
 	}
 	return nullptr;
 }
 
+// Deprecated
 const FLandscapeLayer* ALandscape::GetLayer(int32 InLayerIndex) const
 {
-	if (LandscapeLayers.IsValidIndex(InLayerIndex))
+	return GetLayerConst(InLayerIndex);
+}
+
+// Deprecated
+const FLandscapeLayer* ALandscape::GetLayer(const FGuid& InLayerGuid) const
+{
+	return GetLayerConst(InLayerGuid);
+}
+
+// Deprecated
+const FLandscapeLayer* ALandscape::GetLayer(const FName& InLayerName) const
+{
+	return GetLayerConst(InLayerName);
+}
+
+const FLandscapeLayer* ALandscape::GetLayerConst(int32 InLayerIndex) const
+{
+	if (LandscapeEditLayers.IsValidIndex(InLayerIndex))
 	{
-		return &LandscapeLayers[InLayerIndex];
+		return &LandscapeEditLayers[InLayerIndex];
 	}
 	return nullptr;
 }
 
-int32 ALandscape::GetLayerIndex(FName InLayerName) const
+const FLandscapeLayer* ALandscape::GetLayerConst(const FGuid& InLayerGuid) const
 {
-	return LandscapeLayers.IndexOfByPredicate([InLayerName](const FLandscapeLayer& Layer) { return Layer.Name == InLayerName; });
+	return LandscapeEditLayers.FindByPredicate([&InLayerGuid](const FLandscapeLayer& Other) { return Other.Guid == InLayerGuid; });
 }
 
-const FLandscapeLayer* ALandscape::GetLayer(const FGuid& InLayerGuid) const
+const FLandscapeLayer* ALandscape::GetLayerConst(const FName& InLayerName) const
 {
-	return LandscapeLayers.FindByPredicate([&InLayerGuid](const FLandscapeLayer& Other) { return Other.Guid == InLayerGuid; });
+	return LandscapeEditLayers.FindByPredicate([InLayerName](const FLandscapeLayer& Layer) { return Layer.Name == InLayerName; });
 }
 
-const FLandscapeLayer* ALandscape::GetLayer(const FName& InLayerName) const
+int32 ALandscape::GetLayerIndex(FName InLayerName) const
 {
-	return LandscapeLayers.FindByPredicate([InLayerName](const FLandscapeLayer& Layer) { return Layer.Name == InLayerName; });
+	return LandscapeEditLayers.IndexOfByPredicate([InLayerName](const FLandscapeLayer& Layer) { return Layer.Name == InLayerName; });
 }
 
+// Deprecated
 void ALandscape::ForEachLayer(TFunctionRef<void(struct FLandscapeLayer&)> Fn)
 {
-	for (FLandscapeLayer& Layer : LandscapeLayers)
+	for (FLandscapeLayer& Layer : LandscapeEditLayers)
 	{
 		Fn(Layer);
 	}
 }
 
+void ALandscape::ForEachLayerConst(TFunctionRef<bool(const FLandscapeLayer&)> Fn)
+{
+	for (FLandscapeLayer& Layer : LandscapeEditLayers)
+	{
+		if (!Fn(Layer))
+		{
+			return;
+		}
+	}
+}
+
+const FLandscapeLayer* ALandscape::FindLayerOfType(const TSubclassOf<ULandscapeEditLayerBase>& InLayerClass) const
+{
+	return LandscapeEditLayers.FindByPredicate([&InLayerClass](const FLandscapeLayer& InLayer) { check(InLayer.EditLayer != nullptr); return InLayer.EditLayer->GetClass()->IsChildOf(InLayerClass); });
+}
+
+TArray<const FLandscapeLayer*> ALandscape::GetLayersOfType(const TSubclassOf<ULandscapeEditLayerBase>& InLayerClass) const
+{
+	TArray<const FLandscapeLayer*> Result;
+	Result.Reserve(LandscapeEditLayers.Num());
+	Algo::TransformIf(LandscapeEditLayers, Result,
+		[&InLayerClass](const FLandscapeLayer& InLayer) { check(InLayer.EditLayer != nullptr); return InLayer.EditLayer->GetClass()->IsChildOf(InLayerClass); },
+		[](const FLandscapeLayer& InLayer) { return &InLayer; });
+	return Result;
+}
+
 void ALandscape::DeleteLayers()
 {
-	for (int32 LayerIndex = LandscapeLayers.Num() - 1; LayerIndex >= 0; --LayerIndex)
+	for (int32 LayerIndex = LandscapeEditLayers.Num() - 1; LayerIndex >= 0; --LayerIndex)
 	{
 		DeleteLayer(LayerIndex);
 	}
@@ -9261,7 +9326,7 @@ void ALandscape::DeleteLayer(int32 InLayerIndex)
 	ensure(HasLayersContent());
 
 	ULandscapeInfo* LandscapeInfo = GetLandscapeInfo();
-	const FLandscapeLayer* Layer = GetLayer(InLayerIndex);
+	const FLandscapeLayer* Layer = GetLayerConst(InLayerIndex);
 	if (!LandscapeInfo || !Layer)
 	{
 		return;
@@ -9277,14 +9342,8 @@ void ALandscape::DeleteLayer(int32 InLayerIndex)
 		return true;
 	});
 
-	const FLandscapeLayer* SplinesReservedLayer = GetLandscapeSplinesReservedLayer();
-	if (SplinesReservedLayer == Layer)
-	{
-		LandscapeSplinesTargetLayerGuid.Invalidate();
-	}
-
 	// Remove layer from list
-	LandscapeLayers.RemoveAt(InLayerIndex);
+	LandscapeEditLayers.RemoveAt(InLayerIndex);
 
 	// Request Update
 	RequestLayersContentUpdateForceAll();
@@ -9296,16 +9355,16 @@ void ALandscape::CollapseLayer(int32 InLayerIndex)
 	SlowTask.MakeDialog();
 	TArray<bool> BackupVisibility;
 	TArray<bool> BackupBrushVisibility;
-	for (int32 i = 0; i < LandscapeLayers.Num(); ++i)
+	for (int32 i = 0; i < LandscapeEditLayers.Num(); ++i)
 	{
-		BackupVisibility.Add(LandscapeLayers[i].bVisible);
-		LandscapeLayers[i].bVisible = i == InLayerIndex || i == InLayerIndex - 1;
+		BackupVisibility.Add(LandscapeEditLayers[i].bVisible);
+		LandscapeEditLayers[i].bVisible = i == InLayerIndex || i == InLayerIndex - 1;
 	}
 
-	for (int32 i = 0; i < LandscapeLayers[InLayerIndex].Brushes.Num(); ++i)
+	for (int32 i = 0; i < LandscapeEditLayers[InLayerIndex].Brushes.Num(); ++i)
 	{
-		BackupBrushVisibility.Add(LandscapeLayers[InLayerIndex].Brushes[i].GetBrush()->IsVisible());
-		LandscapeLayers[InLayerIndex].Brushes[i].GetBrush()->SetIsVisible(false);
+		BackupBrushVisibility.Add(LandscapeEditLayers[InLayerIndex].Brushes[i].GetBrush()->IsVisible());
+		LandscapeEditLayers[InLayerIndex].Brushes[i].GetBrush()->SetIsVisible(false);
 	}
 
 	// Call Request Update on all components...
@@ -9324,7 +9383,7 @@ void ALandscape::CollapseLayer(int32 InLayerIndex)
 		DataInterface.SetShouldDirtyPackage(true);
 
 		TSet<UTexture2D*> ProcessedHeightmaps;
-		FScopedSetLandscapeEditingLayer ScopeEditingLayer(this, LandscapeLayers[InLayerIndex - 1].Guid);
+		FScopedSetLandscapeEditingLayer ScopeEditingLayer(this, LandscapeEditLayers[InLayerIndex - 1].Guid);
 		GetLandscapeInfo()->ForAllLandscapeComponents([&](ULandscapeComponent* LandscapeComponent)
 		{
 			SlowTask.EnterProgressFrame(1.f);
@@ -9333,9 +9392,9 @@ void ALandscape::CollapseLayer(int32 InLayerIndex)
 	}
 
 	TArray<ALandscapeBlueprintBrushBase*> BrushesToMove;
-	for (int32 i = 0; i < LandscapeLayers[InLayerIndex].Brushes.Num(); ++i)
+	for (int32 i = 0; i < LandscapeEditLayers[InLayerIndex].Brushes.Num(); ++i)
 	{
-		ALandscapeBlueprintBrushBase* CurrentBrush = LandscapeLayers[InLayerIndex].Brushes[i].GetBrush();
+		ALandscapeBlueprintBrushBase* CurrentBrush = LandscapeEditLayers[InLayerIndex].Brushes[i].GetBrush();
 		CurrentBrush->SetIsVisible(BackupBrushVisibility[i]);
 		BrushesToMove.Add(CurrentBrush);
 	}
@@ -9346,9 +9405,9 @@ void ALandscape::CollapseLayer(int32 InLayerIndex)
 		AddBrushToLayer(InLayerIndex - 1, Brush);
 	}
 
-	for (int32 i = 0; i < LandscapeLayers.Num(); ++i)
+	for (int32 i = 0; i < LandscapeEditLayers.Num(); ++i)
 	{
-		LandscapeLayers[i].bVisible = BackupVisibility[i];
+		LandscapeEditLayers[i].bVisible = BackupVisibility[i];
 	}
 
 	DeleteLayer(InLayerIndex);
@@ -9358,7 +9417,7 @@ void ALandscape::CollapseLayer(int32 InLayerIndex)
 
 void ALandscape::GetUsedPaintLayers(int32 InLayerIndex, TArray<ULandscapeLayerInfoObject*>& OutUsedLayerInfos) const
 {
-	const FLandscapeLayer* Layer = GetLayer(InLayerIndex);
+	const FLandscapeLayer* Layer = GetLayerConst(InLayerIndex);
 	if (Layer)
 	{
 		GetUsedPaintLayers(Layer->Guid, OutUsedLayerInfos);
@@ -9378,7 +9437,7 @@ void ALandscape::GetUsedPaintLayers(const FGuid& InLayerGuid, TArray<ULandscapeL
 
 void ALandscape::ClearPaintLayer(int32 InLayerIndex, ULandscapeLayerInfoObject* InLayerInfo)
 {
-	const FLandscapeLayer* Layer = GetLayer(InLayerIndex);
+	const FLandscapeLayer* Layer = GetLayerConst(InLayerIndex);
 	if (Layer)
 	{
 		ClearPaintLayer(Layer->Guid, InLayerInfo);
@@ -9410,7 +9469,7 @@ void ALandscape::ClearPaintLayer(const FGuid& InLayerGuid, ULandscapeLayerInfoOb
 
 void ALandscape::ClearLayer(int32 InLayerIndex, TSet<TObjectPtr<ULandscapeComponent>>* InComponents, ELandscapeClearMode InClearMode)
 {
-	const FLandscapeLayer* Layer = GetLayer(InLayerIndex);
+	const FLandscapeLayer* Layer = GetLayerConst(InLayerIndex);
 	if (Layer)
 	{
 		ClearLayer(Layer->Guid, InComponents, InClearMode);
@@ -9421,7 +9480,7 @@ void ALandscape::ClearLayer(const FGuid& InLayerGuid, TSet<TObjectPtr<ULandscape
 {
 	ensure(HasLayersContent());
 
-	const FLandscapeLayer* Layer = GetLayer(InLayerGuid);
+	const FLandscapeLayer* Layer = GetLayerConst(InLayerGuid);
 	ULandscapeInfo* LandscapeInfo = GetLandscapeInfo();
 	if (!LandscapeInfo || !Layer)
 	{
@@ -9509,11 +9568,11 @@ void ALandscape::ClearLayer(const FGuid& InLayerGuid, TSet<TObjectPtr<ULandscape
 
 void ALandscape::ShowOnlySelectedLayer(int32 InLayerIndex)
 {
-	const FLandscapeLayer* VisibleLayer = GetLayer(InLayerIndex);
+	const FLandscapeLayer* VisibleLayer = GetLayerConst(InLayerIndex);
 	if (VisibleLayer)
 	{
 		bool bModified = false;
-		for (FLandscapeLayer& Layer : LandscapeLayers)
+		for (FLandscapeLayer& Layer : LandscapeEditLayers)
 		{
 			bool bDesiredVisible = (&Layer == VisibleLayer);
 			if (Layer.bVisible != bDesiredVisible)
@@ -9522,23 +9581,23 @@ void ALandscape::ShowOnlySelectedLayer(int32 InLayerIndex)
 				{
 					Modify();
 					bModified = true;
-				}
+		}
 				Layer.bVisible = bDesiredVisible;
 			}
 		}
 		if (bModified)
 		{
-			RequestLayersContentUpdateForceAll();
-		}
+		RequestLayersContentUpdateForceAll();
 	}
 }
+}
 
 void ALandscape::ShowAllLayers()
 {
-	if (LandscapeLayers.Num() > 0)
+	if (LandscapeEditLayers.Num() > 0)
 	{
 		bool bModified = false;
-		for (FLandscapeLayer& Layer : LandscapeLayers)
+		for (FLandscapeLayer& Layer : LandscapeEditLayers)
 		{
 			if (Layer.bVisible != true)
 			{
@@ -9547,60 +9606,31 @@ void ALandscape::ShowAllLayers()
 					Modify();
 					bModified = true;
 				}
-				Layer.bVisible = true;
-			}
+			Layer.bVisible = true;
+		}
 		}
 		if (bModified)
 		{
-			RequestLayersContentUpdateForceAll();
-		}
+		RequestLayersContentUpdateForceAll();
 	}
 }
+}
 
+// Deprecated
 void ALandscape::SetLandscapeSplinesReservedLayer(int32 InLayerIndex)
 {
-	Modify();
-	FLandscapeLayer* NewLayer = GetLayer(InLayerIndex);
-	FLandscapeLayer* PreviousLayer = GetLandscapeSplinesReservedLayer();
-	if (NewLayer != PreviousLayer)
-	{
-		LandscapeSplinesAffectedComponents.Empty();
-		if (PreviousLayer)
-		{
-			ClearLayer(LandscapeSplinesTargetLayerGuid);
-			PreviousLayer->BlendMode = LSBM_AdditiveBlend;
-		}
-		if (NewLayer)
-		{
-			NewLayer->HeightmapAlpha = 1.0f;
-			NewLayer->WeightmapAlpha = 1.0f;
-			NewLayer->BlendMode = LSBM_AlphaBlend;
-			LandscapeSplinesTargetLayerGuid = NewLayer->Guid;
-			ClearLayer(LandscapeSplinesTargetLayerGuid);
-		}
-		else
-		{
-			LandscapeSplinesTargetLayerGuid.Invalidate();
-		}
-	}
 }
 
+// Deprecated
 const FLandscapeLayer* ALandscape::GetLandscapeSplinesReservedLayer() const
 {
-	if (LandscapeSplinesTargetLayerGuid.IsValid())
-	{
-		return LandscapeLayers.FindByPredicate([this](const FLandscapeLayer& Other) { return Other.Guid == LandscapeSplinesTargetLayerGuid; });
-	}
-	return nullptr;
+	return FindLayerOfType(ULandscapeEditLayerSplines::StaticClass());
 }
 
+// Deprecated
 FLandscapeLayer* ALandscape::GetLandscapeSplinesReservedLayer()
 {
-	if (LandscapeSplinesTargetLayerGuid.IsValid())
-	{
-		return LandscapeLayers.FindByPredicate([this](const FLandscapeLayer& Other) { return Other.Guid == LandscapeSplinesTargetLayerGuid; });
-	}
-	return nullptr;
+	return const_cast<FLandscapeLayer*>(FindLayerOfType(ULandscapeEditLayerSplines::StaticClass()));
 }
 
 LANDSCAPE_API extern bool GDisableUpdateLandscapeMaterialInstances;
@@ -9658,15 +9688,16 @@ void ALandscape::UpdateLandscapeSplines(const FGuid& InTargetLayer, bool bInUpda
 	TRACE_CPUPROFILER_EVENT_SCOPE(LandscapeLayers_UpdateLandscapeSplines);
 	check(CanHaveLayersContent());
 	ULandscapeInfo* LandscapeInfo = GetLandscapeInfo();
-	FGuid TargetLayerGuid = LandscapeSplinesTargetLayerGuid.IsValid() ? LandscapeSplinesTargetLayerGuid : InTargetLayer;
-	const FLandscapeLayer* TargetLayer = GetLayer(TargetLayerGuid);
+	const FLandscapeLayer* SplinesLayer = FindLayerOfType(ULandscapeEditLayerSplines::StaticClass());
+	FGuid TargetLayerGuid = (SplinesLayer != nullptr) ? SplinesLayer->Guid : InTargetLayer;
+	const FLandscapeLayer* TargetLayer = GetLayerConst(TargetLayerGuid);
 	if (LandscapeInfo && TargetLayer)
 	{
 		FScopedSetLandscapeEditingLayer Scope(this, TargetLayerGuid, [this] { this->RequestLayersContentUpdate(ELandscapeLayerUpdateMode::Update_All); });
 		// Temporarily disable material instance updates since it will be done once at the end (requested by RequestLayersContentUpdateForceAll)
 		GDisableUpdateLandscapeMaterialInstances = true;
 		TSet<TObjectPtr<ULandscapeComponent>>* ModifiedComponent = nullptr;
-		if (LandscapeSplinesTargetLayerGuid.IsValid())
+		if (SplinesLayer != nullptr)
 		{
 			// Check that we can modify data
 			if (!LandscapeInfo->AreAllComponentsRegistered())
@@ -9695,7 +9726,7 @@ void ALandscape::UpdateLandscapeSplines(const FGuid& InTargetLayer, bool bInUpda
 
 			// Clear layers without affecting weightmap allocations
 			const bool bMarkPackageDirty = false;
-			ClearLayer(LandscapeSplinesTargetLayerGuid, (!bInForceUpdateAllCompoments && LandscapeSplinesAffectedComponents.Num()) ? &LandscapeSplinesAffectedComponents : nullptr, Clear_All, bMarkPackageDirty);
+			ClearLayer(SplinesLayer->Guid, (!bInForceUpdateAllCompoments && LandscapeSplinesAffectedComponents.Num()) ? &LandscapeSplinesAffectedComponents : nullptr, Clear_All, bMarkPackageDirty);
 			LandscapeSplinesAffectedComponents.Empty();
 			ModifiedComponent = &LandscapeSplinesAffectedComponents;
 			// For now, in Landscape Layer System Mode with a reserved layer for splines, we always update all the splines since we clear the whole layer first
@@ -9752,13 +9783,9 @@ FScopedSetLandscapeEditingLayer::~FScopedSetLandscapeEditingLayer()
 	}
 }
 
+// Deprecated
 bool ALandscape::IsEditingLayerReservedForSplines() const
 {
-	if (CanHaveLayersContent())
-	{
-		const FLandscapeLayer* SplinesReservedLayer = GetLandscapeSplinesReservedLayer();
-		return SplinesReservedLayer && SplinesReservedLayer->Guid == GetEditingLayer();
-	}
 	return false;
 }
 
@@ -9790,7 +9817,7 @@ const FGuid& ALandscape::GetEditingLayer() const
 
 bool ALandscape::IsMaxLayersReached() const
 {
-	return LandscapeLayers.Num() >= GetDefault<ULandscapeSettings>()->MaxNumberOfLayers;
+	return LandscapeEditLayers.Num() >= GetDefault<ULandscapeSettings>()->MaxNumberOfLayers;
 }
 
 void ALandscape::CreateDefaultLayer()
@@ -9801,7 +9828,7 @@ void ALandscape::CreateDefaultLayer()
 		return;
 	}
 
-	check(LandscapeLayers.Num() == 0); // We can only call this function if we have no layers
+	check(LandscapeEditLayers.Num() == 0); // We can only call this function if we have no layers
 
 	CreateLayer(FName(TEXT("Layer")));
 	// Force update rendering resources
@@ -9821,13 +9848,17 @@ FLandscapeLayer* ALandscape::DuplicateLayerAndMoveBrushes(const FLandscapeLayer&
 	FLandscapeLayer NewLayer(InOtherLayer);
 	NewLayer.Guid = FGuid::NewGuid();
 
+	// Duplicate the internal edit layer object by hand : 
+	check(InOtherLayer.EditLayer != nullptr);
+	NewLayer.EditLayer = DuplicateObject(InOtherLayer.EditLayer, this, MakeUniqueObjectName(this, InOtherLayer.EditLayer->GetClass(), InOtherLayer.EditLayer->GetFName()));
+
 	// Update owning landscape and reparent to landscape's level if necessary
 	for (FLandscapeLayerBrush& Brush : NewLayer.Brushes)
 	{
 		Brush.SetOwner(this);
 	}
 
-	int32 AddedIndex = LandscapeLayers.Add(NewLayer);
+	int32 AddedIndex = LandscapeEditLayers.Add(NewLayer);
 
 	// Create associated layer data in each landscape proxy
 	LandscapeInfo->ForEachLandscapeProxy([&NewLayer](ALandscapeProxy* Proxy)
@@ -9836,10 +9867,10 @@ FLandscapeLayer* ALandscape::DuplicateLayerAndMoveBrushes(const FLandscapeLayer&
 		return true;
 	});
 
-	return &LandscapeLayers[AddedIndex];
+	return &LandscapeEditLayers[AddedIndex];
 }
 
-int32 ALandscape::CreateLayer(FName InName)
+int32 ALandscape::CreateLayer(FName InName, const TSubclassOf<ULandscapeEditLayerBase>& InEditLayerClass)
 {
 	ULandscapeInfo* LandscapeInfo = GetLandscapeInfo();
 	if (!LandscapeInfo || IsMaxLayersReached() || !CanHaveLayersContent())
@@ -9848,18 +9879,56 @@ int32 ALandscape::CreateLayer(FName InName)
 	}
 
 	Modify();
-	FLandscapeLayer NewLayer;
+
+	const UClass* EditLayerClass = (InEditLayerClass.Get() != nullptr) ? InEditLayerClass.Get() : ULandscapeEditLayer::StaticClass();
+	int32 LayerIndex = LandscapeEditLayers.Emplace();
+	FLandscapeLayer& NewLayer = LandscapeEditLayers[LayerIndex];
 	NewLayer.Name = GenerateUniqueLayerName(InName);
-	int32 LayerIndex = LandscapeLayers.Add(NewLayer);
+	NewLayer.EditLayer = NewObject<ULandscapeEditLayerBase>(this, EditLayerClass, MakeUniqueObjectName(this, EditLayerClass));
+	OnLayerCreatedInternal(NewLayer);
+	return LayerIndex;
+}
+
+int32 ALandscape::CreateLayerFrom(const FLandscapeLayer& InLayer)
+{
+	ULandscapeInfo* LandscapeInfo = GetLandscapeInfo();
+	if (!LandscapeInfo || IsMaxLayersReached() || !CanHaveLayersContent())
+	{
+		return INDEX_NONE;
+	}
+
+	// Cannot create a layer whose Guid is already taken : 
+	if (GetLayerConst(InLayer.Guid) != nullptr)
+	{
+		return INDEX_NONE;
+	}
+
+	check(InLayer.EditLayer != nullptr);
+
+	int32 LayerIndex = LandscapeEditLayers.Add(InLayer);
+	FLandscapeLayer& NewLayer = LandscapeEditLayers[LayerIndex];
+	NewLayer.EditLayer = DuplicateObject<ULandscapeEditLayerBase>(InLayer.EditLayer, this, MakeUniqueObjectName(this, InLayer.EditLayer->GetClass()));
+	OnLayerCreatedInternal(NewLayer);
+	return LayerIndex;
+}
+
+void ALandscape::OnLayerCreatedInternal(FLandscapeLayer& Layer)
+{
+	ULandscapeInfo* LandscapeInfo = GetLandscapeInfo();
+	check(LandscapeInfo != nullptr); 
+
+	Layer.EditLayer->OnLayerCreated(Layer);
 
 	// Create associated layer data in each landscape proxy
-	LandscapeInfo->ForEachLandscapeProxy([&NewLayer](ALandscapeProxy* Proxy)
+	LandscapeInfo->ForEachLandscapeProxy([&Layer](ALandscapeProxy* Proxy)
 	{
-		Proxy->AddLayer(NewLayer.Guid);
+		Proxy->AddLayer(Layer.Guid);
 		return true;
 	});
 
-	return LayerIndex;
+	// Request Update
+	RequestSplineLayerUpdate(); // Request a spline update as well, in case we already had spline actors and a spline layers was just created
+	RequestLayersContentUpdateForceAll();
 }
 
 void ALandscape::AddLayersToProxy(ALandscapeProxy* InProxy)
@@ -9873,9 +9942,10 @@ void ALandscape::AddLayersToProxy(ALandscapeProxy* InProxy)
 	check(InProxy != this);
 	check(InProxy != nullptr);
 
-	ForEachLayer([&](FLandscapeLayer& Layer)
+	ForEachLayerConst([&](const FLandscapeLayer& Layer)
 	{
 		InProxy->AddLayer(Layer.Guid);
+		return true;
 	});
 
 	// Force update rendering resources
@@ -9885,13 +9955,13 @@ void ALandscape::AddLayersToProxy(ALandscapeProxy* InProxy)
 bool ALandscape::ReorderLayer(int32 InStartingLayerIndex, int32 InDestinationLayerIndex)
 {
 	if (InStartingLayerIndex != InDestinationLayerIndex &&
-		LandscapeLayers.IsValidIndex(InStartingLayerIndex) &&
-		LandscapeLayers.IsValidIndex(InDestinationLayerIndex))
+		LandscapeEditLayers.IsValidIndex(InStartingLayerIndex) &&
+		LandscapeEditLayers.IsValidIndex(InDestinationLayerIndex))
 	{
 		Modify();
-		FLandscapeLayer Layer = LandscapeLayers[InStartingLayerIndex];
-		LandscapeLayers.RemoveAt(InStartingLayerIndex);
-		LandscapeLayers.Insert(Layer, InDestinationLayerIndex);
+		FLandscapeLayer Layer = LandscapeEditLayers[InStartingLayerIndex];
+		LandscapeEditLayers.RemoveAt(InStartingLayerIndex);
+		LandscapeEditLayers.Insert(Layer, InDestinationLayerIndex);
 		RequestLayersContentUpdateForceAll();
 		return true;
 	}
@@ -9901,7 +9971,7 @@ bool ALandscape::ReorderLayer(int32 InStartingLayerIndex, int32 InDestinationLay
 FName ALandscape::GenerateUniqueLayerName(FName InName) const
 {
 	// If we are receiving a unique name, use it.
-	if (InName != NAME_None && !LandscapeLayers.ContainsByPredicate([InName](const FLandscapeLayer& Layer) { return Layer.Name == InName; }))
+	if (InName != NAME_None && !LandscapeEditLayers.ContainsByPredicate([InName](const FLandscapeLayer& Layer) { return Layer.Name == InName; }))
 	{
 		return InName;
 	}
@@ -9913,14 +9983,14 @@ FName ALandscape::GenerateUniqueLayerName(FName InName) const
 	{
 		++LayerIndex;
 		NewName = FName(*FString::Printf(TEXT("%s%d"), *BaseName, LayerIndex));
-	} while (LandscapeLayers.ContainsByPredicate([NewName](const FLandscapeLayer& Layer) { return Layer.Name == NewName; }));
+	} while (LandscapeEditLayers.ContainsByPredicate([NewName](const FLandscapeLayer& Layer) { return Layer.Name == NewName; }));
 
 	return NewName;
 }
 
 bool ALandscape::IsLayerBlendSubstractive(int32 InLayerIndex, const TWeakObjectPtr<ULandscapeLayerInfoObject>& InLayerInfoObj) const
 {
-	const FLandscapeLayer* Layer = GetLayer(InLayerIndex);
+	const FLandscapeLayer* Layer = GetLayerConst(InLayerIndex);
 
 	if (Layer == nullptr)
 	{
@@ -9939,7 +10009,7 @@ bool ALandscape::IsLayerBlendSubstractive(int32 InLayerIndex, const TWeakObjectP
 
 void ALandscape::SetLayerSubstractiveBlendStatus(int32 InLayerIndex, bool InStatus, const TWeakObjectPtr<ULandscapeLayerInfoObject>& InLayerInfoObj)
 {
-	FLandscapeLayer* Layer = GetLayer(InLayerIndex);
+	FLandscapeLayer* Layer = GetLayerInternal(InLayerIndex);
 
 	if (Layer == nullptr)
 	{
@@ -9961,9 +10031,23 @@ void ALandscape::SetLayerSubstractiveBlendStatus(int32 InLayerIndex, bool InStat
 	RequestLayersContentUpdateForceAll(ELandscapeLayerUpdateMode::Update_Weightmap_All);
 }
 
+void ALandscape::ReplaceLayerSubstractiveBlendStatus(ULandscapeLayerInfoObject* InFromLayerInfo, ULandscapeLayerInfoObject* InToLayerInfo, bool bInShouldDirtyPackage)
+{
+	Modify(bInShouldDirtyPackage);
+
+	for (FLandscapeLayer& Layer : LandscapeEditLayers)
+	{
+		bool OutValue;
+		if (Layer.WeightmapLayerAllocationBlend.RemoveAndCopyValue(InFromLayerInfo, OutValue))
+		{
+			Layer.WeightmapLayerAllocationBlend.Add(InToLayerInfo, OutValue);
+		}
+	}
+}
+
 bool ALandscape::ReorderLayerBrush(int32 InLayerIndex, int32 InStartingLayerBrushIndex, int32 InDestinationLayerBrushIndex)
 {
-	if (FLandscapeLayer* Layer = GetLayer(InLayerIndex))
+	if (FLandscapeLayer* Layer = GetLayerInternal(InLayerIndex))
 	{
 		if (InStartingLayerBrushIndex != InDestinationLayerBrushIndex &&
 			Layer->Brushes.IsValidIndex(InStartingLayerBrushIndex) &&
@@ -9982,9 +10066,9 @@ bool ALandscape::ReorderLayerBrush(int32 InLayerIndex, int32 InStartingLayerBrus
 
 int32 ALandscape::GetBrushLayer(class ALandscapeBlueprintBrushBase* InBrush) const
 {
-	for (int32 LayerIndex = 0; LayerIndex < LandscapeLayers.Num(); ++LayerIndex)
+	for (int32 LayerIndex = 0; LayerIndex < LandscapeEditLayers.Num(); ++LayerIndex)
 	{
-		for (const FLandscapeLayerBrush& Brush : LandscapeLayers[LayerIndex].Brushes)
+		for (const FLandscapeLayerBrush& Brush : LandscapeEditLayers[LayerIndex].Brushes)
 		{
 			if (Brush.GetBrush() == InBrush)
 			{
@@ -9999,7 +10083,7 @@ int32 ALandscape::GetBrushLayer(class ALandscapeBlueprintBrushBase* InBrush) con
 void ALandscape::AddBrushToLayer(int32 InLayerIndex, ALandscapeBlueprintBrushBase* InBrush)
 {
 	check(GetBrushLayer(InBrush) == INDEX_NONE);
-	if (FLandscapeLayer* Layer = GetLayer(InLayerIndex))
+	if (FLandscapeLayer* Layer = GetLayerInternal(InLayerIndex))
 	{
 		Modify();
 		Layer->Brushes.Add(FLandscapeLayerBrush(InBrush));
@@ -10028,7 +10112,7 @@ void ALandscape::RemoveBrushFromLayer(int32 InLayerIndex, ALandscapeBlueprintBru
 
 void ALandscape::RemoveBrushFromLayer(int32 InLayerIndex, int32 InBrushIndex)
 {
-	if (FLandscapeLayer* Layer = GetLayer(InLayerIndex))
+	if (FLandscapeLayer* Layer = GetLayerInternal(InLayerIndex))
 	{
 		if (Layer->Brushes.IsValidIndex(InBrushIndex))
 		{
@@ -10046,7 +10130,7 @@ void ALandscape::RemoveBrushFromLayer(int32 InLayerIndex, int32 InBrushIndex)
 
 int32 ALandscape::GetBrushIndexForLayer(int32 InLayerIndex, ALandscapeBlueprintBrushBase* InBrush)
 {
-	if (FLandscapeLayer* Layer = GetLayer(InLayerIndex))
+	if (const FLandscapeLayer* Layer = GetLayerConst(InLayerIndex))
 	{
 		for (int32 i = 0; i < Layer->Brushes.Num(); ++i)
 		{
@@ -10098,7 +10182,7 @@ void ALandscape::OnLayerInfoSplineFalloffModulationChanged(ULandscapeLayerInfoOb
 
 ALandscapeBlueprintBrushBase* ALandscape::GetBrushForLayer(int32 InLayerIndex, int32 InBrushIndex) const
 {
-	if (const FLandscapeLayer* Layer = GetLayer(InLayerIndex))
+	if (const FLandscapeLayer* Layer = GetLayerConst(InLayerIndex))
 	{
 		if (Layer->Brushes.IsValidIndex(InBrushIndex))
 		{
@@ -10111,7 +10195,7 @@ ALandscapeBlueprintBrushBase* ALandscape::GetBrushForLayer(int32 InLayerIndex, i
 TArray<ALandscapeBlueprintBrushBase*> ALandscape::GetBrushesForLayer(int32 InLayerIndex) const
 {
 	TArray<ALandscapeBlueprintBrushBase*> Brushes;
-	if (const FLandscapeLayer* Layer = GetLayer(InLayerIndex))
+	if (const FLandscapeLayer* Layer = GetLayerConst(InLayerIndex))
 	{
 		Brushes.Reserve(Layer->Brushes.Num());
 		for (const FLandscapeLayerBrush& Brush : Layer->Brushes)
diff --git a/Engine/Source/Runtime/Landscape/Private/LandscapeSplineRaster.cpp b/Engine/Source/Runtime/Landscape/Private/LandscapeSplineRaster.cpp
--- a/Engine/Source/Runtime/Landscape/Private/LandscapeSplineRaster.cpp
+++ b/Engine/Source/Runtime/Landscape/Private/LandscapeSplineRaster.cpp
@@ -14,6 +14,7 @@
 #include "LandscapeHeightfieldCollisionComponent.h"
 #include "LandscapeDataAccess.h"
 #include "LandscapeEdit.h"
+#include "LandscapeEditLayer.h"
 #include "LandscapeSplinesComponent.h"
 #include "LandscapeSplineControlPoint.h"
 #include "LandscapePrivate.h"
@@ -287,8 +288,13 @@ void RasterizeHeight(int32& MinX, int32& MinY, int32& MaxX, int32& MaxY, FLandsc
 	FLandscapeEditDataInterface::ShrinkData(HeightData, MinX, MinY, MaxX, MaxY, ValidMinX, ValidMinY, ValidMaxX, ValidMaxY);
 
 	const ALandscape* Landscape = LandscapeEdit.GetTargetLandscape();
-	bool bIsEditingLayerReservedForSplines = Landscape && Landscape->IsEditingLayerReservedForSplines();
-	check(!bIsEditingLayerReservedForSplines || (Landscape->GetLandscapeSplinesReservedLayer()->BlendMode == LSBM_AlphaBlend));
+	check(Landscape != nullptr);
+
+	const FLandscapeLayer* EditingLayer = Landscape->GetLayerConst(Landscape->GetEditingLayer());
+	check((EditingLayer == nullptr) || (EditingLayer->EditLayer != nullptr));
+	const FLandscapeLayer* SplinesLayer = Landscape->FindLayerOfType(ULandscapeEditLayerSplines::StaticClass());
+	bool bIsEditingLayerReservedForSplines = (EditingLayer != nullptr) && (SplinesLayer == EditingLayer);
+	check(!bIsEditingLayerReservedForSplines || (SplinesLayer->BlendMode == LSBM_AlphaBlend));
 
 	TArray<uint16> HeightAlphaBlendData;
 	TArray<uint8> HeightFlagsData;
@@ -568,7 +574,7 @@ bool ULandscapeInfo::ApplySplines(bool bOnlySelected, TSet<TObjectPtr<ULandscape
 	bool bResult = false;
 
 	ALandscape* Landscape = LandscapeActor.Get();
-	const FLandscapeLayer* Layer = Landscape ? Landscape->GetLandscapeSplinesReservedLayer() : nullptr;
+	const FLandscapeLayer* Layer = Landscape ? Landscape->FindLayerOfType(ULandscapeEditLayerSplines::StaticClass()) : nullptr;
 	FGuid SplinesTargetLayerGuid = Layer ? Layer->Guid : Landscape ? Landscape->GetEditingLayer() : FGuid();
 	FScopedSetLandscapeEditingLayer Scope(Landscape, SplinesTargetLayerGuid, [=] { Landscape->RequestLayersContentUpdate(ELandscapeLayerUpdateMode::Update_All); });
 
diff --git a/Engine/Source/Runtime/Landscape/Public/LandscapeEditorServices.h b/Engine/Source/Runtime/Landscape/Public/LandscapeEditorServices.h
--- a/Engine/Source/Runtime/Landscape/Public/LandscapeEditorServices.h
+++ b/Engine/Source/Runtime/Landscape/Public/LandscapeEditorServices.h
@@ -15,9 +15,10 @@ public:
 	 * 
 	 * @param InEditLayerName The name of the layer to search for and possibly create
 	 * @param InTargetLandscape The landscape which should have a layer called InEditLayerName
+	 * @param InEditLayerClass The class of the edit layer to create. If none specified, a standard edit layer (ULandscapeEditLayer) will be created
 	 * @return The index at which the edit layer named InEditLayerName exists
 	 */
-	virtual int32 GetOrCreateEditLayer(FName InEditLayerName, ALandscape* InTargetLandscape) = 0;
+	virtual int32 GetOrCreateEditLayer(FName InEditLayerName, ALandscape* InTargetLandscape, const TSubclassOf<ULandscapeEditLayerBase>& InEditLayerClass = TSubclassOf<ULandscapeEditLayerBase>()) = 0;
 
 	/** Requires the landscape editor mode, if active to refresh its detail panel */
 	virtual void RefreshDetailPanel() = 0;
