Fixed locatlization issue

diff --git a/Engine/Source/Editor/LandscapeEditor/Private/LandscapeEditorDetailCustomization_Layers.cpp b/Engine/Source/Editor/LandscapeEditor/Private/LandscapeEditorDetailCustomization_Layers.cpp
--- a/Engine/Source/Editor/LandscapeEditor/Private/LandscapeEditorDetailCustomization_Layers.cpp
+++ b/Engine/Source/Editor/LandscapeEditor/Private/LandscapeEditorDetailCustomization_Layers.cpp
@@ -911,7 +911,7 @@ bool FLandscapeEditorCustomNodeBuilder_Layers::CanCollapseLayer(int32 InLayerInd
 
 	if (BottomLayer->Brushes.Num() > 0)
 	{
-		OutReason = FText::Format(LOCTEXT("Landscape_CollapseLayer_Reason_TopLayerHasBrush", "Cannot collapse layer '{0}' onto layer '{1}'. Layer '{1}' contains brush(es)"),
+		OutReason = FText::Format(LOCTEXT("Landscape_CollapseLayer_Reason_BottomLayerHasBrush", "Cannot collapse layer '{0}' onto layer '{1}'. Layer '{1}' contains brush(es)"),
 			FText::FromName(TopLayer->Name), FText::FromName(BottomLayer->Name));
 		return false;
 	}
