Improvements to landscape flatten UX:
* The flatten target value and terrace interval properties are defined in absolute world-space so we now bound their min and max slider values to what that the landscape can actually reach, so that the interval is always valid, regardless of the landscape scale
* Adjust the preview grid in realtime when the target value is changed via the slider, to make it much easier to eye-ball it
* Prevent the flatten target value eye dropper tool to validate the change when clicking anywhere outside the viewport or when leaving landscape mode
* Save the value of preview grid in the user settings
#tests editor

diff --git a/Engine/Source/Editor/LandscapeEditor/Private/LandscapeEdModePaintTools.cpp b/Engine/Source/Editor/LandscapeEditor/Private/LandscapeEdModePaintTools.cpp
--- a/Engine/Source/Editor/LandscapeEditor/Private/LandscapeEdModePaintTools.cpp
+++ b/Engine/Source/Editor/LandscapeEditor/Private/LandscapeEdModePaintTools.cpp
@@ -875,7 +875,7 @@ public:
 		if (InEdMode->UISettings->bUseFlattenTarget && bTargetIsHeightmap)
 		{
 			FTransform LocalToWorld = InTarget.LandscapeInfo->GetLandscapeProxy()->ActorToWorld();
-			float Height = static_cast<float>((InEdMode->UISettings->FlattenTarget - LocalToWorld.GetTranslation().Z) / LocalToWorld.GetScale3D().Z);
+			float Height = static_cast<float>((InEdMode->UISettings->GetFlattenTarget(/*bInReturnPreviewValueIfActive = */false) - LocalToWorld.GetTranslation().Z) / LocalToWorld.GetScale3D().Z);
 			FlattenValue = static_cast<ValueType>(LandscapeDataAccess::GetTexHeight(Height));
 			bInitializedFlattenValue = true;
 		}
@@ -1141,18 +1141,13 @@ class FLandscapeToolFlatten : public FLandscapeToolPaintBase<ToolTarget, FLandsc
 protected:
 	TObjectPtr<UStaticMesh> HeightmapFlattenPlaneMesh;
 	TObjectPtr<UStaticMeshComponent> HeightmapFlattenPreviewComponent;
-	bool CanToolBeActivatedNextTick;
-	bool CanToolBeActivatedValue;
-	float EyeDropperFlattenTargetValue;
+	bool bCanToolBeActivatedNextTick = false;
+	bool bCanToolBeActivatedValue = false;
 
 public:
 	FLandscapeToolFlatten(FEdModeLandscape* InEdMode)
 		: Super(InEdMode)
 		, HeightmapFlattenPlaneMesh(LoadObject<UStaticMesh>(nullptr, TEXT("/Engine/EditorLandscapeResources/FlattenPlaneMesh.FlattenPlaneMesh")))
-		, HeightmapFlattenPreviewComponent(nullptr)
-		, CanToolBeActivatedNextTick(false)
-		, CanToolBeActivatedValue(false)
-		, EyeDropperFlattenTargetValue(0.0f)
 	{
 		check(HeightmapFlattenPlaneMesh);
 	}
@@ -1170,8 +1165,8 @@ public:
 
 	virtual void SetCanToolBeActivated(bool Value) override
 	{ 
-		CanToolBeActivatedNextTick = true;
-		CanToolBeActivatedValue = Value;
+		bCanToolBeActivatedNextTick = true;
+		bCanToolBeActivatedValue = Value;
 	}
 
 	virtual void AddReferencedObjects(FReferenceCollector& Collector) override
@@ -1186,18 +1181,29 @@ public:
 
 	virtual void Tick(FEditorViewportClient* ViewportClient, float DeltaTime) override
 	{
-		if (CanToolBeActivatedNextTick)
+		if (bCanToolBeActivatedNextTick)
 		{
-			this->bCanToolBeActivated = CanToolBeActivatedValue;
-			CanToolBeActivatedNextTick = false;
+			this->bCanToolBeActivated = bCanToolBeActivatedValue;
+			bCanToolBeActivatedNextTick = false;
 		}
 
 		Super::Tick(ViewportClient, DeltaTime);
 
+		if (this->EdMode->UISettings->bFlattenEyeDropperModeActivated)
+		{
+			this->EdMode->UISettings->bFlattenEyeDropperModeMousingOverViewport = false;
+		}
+
 		if (HeightmapFlattenPreviewComponent != nullptr)
 		{
 			bool bShowGrid = this->EdMode->UISettings->bUseFlattenTarget && this->EdMode->UISettings->bShowFlattenTargetPreview;
 			HeightmapFlattenPreviewComponent->SetVisibility(bShowGrid);
+
+			// Always adjust the preview component's Z so that we react to the FlattenTarget changes in the details panel, which greatly helps eyeballing the adjustment : 
+			const FTransform LocalToWorld = this->EdMode->CurrentToolTarget.LandscapeInfo->GetLandscapeProxy()->ActorToWorld();
+			FVector Origin = HeightmapFlattenPreviewComponent->GetRelativeLocation();
+			Origin.Z = ConvertFlattenTargetToLocalZ(/*bInReturnPreviewValueIfActive = */true);
+			HeightmapFlattenPreviewComponent->SetRelativeLocation(Origin, /*bSweep = */false);
 		}
 	}
 
@@ -1205,24 +1211,32 @@ public:
 	{
 		bool bResult = Super::MouseMove(ViewportClient, Viewport, x, y);
 
-		if (ViewportClient->IsLevelEditorClient() && HeightmapFlattenPreviewComponent != nullptr)
+		if (ViewportClient->IsLevelEditorClient())
 		{
-			FVector MousePosition;
-			if (this->EdMode->LandscapeMouseTrace((FEditorViewportClient*)ViewportClient, x, y, MousePosition))
+			if (this->EdMode->UISettings->bFlattenEyeDropperModeActivated)
 			{
-				const FTransform LocalToWorld = this->EdMode->CurrentToolTarget.LandscapeInfo->GetLandscapeProxy()->ActorToWorld();
-				FVector Origin;
-				Origin.X = FMath::RoundToFloat(MousePosition.X);
-				Origin.Y = FMath::RoundToFloat(MousePosition.Y);
-				Origin.Z = (FMath::RoundToFloat((this->EdMode->UISettings->FlattenTarget - LocalToWorld.GetTranslation().Z) / LocalToWorld.GetScale3D().Z * LANDSCAPE_INV_ZSCALE) - 0.1f) * LANDSCAPE_ZSCALE;
-				HeightmapFlattenPreviewComponent->SetRelativeLocation(Origin, false);
-
-				// Clamp the value to the height map
-				uint16 TexHeight = LandscapeDataAccess::GetTexHeight(static_cast<float>(MousePosition.Z));
-				float Height = LandscapeDataAccess::GetLocalHeight(TexHeight);
-
-				// Convert the height back to world space
-				this->EdMode->UISettings->FlattenEyeDropperModeDesiredTarget = static_cast<float>(Height * LocalToWorld.GetScale3D().Z + LocalToWorld.GetTranslation().Z);
+				this->EdMode->UISettings->bFlattenEyeDropperModeMousingOverViewport = true;
+			}
+
+			if (HeightmapFlattenPreviewComponent != nullptr)
+			{
+				FVector MousePosition;
+				if (this->EdMode->LandscapeMouseTrace((FEditorViewportClient*)ViewportClient, x, y, MousePosition))
+				{
+					const FTransform LocalToWorld = this->EdMode->CurrentToolTarget.LandscapeInfo->GetLandscapeProxy()->ActorToWorld();
+					FVector Origin;
+					Origin.X = FMath::RoundToDouble(MousePosition.X);
+					Origin.Y = FMath::RoundToDouble(MousePosition.Y);
+					Origin.Z = ConvertFlattenTargetToLocalZ(/*bInReturnPreviewValueIfActive = */true);
+					HeightmapFlattenPreviewComponent->SetRelativeLocation(Origin, false);
+
+					// Clamp the value to the height map
+					uint16 TexHeight = LandscapeDataAccess::GetTexHeight(static_cast<float>(MousePosition.Z));
+					float Height = LandscapeDataAccess::GetLocalHeight(TexHeight);
+
+					// Convert the height back to world space
+					this->EdMode->UISettings->FlattenEyeDropperModeDesiredTarget = static_cast<float>(Height * LocalToWorld.GetScale3D().Z + LocalToWorld.GetTranslation().Z);
+				}
 			}
 		}
 
@@ -1237,6 +1251,8 @@ public:
 			return;
 		}
 
+		this->EdMode->UISettings->bFlattenEyeDropperModeMousingOverViewport = false;
+
 		if (ToolTarget::TargetType == ELandscapeToolTargetType::Heightmap)
 		{
 			ALandscapeProxy* LandscapeProxy = this->EdMode->CurrentToolTarget.LandscapeInfo->GetLandscapeProxy();
@@ -1252,7 +1268,7 @@ public:
 			// Try to set a sane initial location for the preview grid
 			const FTransform LocalToWorld = this->EdMode->CurrentToolTarget.LandscapeInfo->GetLandscapeProxy()->GetRootComponent()->GetComponentToWorld();
 			FVector Origin = FVector::ZeroVector;
-			Origin.Z = (FMath::RoundToFloat((this->EdMode->UISettings->FlattenTarget - LocalToWorld.GetTranslation().Z) / LocalToWorld.GetScale3D().Z * LANDSCAPE_INV_ZSCALE) - 0.1f) * LANDSCAPE_ZSCALE;
+			Origin.Z = ConvertFlattenTargetToLocalZ(/*bInReturnPreviewValueIfActive = */false);
 			HeightmapFlattenPreviewComponent->SetRelativeLocation(Origin, false);
 		}
 	}
@@ -1267,6 +1283,17 @@ public:
 			HeightmapFlattenPreviewComponent->DestroyComponent();
 			HeightmapFlattenPreviewComponent = nullptr;
 		}
+
+		// Make sure the eye dropper tool is reset upon leaving the flatten tool :
+		this->EdMode->UISettings->bFlattenEyeDropperModeActivated = false;
+		this->EdMode->UISettings->bFlattenEyeDropperModeMousingOverViewport = false;
+	}
+
+private:
+	double ConvertFlattenTargetToLocalZ(bool bInReturnPreviewValueIfActive) const
+	{
+		const FTransform LocalToWorld = this->EdMode->CurrentToolTarget.LandscapeInfo->GetLandscapeProxy()->ActorToWorld();
+		return (FMath::RoundToDouble((this->EdMode->UISettings->GetFlattenTarget(bInReturnPreviewValueIfActive) - LocalToWorld.GetTranslation().Z) / LocalToWorld.GetScale3D().Z * LANDSCAPE_INV_ZSCALE) - 0.1f) * LANDSCAPE_ZSCALE;
 	}
 };
 
diff --git a/Engine/Source/Editor/LandscapeEditor/Private/LandscapeEditorDetailCustomization_MiscTools.cpp b/Engine/Source/Editor/LandscapeEditor/Private/LandscapeEditorDetailCustomization_MiscTools.cpp
--- a/Engine/Source/Editor/LandscapeEditor/Private/LandscapeEditorDetailCustomization_MiscTools.cpp
+++ b/Engine/Source/Editor/LandscapeEditor/Private/LandscapeEditorDetailCustomization_MiscTools.cpp
@@ -82,12 +82,10 @@ void FLandscapeEditorDetailCustomization_MiscTools::CustomizeDetails(IDetailLayo
 					.Value(this, &FLandscapeEditorDetailCustomization_MiscTools::GetFlattenValue)
 					.OnValueChanged_Static(&FLandscapeEditorDetailCustomization_Base::OnValueChanged<float>, FlattenValueProperty)
 					.OnValueCommitted_Static(&FLandscapeEditorDetailCustomization_Base::OnValueCommitted<float>, FlattenValueProperty)
-					.MinValue(-32768.0f)
-					.MaxValue(32768.0f)
-					.SliderExponentNeutralValue(0.0f)
+					.SliderExponentNeutralValue(this, &FLandscapeEditorDetailCustomization_MiscTools::GetFlattenTargetValueMid)
 					.SliderExponent(5.0f)
-					.MinSliderValue(-32768.0f)
-					.MaxSliderValue(32768.0f)
+					.MinSliderValue(this, &FLandscapeEditorDetailCustomization_MiscTools::GetFlattenTargetValueMin)
+					.MaxSliderValue(this, &FLandscapeEditorDetailCustomization_MiscTools::GetFlattenTargetValueMax)
 					.MinDesiredValueWidth(75.0f)
 					.ToolTipText(LOCTEXT("FlattenToolTips", "Target height to flatten towards (in Unreal Units)"))
 				]
@@ -102,6 +100,28 @@ void FLandscapeEditorDetailCustomization_MiscTools::CustomizeDetails(IDetailLayo
 					.OnComplete(this, &FLandscapeEditorDetailCustomization_MiscTools::OnCompletedFlattenToolEyeDrop)
 				]		
 			];
+
+		TSharedRef<IPropertyHandle> TerraceIntervalProperty = DetailBuilder.GetProperty(GET_MEMBER_NAME_CHECKED(ULandscapeEditorObject, TerraceInterval));
+		IDetailPropertyRow& TerraceIntervalRow = ToolsCategory.AddProperty(TerraceIntervalProperty);
+		TerraceIntervalRow.CustomWidget()
+			.NameContent()
+			[
+				TerraceIntervalProperty->CreatePropertyNameWidget()
+			]
+			.ValueContent()
+			[
+				SNew(SNumericEntryBox<float>)
+					.AllowSpin(true)
+					.Font(DetailBuilder.GetDetailFont())
+					.Value_Static(&FLandscapeEditorDetailCustomization_Base::GetOptionalPropertyValue<float>, TerraceIntervalProperty)
+					.OnValueChanged_Static(&FLandscapeEditorDetailCustomization_Base::OnValueChanged<float>, TerraceIntervalProperty)
+					.OnValueCommitted_Static(&FLandscapeEditorDetailCustomization_Base::OnValueCommitted<float>, TerraceIntervalProperty)
+					.SliderExponent(5.0f)
+					.MinValue(1.0f)
+					.MinSliderValue(1.0f)
+					.MaxSliderValue(this, &FLandscapeEditorDetailCustomization_MiscTools::GetFlattenTerraceIntervalValueMax)
+					.ToolTipText(LOCTEXT("TerraceIntervalToolTips", "Height of the terrace intervals in unreal units"))
+			];
 	}
 
 	if (IsToolActive("Splines"))
@@ -515,17 +535,82 @@ TOptional<float> FLandscapeEditorDetailCustomization_MiscTools::GetFlattenValue(
 	FEdModeLandscape* LandscapeEdMode = GetEditorMode();
 	if (LandscapeEdMode != nullptr && IsToolActive(FName("Flatten")))
 	{
-		if (LandscapeEdMode->UISettings->bFlattenEyeDropperModeActivated)
+		return LandscapeEdMode->UISettings->GetFlattenTarget(/* bInReturnPreviewValueIfActive = */true);
+	}
+
+	return 0.0f;
+}
+
+TOptional<float> FLandscapeEditorDetailCustomization_MiscTools::GetFlattenTargetValueMin() const
+{
+	FEdModeLandscape* LandscapeEdMode = GetEditorMode();
+	if ((LandscapeEdMode != nullptr) && IsToolActive(FName("Flatten")))
+	{
+		if (ULandscapeInfo* LandscapeInfo = LandscapeEdMode->CurrentToolTarget.LandscapeInfo.Get())
 		{
-			return LandscapeEdMode->UISettings->FlattenEyeDropperModeDesiredTarget;
-		}
+			if (ALandscapeProxy* LandscapeProxy = LandscapeInfo->GetLandscapeProxy())
+			{
+				return LandscapeProxy->ActorToWorld().TransformPosition(FVector(0.0, 0.0, LandscapeDataAccess::GetLocalHeight(0))).Z;
+			}
+		}		
+	}
+
+	return TOptional<float>();
+}
 
-		return LandscapeEdMode->UISettings->FlattenTarget;
+float FLandscapeEditorDetailCustomization_MiscTools::GetFlattenTargetValueMid() const
+{
+	FEdModeLandscape* LandscapeEdMode = GetEditorMode();
+	if ((LandscapeEdMode != nullptr) && IsToolActive(FName("Flatten")))
+	{
+		if (ULandscapeInfo* LandscapeInfo = LandscapeEdMode->CurrentToolTarget.LandscapeInfo.Get())
+		{
+			if (ALandscapeProxy* LandscapeProxy = LandscapeInfo->GetLandscapeProxy())
+			{
+				return LandscapeProxy->ActorToWorld().TransformPosition(FVector(0.0, 0.0, LandscapeDataAccess::GetLocalHeight(LandscapeDataAccess::MidValue))).Z;
+			}
+		}
 	}
 
 	return 0.0f;
 }
 
+TOptional<float> FLandscapeEditorDetailCustomization_MiscTools::GetFlattenTargetValueMax() const
+{
+	FEdModeLandscape* LandscapeEdMode = GetEditorMode();
+	if ((LandscapeEdMode != nullptr) && IsToolActive(FName("Flatten")))
+	{
+		if (ULandscapeInfo* LandscapeInfo = LandscapeEdMode->CurrentToolTarget.LandscapeInfo.Get())
+		{
+			if (ALandscapeProxy* LandscapeProxy = LandscapeInfo->GetLandscapeProxy())
+			{
+				return LandscapeProxy->ActorToWorld().TransformPosition(FVector(0.0, 0.0, LandscapeDataAccess::GetLocalHeight(LandscapeDataAccess::MaxValue))).Z;
+			}
+		}
+	}
+
+	return TOptional<float>();
+}
+
+TOptional<float> FLandscapeEditorDetailCustomization_MiscTools::GetFlattenTerraceIntervalValueMax() const
+{
+	FEdModeLandscape* LandscapeEdMode = GetEditorMode();
+	if ((LandscapeEdMode != nullptr) && IsToolActive(FName("Flatten")))
+	{
+		if (ULandscapeInfo* LandscapeInfo = LandscapeEdMode->CurrentToolTarget.LandscapeInfo.Get())
+		{
+			if (ALandscapeProxy* LandscapeProxy = LandscapeInfo->GetLandscapeProxy())
+			{
+				float WorldMax = LandscapeProxy->ActorToWorld().TransformPosition(FVector(0.0, 0.0, LandscapeDataAccess::GetLocalHeight(LandscapeDataAccess::MaxValue))).Z;
+				float WorldMin = LandscapeProxy->ActorToWorld().TransformPosition(FVector(0.0, 0.0, LandscapeDataAccess::GetLocalHeight(0))).Z;
+				return FMath::Max(WorldMax - WorldMin, 1.0f);
+			}
+		}
+	}
+
+	return TOptional<float>();
+}
+
 EVisibility FLandscapeEditorDetailCustomization_MiscTools::GetMiscLandscapeErrorVisibility() const
 {
 	FEdModeLandscape* EdMode = GetEditorMode();
@@ -556,6 +641,8 @@ void FLandscapeEditorDetailCustomization_MiscTools::OnBeginFlattenToolEyeDrop()
 	if (LandscapeEdMode != nullptr && IsToolActive(FName("Flatten")))
 	{
 		LandscapeEdMode->UISettings->bFlattenEyeDropperModeActivated = true;
+		// Use the current target value when activating the eye drop for consistency. This will be reset when the mouse moves in the viewport anyway : 
+		LandscapeEdMode->UISettings->FlattenEyeDropperModeDesiredTarget = LandscapeEdMode->UISettings->FlattenTarget;
 		LandscapeEdMode->CurrentTool->SetCanToolBeActivated(false);
 	}
 }
@@ -565,13 +652,13 @@ void FLandscapeEditorDetailCustomization_MiscTools::OnCompletedFlattenToolEyeDro
 	FEdModeLandscape* LandscapeEdMode = GetEditorMode();
 	if (LandscapeEdMode != nullptr && IsToolActive(FName("Flatten")))
 	{
-		LandscapeEdMode->UISettings->bFlattenEyeDropperModeActivated = false;
-		LandscapeEdMode->CurrentTool->SetCanToolBeActivated(true);
-
-		if (!Canceled)
+		// Consider clicks outside the viewport as a cancellation : this avoids validating the preview value when clicking outside the viewport
+		if (!Canceled && LandscapeEdMode->UISettings->bFlattenEyeDropperModeActivated && LandscapeEdMode->UISettings->bFlattenEyeDropperModeMousingOverViewport)
 		{
 			LandscapeEdMode->UISettings->FlattenTarget = LandscapeEdMode->UISettings->FlattenEyeDropperModeDesiredTarget;
 		}
+		LandscapeEdMode->UISettings->bFlattenEyeDropperModeActivated = false;
+		LandscapeEdMode->CurrentTool->SetCanToolBeActivated(true);
 	}
 }
 
diff --git a/Engine/Source/Editor/LandscapeEditor/Private/LandscapeEditorDetailCustomization_MiscTools.h b/Engine/Source/Editor/LandscapeEditor/Private/LandscapeEditorDetailCustomization_MiscTools.h
--- a/Engine/Source/Editor/LandscapeEditor/Private/LandscapeEditorDetailCustomization_MiscTools.h
+++ b/Engine/Source/Editor/LandscapeEditor/Private/LandscapeEditorDetailCustomization_MiscTools.h
@@ -56,6 +56,10 @@ protected:
 	void OnBeginFlattenToolEyeDrop();
 	void OnCompletedFlattenToolEyeDrop(bool Canceled);
 	TOptional<float> GetFlattenValue() const;
+	TOptional<float> GetFlattenTargetValueMin() const;
+	float GetFlattenTargetValueMid() const;
+	TOptional<float> GetFlattenTargetValueMax() const;
+	TOptional<float> GetFlattenTerraceIntervalValueMax() const;
 
 	// Common Error handling
 	EVisibility GetMiscLandscapeErrorVisibility() const;
diff --git a/Engine/Source/Editor/LandscapeEditor/Private/LandscapeEditorObject.cpp b/Engine/Source/Editor/LandscapeEditor/Private/LandscapeEditorObject.cpp
--- a/Engine/Source/Editor/LandscapeEditor/Private/LandscapeEditorObject.cpp
+++ b/Engine/Source/Editor/LandscapeEditor/Private/LandscapeEditorObject.cpp
@@ -259,8 +259,13 @@ void ULandscapeEditorObject::Load()
 	bool InbUseFlattenTarget = bUseFlattenTarget;
 	GConfig->GetBool(TEXT("LandscapeEdit"), TEXT("bUseFlattenTarget"), InbUseFlattenTarget, GEditorPerProjectIni);
 	bUseFlattenTarget = InbUseFlattenTarget;
+
 	GConfig->GetFloat(TEXT("LandscapeEdit"), TEXT("FlattenTarget"), FlattenTarget, GEditorPerProjectIni);
 
+	bool InbShowFlattenTargetPreview = bShowFlattenTargetPreview;
+	GConfig->GetBool(TEXT("LandscapeEdit"), TEXT("bShowFlattenTargetPreview"), InbShowFlattenTargetPreview, GEditorPerProjectIni);
+	bShowFlattenTargetPreview = InbShowFlattenTargetPreview;
+
 	GConfig->GetFloat(TEXT("LandscapeEdit"), TEXT("TerraceSmooth"), TerraceSmooth, GEditorPerProjectIni);
 	GConfig->GetFloat(TEXT("LandscapeEdit"), TEXT("TerraceInterval"), TerraceInterval, GEditorPerProjectIni);
 
@@ -413,7 +418,7 @@ void ULandscapeEditorObject::Save()
 	GConfig->SetBool(TEXT("LandscapeEdit"), TEXT("bPickValuePerApply"), bPickValuePerApply, GEditorPerProjectIni);
 	GConfig->SetBool(TEXT("LandscapeEdit"), TEXT("bUseFlattenTarget"), bUseFlattenTarget, GEditorPerProjectIni);
 	GConfig->SetFloat(TEXT("LandscapeEdit"), TEXT("FlattenTarget"), FlattenTarget, GEditorPerProjectIni);
-
+	GConfig->SetBool(TEXT("LandscapeEdit"), TEXT("bShowFlattenTargetPreview"), bShowFlattenTargetPreview, GEditorPerProjectIni);
 	GConfig->SetFloat(TEXT("LandscapeEdit"), TEXT("TerraceSmooth"), TerraceSmooth, GEditorPerProjectIni);
 	GConfig->SetFloat(TEXT("LandscapeEdit"), TEXT("TerraceInterval"), TerraceInterval, GEditorPerProjectIni);
 
@@ -1058,3 +1063,8 @@ void ULandscapeEditorObject::SetCurrentToolBrushFalloff(float NewBrushFalloff)
 		BrushFalloff = NewBrushFalloff;
 	}
 }
+
+float ULandscapeEditorObject::GetFlattenTarget(bool bInReturnPreviewValueIfActive) const
+{
+	return (bFlattenEyeDropperModeActivated && bInReturnPreviewValueIfActive) ? FlattenEyeDropperModeDesiredTarget : FlattenTarget;
+}
\ No newline at end of file
diff --git a/Engine/Source/Editor/LandscapeEditor/Private/SFlattenHeightEyeDropperButton.cpp b/Engine/Source/Editor/LandscapeEditor/Private/SFlattenHeightEyeDropperButton.cpp
--- a/Engine/Source/Editor/LandscapeEditor/Private/SFlattenHeightEyeDropperButton.cpp
+++ b/Engine/Source/Editor/LandscapeEditor/Private/SFlattenHeightEyeDropperButton.cpp
@@ -22,6 +22,7 @@ void SFlattenHeightEyeDropperButton::Construct(const FArguments& InArgs)
 	// Their visibility are changed according to whether dropper mode is active or not.
 	SButton::Construct(
 		SButton::FArguments()
+		.ButtonStyle(FAppStyle::Get(), "SimpleButton")
 		.ContentPadding(1.0f)
 		.OnClicked(this, &SFlattenHeightEyeDropperButton::OnClicked)
 		.OnHovered(this, &SFlattenHeightEyeDropperButton::OnMouseHovered)
diff --git a/Engine/Source/Editor/LandscapeEditor/Public/LandscapeEditorObject.h b/Engine/Source/Editor/LandscapeEditor/Public/LandscapeEditorObject.h
--- a/Engine/Source/Editor/LandscapeEditor/Public/LandscapeEditorObject.h
+++ b/Engine/Source/Editor/LandscapeEditor/Public/LandscapeEditorObject.h
@@ -328,7 +328,7 @@ class ULandscapeEditorObject : public UObject
 	float FlattenTarget;
 
 	// Whether to show the preview grid for the flatten target height
-	UPROPERTY(Category = "Tool Settings", EditAnywhere, NonTransactional, AdvancedDisplay, meta = (DisplayName = "Show Preview Grid", ShowForTools = "Flatten", ShowForTargetTypes = "Heightmap", EditCondition = "bUseFlattenTarget", HideEditConditionToggle, UIMin = "-32768", UIMax = "32768"))
+	UPROPERTY(Category = "Tool Settings", EditAnywhere, NonTransactional, AdvancedDisplay, meta = (DisplayName = "Show Preview Grid", ShowForTools = "Flatten", ShowForTargetTypes = "Heightmap", EditCondition = "bUseFlattenTarget", HideEditConditionToggle))
 	bool bShowFlattenTargetPreview;
 
 	// Height of the terrace intervals in unreal units, for the terrace flatten mode 
@@ -343,6 +343,10 @@ class ULandscapeEditorObject : public UObject
 	UPROPERTY(NonTransactional, Transient)
 	bool bFlattenEyeDropperModeActivated;
 
+	// When in Eye Dropper mode, indicates whether we're currently mousing over the viewport 
+	UPROPERTY(NonTransactional, Transient)
+	bool bFlattenEyeDropperModeMousingOverViewport;
+
 	UPROPERTY(NonTransactional, Transient)
 	float FlattenEyeDropperModeDesiredTarget;
 
@@ -813,6 +817,8 @@ public:
 	float GetCurrentToolBrushFalloff() const;
 	void SetCurrentToolBrushFalloff(float NewBrushFalloff);
 
+	float GetFlattenTarget(bool bInReturnPreviewValueIfActive) const;
+
 private:
 
 	bool IsWeightmapTarget() const
