Adding localization support to the Landscape "Edit Layers Blueprint Brushes" text

diff --git a/Engine/Source/Editor/LandscapeEditor/Private/LandscapeEditorDetailCustomization_LayersBrushStack.cpp b/Engine/Source/Editor/LandscapeEditor/Private/LandscapeEditorDetailCustomization_LayersBrushStack.cpp
--- a/Engine/Source/Editor/LandscapeEditor/Private/LandscapeEditorDetailCustomization_LayersBrushStack.cpp
+++ b/Engine/Source/Editor/LandscapeEditor/Private/LandscapeEditorDetailCustomization_LayersBrushStack.cpp
@@ -89,7 +89,7 @@ void FLandscapeEditorCustomNodeBuilder_LayersBrushStack::GenerateChildContent(ID
 		BrushesList->SetDropIndicator_Above(*FAppStyle::GetBrush("LandscapeEditor.TargetList.DropZone.Above"));
 		BrushesList->SetDropIndicator_Below(*FAppStyle::GetBrush("LandscapeEditor.TargetList.DropZone.Below"));
 
-		ChildrenBuilder.AddCustomRow(FText::FromString(FString(TEXT("Edit Layer Blueprint Brushes"))))
+		ChildrenBuilder.AddCustomRow(LOCTEXT("LandscapeEditLayerBlueprintBrushes", "Edit Layer Blueprint Brushes"))
 			.Visibility(EVisibility::Visible)
 			[
 				SNew(SVerticalBox)
