Landscape: Fix crash when opening landscape mode on a map without edit layers.

diff --git a/Engine/Source/Editor/LandscapeEditor/Private/LandscapeEdMode.cpp b/Engine/Source/Editor/LandscapeEditor/Private/LandscapeEdMode.cpp
--- a/Engine/Source/Editor/LandscapeEditor/Private/LandscapeEdMode.cpp
+++ b/Engine/Source/Editor/LandscapeEditor/Private/LandscapeEdMode.cpp
@@ -3929,8 +3929,9 @@ void FEdModeLandscape::SetCurrentLayer(int32 InLayerIndex)
 	if (ALandscape* Landscape = GetLandscape())
 	{
 		const FLandscapeLayer* Layer = Landscape->GetLayerConst(InLayerIndex);
-		check((Layer != nullptr) && (Layer->EditLayer != nullptr));
-		if (Layer->EditLayer->IsA<ULandscapeEditLayerSplines>())
+		if (Landscape->HasLayersContent() 
+			&& ensure(Layer && Layer->EditLayer)
+			&& Layer->EditLayer->IsA<ULandscapeEditLayerSplines>())
 		{
 			SetCurrentToolMode("ToolMode_Manage", false);
 			SetCurrentTool(FName("Splines"));
