Landscape: Make the new layer creation dialog be optional, controlled by bShowDialogForAutomaticLayerCreation landscape editor setting. Also make it have a single "complete" button since the cancel button does not actually cancel, and update the layer view when a layer is created.

diff --git a/Engine/Source/Editor/LandscapeEditor/Private/LandscapeEditorModule.cpp b/Engine/Source/Editor/LandscapeEditor/Private/LandscapeEditorModule.cpp
--- a/Engine/Source/Editor/LandscapeEditor/Private/LandscapeEditorModule.cpp
+++ b/Engine/Source/Editor/LandscapeEditor/Private/LandscapeEditorModule.cpp
@@ -20,6 +20,7 @@
 #include "LandscapeFileFormatRaw.h"
 #include "LandscapeEditorServices.h"
 #include "LandscapeImageFileCache.h"
+#include "LandscapeSettings.h"
 #include "SLandscapeLayerListDialog.h"
 
 #include "Framework/MultiBox/MultiBoxBuilder.h"
@@ -528,11 +529,19 @@ int32 FLandscapeEditorModule::GetOrCreateEditLayer(FName InEditLayerName, ALands
 	int32 ExistingLayerIndex = InTargetLandscape->GetLayerIndex(InEditLayerName);
 	if (ExistingLayerIndex == INDEX_NONE)
 	{
-		InTargetLandscape->CreateLayer(InEditLayerName, InEditLayerClass);
-		TSharedPtr<SLandscapeLayerListDialog> Dialog = SNew(SLandscapeLayerListDialog, InTargetLandscape);
-		Dialog->ShowModal();
-		ExistingLayerIndex = Dialog->GetInsertedLayerIndex();
+		ExistingLayerIndex = InTargetLandscape->CreateLayer(InEditLayerName, InEditLayerClass);
+
+		const ULandscapeSettings* Settings = GetDefault<ULandscapeSettings>();
+		if (Settings && Settings->bShowDialogForAutomaticLayerCreation)
+		{
+			TSharedPtr<SLandscapeLayerListDialog> Dialog = SNew(SLandscapeLayerListDialog, InTargetLandscape);
+			Dialog->ShowModal();
+			ExistingLayerIndex = Dialog->GetInsertedLayerIndex();
+		}
 	}
+
+	RefreshDetailPanel();
+
 	return ExistingLayerIndex;
 }
 
diff --git a/Engine/Source/Editor/LandscapeEditor/Private/SLandscapeLayerListDialog.cpp b/Engine/Source/Editor/LandscapeEditor/Private/SLandscapeLayerListDialog.cpp
--- a/Engine/Source/Editor/LandscapeEditor/Private/SLandscapeLayerListDialog.cpp
+++ b/Engine/Source/Editor/LandscapeEditor/Private/SLandscapeLayerListDialog.cpp
@@ -197,14 +197,15 @@ void SLandscapeLayerListDialog::Construct(const FArguments& InArgs, const TWeakO
 	.OnGenerateRow(this, &SLandscapeLayerListDialog::OnGenerateRow)
 	.ListItemsSource(&WidgetLayerList);
 	
-	FButton CancelButton(LOCTEXT("Cancel", "Cancel"));
-	FButton AcceptButton(LOCTEXT("Accept", "Accept"));
-	AcceptButton.OnClicked.BindSP(this, &SLandscapeLayerListDialog::OnAccept);
+	FButton CompleteButton(LOCTEXT("CompleteButtonText", "Complete"));
+	CompleteButton.OnClicked.BindSP(this, &SLandscapeLayerListDialog::OnComplete);
 	
 	// Construct custom dialog with list view supporting drag + drop
 	SCustomDialog::Construct(SCustomDialog::FArguments()
 		.Title(FText(LOCTEXT("LandscapeLayerListDialogTitleText", "Insert New Landscape Edit Layer")))
 		.UseScrollBox(false)
+		.WindowArguments(SWindow::FArguments()
+			.HasCloseButton(false))
 		.Content()
 		[
 			SNew(SBox)
@@ -220,8 +221,7 @@ void SLandscapeLayerListDialog::Construct(const FArguments& InArgs, const TWeakO
 					.WrapTextAt(350)
 					.Text(FText::Format(LOCTEXT(
 						"LandscapeLayerListDialogInstructionText",
-						"Drag/drop the \"{0}\" layer in the list to choose where in the edit layers stack it should be inserted.\n"
-						"\n(Note: Closing the window or pressing cancel will result in the \"{0}\" layer going onto the top of the layer stack.)\n"),
+						"Drag/drop the \"{0}\" layer in the list to choose where in the edit layer stack it should be inserted.\n"),
 						FText::FromName(WidgetLayerList[0]->LayerName)))
 				]
 				+ SVerticalBox::Slot()
@@ -236,8 +236,7 @@ void SLandscapeLayerListDialog::Construct(const FArguments& InArgs, const TWeakO
 			]
 		]
 		.Buttons({
-			AcceptButton,
-			CancelButton
+			CompleteButton
 		})
 	);
 }
@@ -255,7 +254,7 @@ TSharedRef<ITableRow> SLandscapeLayerListDialog::OnGenerateRow(TSharedPtr<FWidge
 	return SNew( SWidgetLayerListItem, InOwnerTableView, InListItem );
 }
 
-void SLandscapeLayerListDialog::OnAccept()
+void SLandscapeLayerListDialog::OnComplete()
 {
 	check(Landscape.IsValid());
 	TArrayView<const FLandscapeLayer> LandscapeLayers = Landscape->GetLayers();
diff --git a/Engine/Source/Editor/LandscapeEditor/Public/SLandscapeLayerListDialog.h b/Engine/Source/Editor/LandscapeEditor/Public/SLandscapeLayerListDialog.h
--- a/Engine/Source/Editor/LandscapeEditor/Public/SLandscapeLayerListDialog.h
+++ b/Engine/Source/Editor/LandscapeEditor/Public/SLandscapeLayerListDialog.h
@@ -39,8 +39,8 @@ private:
 	/* Called when a new row is being generated. Converts FWidgetLayerListItem to ITableRow. */
 	TSharedRef<ITableRow> OnGenerateRow( TSharedPtr<FWidgetLayerListItem> InListItem, const TSharedRef< STableViewBase >& InOwnerTableView ) const;
 
-	/* Called when Accept button is pressed */
-	void OnAccept();
+	/* Called when Complete button is pressed */
+	void OnComplete();
 
 	/* The list view of WidgetLayerList. */
 	TSharedPtr<SWidgetLayerListView> LayerListView;
diff --git a/Engine/Source/Runtime/Landscape/Public/LandscapeSettings.h b/Engine/Source/Runtime/Landscape/Public/LandscapeSettings.h
--- a/Engine/Source/Runtime/Landscape/Public/LandscapeSettings.h
+++ b/Engine/Source/Runtime/Landscape/Public/LandscapeSettings.h
@@ -58,6 +58,10 @@ public:
 	UPROPERTY(config, EditAnywhere, Category = "Layers", meta=(UIMin = "1", UIMax = "32", ClampMin = "1", ClampMax = "32", ToolTip = "This option controls the maximum editing layers that can be added to a Landscape"))
 	int32 MaxNumberOfLayers = 8;
 
+	UPROPERTY(EditAnywhere, config, Category = "Layers", meta = (ToolTip = 
+		"When true, automatic edit layer creation pops up a dialog where the new layer can be reordered relative to other layers."))
+	bool bShowDialogForAutomaticLayerCreation = true;
+
 	UPROPERTY(config, EditAnywhere, Category = "Configuration", meta=(ToolTip = "Maximum Dimension of Landscape in Components"))
 	int32 MaxComponents = 256;
 
