Fixed rare crash when deleting edit layers

diff --git a/Engine/Source/Editor/LandscapeEditor/Private/LandscapeEditorDetailCustomization_Layers.cpp b/Engine/Source/Editor/LandscapeEditor/Private/LandscapeEditorDetailCustomization_Layers.cpp
--- a/Engine/Source/Editor/LandscapeEditor/Private/LandscapeEditorDetailCustomization_Layers.cpp
+++ b/Engine/Source/Editor/LandscapeEditor/Private/LandscapeEditorDetailCustomization_Layers.cpp
@@ -644,7 +644,10 @@ bool FLandscapeEditorCustomNodeBuilder_Layers::CanRenameLayer(int32 InLayerIndex
 {
 	FEdModeLandscape* LandscapeEdMode = GetEditorMode();
 	const FLandscapeLayer* Layer = LandscapeEdMode->GetLayer(InLayerIndex);
-	check(Layer != nullptr);
+	if (Layer == nullptr)
+	{
+		return false;
+	}
 
 	if (Layer->bLocked)
 	{
@@ -669,7 +672,11 @@ bool FLandscapeEditorCustomNodeBuilder_Layers::CanClearTargetLayerOnLayer(int32
 {
 	FEdModeLandscape* LandscapeEdMode = GetEditorMode();
 	const FLandscapeLayer* Layer = LandscapeEdMode->GetLayer(InLayerIndex);
-	check(Layer != nullptr);
+	if (Layer == nullptr)
+	{
+		return false;
+	}
+
 	const ULandscapeEditLayerBase* EditLayer = Layer->EditLayer;
 	check(EditLayer != nullptr);
 
@@ -719,7 +726,11 @@ bool FLandscapeEditorCustomNodeBuilder_Layers::CanClearLayer(int32 InLayerIndex,
 {
 	FEdModeLandscape* LandscapeEdMode = GetEditorMode();
 	const FLandscapeLayer* Layer = LandscapeEdMode->GetLayer(InLayerIndex);
-	check(Layer != nullptr);
+	if (Layer == nullptr)
+	{
+		return false;
+	}
+
 	const ULandscapeEditLayerBase* EditLayer = Layer->EditLayer;
 	check(EditLayer != nullptr);
 
@@ -744,7 +755,11 @@ bool FLandscapeEditorCustomNodeBuilder_Layers::CanClearTargetLayersOnLayer(int32
 {
 	FEdModeLandscape* LandscapeEdMode = GetEditorMode();
 	const FLandscapeLayer* Layer = LandscapeEdMode->GetLayer(InLayerIndex);
-	check(Layer != nullptr);
+	if (Layer == nullptr)
+	{
+		return false;
+	}
+
 	const ULandscapeEditLayerBase* EditLayer = Layer->EditLayer;
 	check(EditLayer != nullptr);
 
@@ -814,7 +829,10 @@ bool FLandscapeEditorCustomNodeBuilder_Layers::CanDeleteLayer(int32 InLayerIndex
 	}
 
 	const FLandscapeLayer* Layer = LandscapeEdMode->GetLayer(InLayerIndex);
-	check(Layer != nullptr);
+	if (Layer == nullptr)
+	{
+		return false;
+	}
 
 	if (Layer->bLocked)
 	{
@@ -867,7 +885,10 @@ bool FLandscapeEditorCustomNodeBuilder_Layers::CanCollapseLayer(int32 InLayerInd
 	
 	const FLandscapeLayer* TopLayer = LandscapeEdMode->GetLayer(InLayerIndex);
 	const FLandscapeLayer* BottomLayer = LandscapeEdMode->GetLayer(InLayerIndex - 1);
-	check((TopLayer != nullptr) && (BottomLayer != nullptr));
+	if ((TopLayer == nullptr) || (BottomLayer == nullptr))
+	{
+		return false;
+	}
 
 	const ULandscapeEditLayerBase* TopEditLayer = TopLayer->EditLayer;
 	const ULandscapeEditLayerBase* BottomEditLayer = BottomLayer->EditLayer;
@@ -956,7 +977,11 @@ bool FLandscapeEditorCustomNodeBuilder_Layers::CanExecuteCustomLayerAction(int32
 	}
 
 	const FLandscapeLayer* Layer = LandscapeEdMode->GetLayer(InLayerIndex);
-	check(Layer != nullptr);
+	if (Layer == nullptr)
+	{
+		return false;
+	}
+
 	const ULandscapeEditLayerBase* EditLayer = Layer->EditLayer;
 	check(EditLayer != nullptr);
 	ULandscapeEditLayerBase::FEditLayerAction::FExecuteParams ExecuteParams(Layer, Landscape);
@@ -969,17 +994,19 @@ void FLandscapeEditorCustomNodeBuilder_Layers::ExecuteCustomLayerAction(int32 In
 	if (ALandscape* Landscape = LandscapeEdMode ? LandscapeEdMode->GetLandscape() : nullptr)
 	{
 		const FLandscapeLayer* Layer = LandscapeEdMode->GetLayer(InLayerIndex);
-		check(Layer != nullptr);
-		const ULandscapeEditLayerBase* EditLayer = Layer->EditLayer;
-		check(EditLayer != nullptr);
-
-		const FScopedTransaction Transaction(InCustomLayerAction.GetLabel());
-		ULandscapeEditLayerBase::FEditLayerAction::FExecuteParams ExecuteParams(Layer, Landscape);
-		ULandscapeEditLayerBase::FEditLayerAction::FExecuteResult Result = InCustomLayerAction.GetExecuteDelegate().Execute(ExecuteParams);
-		if (!Result.bSuccess)
+		if (Layer != nullptr)
 		{
-			// Indicate to the user that the action failed : 
-			FMessageDialog::Open(EAppMsgType::Ok, Result.Reason);
+			const ULandscapeEditLayerBase* EditLayer = Layer->EditLayer;
+			check(EditLayer != nullptr);
+
+			const FScopedTransaction Transaction(InCustomLayerAction.GetLabel());
+			ULandscapeEditLayerBase::FEditLayerAction::FExecuteParams ExecuteParams(Layer, Landscape);
+			ULandscapeEditLayerBase::FEditLayerAction::FExecuteResult Result = InCustomLayerAction.GetExecuteDelegate().Execute(ExecuteParams);
+			if (!Result.bSuccess)
+			{
+				// Indicate to the user that the action failed : 
+				FMessageDialog::Open(EAppMsgType::Ok, Result.Reason);
+			}
 		}
 	}
 }
@@ -1139,7 +1166,10 @@ bool FLandscapeEditorCustomNodeBuilder_Layers::CanSetLayerAlpha(int32 InLayerInd
 {
 	FEdModeLandscape* LandscapeEdMode = GetEditorMode();
 	const FLandscapeLayer* Layer = LandscapeEdMode->GetLayer(InLayerIndex);
-	check(Layer != nullptr);
+	if (Layer == nullptr)
+	{
+		return false;
+	}
 
 	if (Layer->bLocked)
 	{
@@ -1172,7 +1202,10 @@ bool FLandscapeEditorCustomNodeBuilder_Layers::CanToggleVisibility(int32 InLayer
 {
 	FEdModeLandscape* LandscapeEdMode = GetEditorMode();
 	const FLandscapeLayer* Layer = LandscapeEdMode->GetLayer(InLayerIndex);
-	check(Layer != nullptr);
+	if (Layer == nullptr)
+	{
+		return false;
+	}
 
 	if (Layer->bLocked)
 	{
