Landscape Mode: Disable Resize and Move tools instead of hiding them

diff --git a/Engine/Source/Editor/LandscapeEditor/Private/LandscapeEditorCommands.cpp b/Engine/Source/Editor/LandscapeEditor/Private/LandscapeEditorCommands.cpp
--- a/Engine/Source/Editor/LandscapeEditor/Private/LandscapeEditorCommands.cpp
+++ b/Engine/Source/Editor/LandscapeEditor/Private/LandscapeEditorCommands.cpp
@@ -40,7 +40,7 @@ void FLandscapeEditorCommands::RegisterCommands()
 	UI_COMMAND(NewLandscape, "New", "Create or import a new landscape", EUserInterfaceActionType::RadioButton, FInputChord());
 	NameToCommandMap.Add("Tool_NewLandscape", NewLandscape);
 
-	UI_COMMAND(ResizeLandscape, "Resize", "Change Component Size", EUserInterfaceActionType::RadioButton, FInputChord());
+	UI_COMMAND(ResizeLandscape, "Resize", "Change landscape component size.\n\nResize is not supported when using World Partition.", EUserInterfaceActionType::RadioButton, FInputChord());
 	NameToCommandMap.Add("Tool_ResizeLandscape", ResizeLandscape);
 
 	UI_COMMAND(ImportExportTool, "Import", "Import or Export landscape data", EUserInterfaceActionType::RadioButton, FInputChord());
@@ -91,7 +91,7 @@ void FLandscapeEditorCommands::RegisterCommands()
 	UI_COMMAND(DeleteComponentTool, "Delete", "Delete components from the landscape, leaving a hole", EUserInterfaceActionType::RadioButton, FInputChord());
 	NameToCommandMap.Add("Tool_DeleteComponent", DeleteComponentTool);
 
-	UI_COMMAND(MoveToLevelTool, "Move", "Move landscape components to a landscape proxy in the currently active streaming level, so that they can be streamed in/out independently of the rest of the landscape", EUserInterfaceActionType::RadioButton, FInputChord());
+	UI_COMMAND(MoveToLevelTool, "Move", "Move landscape components to a landscape proxy in the currently active streaming level, so that they can be streamed in/out independently of the rest of the landscape.\n\nMove is not supported when using World Partition.", EUserInterfaceActionType::RadioButton, FInputChord());
 	NameToCommandMap.Add("Tool_MoveToLevel", MoveToLevelTool);
 
 	UI_COMMAND(RegionSelectTool, "Select", "Select a region of landscape to use as a mask for other tools", EUserInterfaceActionType::RadioButton, FInputChord());
diff --git a/Engine/Source/Editor/LandscapeEditor/Private/SLandscapeEditor.cpp b/Engine/Source/Editor/LandscapeEditor/Private/SLandscapeEditor.cpp
--- a/Engine/Source/Editor/LandscapeEditor/Private/SLandscapeEditor.cpp
+++ b/Engine/Source/Editor/LandscapeEditor/Private/SLandscapeEditor.cpp
@@ -213,13 +213,8 @@ void FLandscapeToolKit::BuildToolPalette(FName PaletteName, class FToolBarBuilde
 		ToolBarBuilder.AddToolBarButton(Commands.SelectComponentTool);
 		ToolBarBuilder.AddToolBarButton(Commands.AddComponentTool);
 		ToolBarBuilder.AddToolBarButton(Commands.DeleteComponentTool);
-		// MoveToLevel isn't supported because in GridBased worlds don't support Proxies in different Levels
-		// Resize isn't supported and instead should be done through a Commandlet for GridBased worlds
-		if (!LandscapeEdMode->IsGridBased())
-		{
-			ToolBarBuilder.AddToolBarButton(Commands.MoveToLevelTool);
-			ToolBarBuilder.AddToolBarButton(Commands.ResizeLandscape);
-		}
+		ToolBarBuilder.AddToolBarButton(Commands.MoveToLevelTool);
+		ToolBarBuilder.AddToolBarButton(Commands.ResizeLandscape);
 		ToolBarBuilder.AddToolBarButton(Commands.SplineTool);
 	}
 
@@ -367,7 +362,23 @@ bool FLandscapeToolKit::IsToolEnabled(FName ToolName) const
 	FEdModeLandscape* LandscapeEdMode = GetEditorMode();
 	if (LandscapeEdMode != nullptr)
 	{
-		if (ToolName == "NewLandscape" || LandscapeEdMode->GetLandscapeList().Num() > 0)
+		// When using World Partition:
+		// MoveToLevel isn't supported because we don't support Proxies in different levels.
+		// Resize isn't supported and instead should be done via a user provided Commandlet.
+		if (LandscapeEdMode->IsGridBased() &&
+			(ToolName == "MoveToLevel" || ToolName == "ResizeLandscape"))
+		{
+			return false;
+		}
+
+		// NewLandscape is always available.
+		if (ToolName == "NewLandscape")
+		{
+			return true;
+		}
+
+		// Other tools are available if there is an existing landscape.
+		if (LandscapeEdMode->GetLandscapeList().Num() > 0)
 		{
 			return true;
 		}
