Add test to verify that we can add a new Edit Layer to Landscape
#rnx

diff --git a/Engine/Source/Editor/LandscapeEditor/Private/Tests/LandscapeAutomationTests.cpp b/Engine/Source/Editor/LandscapeEditor/Private/Tests/LandscapeAutomationTests.cpp
--- a/Engine/Source/Editor/LandscapeEditor/Private/Tests/LandscapeAutomationTests.cpp
+++ b/Engine/Source/Editor/LandscapeEditor/Private/Tests/LandscapeAutomationTests.cpp
@@ -53,6 +53,7 @@ namespace LandscapeTestUtils
 	struct LandscapeTestCommands
 	{
 		static void Import(const FString& HeightMapFilenamee);
+		static bool CreateNewMapWithLandscape(int32 ComponentCountXY, int32 QuadsPerComponent);
 	};
 
 	void LandscapeTestCommands::Import(const FString& HeightMapFilename)
@@ -72,38 +73,52 @@ namespace LandscapeTestUtils
 
 		Customization_ImportExportLandscape->OnImportExportButtonClicked();
 	}
-}
 
-/**
-* Latent command to create a new landscape
-*/
-DEFINE_LATENT_AUTOMATION_COMMAND_TWO_PARAMETER(FCreateLandscapeCommand,int32,NumComponents,int32, QuadsPerComponent);
-bool FCreateLandscapeCommand::Update()
-{
-	//Switch to the Landscape tool
-	GLevelEditorModeTools().ActivateMode(FBuiltinEditorModes::EM_Landscape);
-	FEdModeLandscape* LandscapeEdMode = (FEdModeLandscape*)GLevelEditorModeTools().GetActiveMode(FBuiltinEditorModes::EM_Landscape);
+	bool LandscapeTestCommands::CreateNewMapWithLandscape(int32 ComponentCountXY, int32 QuadsPerComponent)
+	{
+		UE_LOG(LogLandscapeAutomationTests, Display, TEXT("Creating a new map..."));
+		UWorld* NewMap = FAutomationEditorCommonUtils::CreateNewMap();
+		if (!NewMap)
+		{
+			UE_LOG(LogLandscapeAutomationTests, Error, TEXT("Unable to create new map"));
+			return false;
+		}
+
+		UE_LOG(LogLandscapeAutomationTests, Display, TEXT("Switching to Landscape Editor Mode..."));
+		GLevelEditorModeTools().ActivateMode(FBuiltinEditorModes::EM_Landscape);
+		FEdModeLandscape* LandscapeEditMode = (FEdModeLandscape*)GLevelEditorModeTools().GetActiveMode(FBuiltinEditorModes::EM_Landscape);
+		if (!LandscapeEditMode)
+		{
+			UE_LOG(LogLandscapeAutomationTests, Error, TEXT("Unable to enter Landscape Edit Mode"));
+			return false;
+		}
 
-	//Modify the "Section size"
-	LandscapeEdMode->UISettings->NewLandscape_QuadsPerSection = QuadsPerComponent;
+		UE_LOG(LogLandscapeAutomationTests, Display, TEXT("Creating a new landscape..."));
+		LandscapeEditMode->UISettings->NewLandscape_QuadsPerSection = QuadsPerComponent;
+		LandscapeEditMode->UISettings->NewLandscape_ComponentCount.X = ComponentCountXY;
+		LandscapeEditMode->UISettings->NewLandscape_ComponentCount.Y = ComponentCountXY;
+		LandscapeEditMode->UISettings->NewLandscape_ClampSize();
 
-	LandscapeEdMode->UISettings->NewLandscape_ComponentCount.X = NumComponents;
-	LandscapeEdMode->UISettings->NewLandscape_ComponentCount.Y = NumComponents;
-	LandscapeEdMode->UISettings->NewLandscape_ClampSize();
+		TSharedPtr<FLandscapeEditorDetailCustomization_NewLandscape> Customization_NewLandscape = MakeShareable(new FLandscapeEditorDetailCustomization_NewLandscape);
+		Customization_NewLandscape->OnCreateButtonClicked();
 
-	//Create the landscape
-	TSharedPtr<FLandscapeEditorDetailCustomization_NewLandscape> Customization_NewLandscape = MakeShareable(new FLandscapeEditorDetailCustomization_NewLandscape);
-	Customization_NewLandscape->OnCreateButtonClicked();
+		if (!LandscapeEditMode->GetLandscape())
+		{
+			UE_LOG(LogLandscapeAutomationTests, Error, TEXT("Unable to create a new landscape"));
+			return false;
+		}
 
-	if (LandscapeEdMode->CurrentToolTarget.LandscapeInfo.IsValid())
-	{
-		UE_LOG(LogLandscapeAutomationTests, Display, TEXT("Created a new landscape"));
-	}
-	else
-	{
-		UE_LOG(LogLandscapeAutomationTests, Error, TEXT("Failed to create a new landscape"));
+		return true;
 	}
+}
 
+/**
+* Latent command to create a new landscape
+*/
+DEFINE_LATENT_AUTOMATION_COMMAND_TWO_PARAMETER(FCreateNewMapWithLandscapeCommand, int32, ComponentCountXY, int32, QuadsPerComponent);
+bool FCreateNewMapWithLandscapeCommand::Update()
+{
+	LandscapeTestUtils::LandscapeTestCommands::CreateNewMapWithLandscape(ComponentCountXY, QuadsPerComponent);
 	return true;
 }
 
@@ -199,22 +214,11 @@ bool FCheckHeight::Update()
 /**
 * Landscape creation / edit test
 */
-IMPLEMENT_SIMPLE_AUTOMATION_TEST(FLandscapeEditorTest, "Editor.Landscape CreateAndModify", EAutomationTestFlags::EditorContext | EAutomationTestFlags::NonNullRHI | EAutomationTestFlags::EngineFilter);
+IMPLEMENT_SIMPLE_AUTOMATION_TEST(FLandscapeEditorTest, "Editor.Landscape.Create and Modify New Landscape", EAutomationTestFlags::EditorContext | EAutomationTestFlags::NonNullRHI | EAutomationTestFlags::EngineFilter);
 bool FLandscapeEditorTest::RunTest(const FString& Parameters)
 {
-	//New level
-	UWorld* NewMap = FAutomationEditorCommonUtils::CreateNewMap();
-	if (NewMap)
-	{
-		UE_LOG(LogLandscapeAutomationTests, Display, TEXT("Created an empty level"));
-	}
-	else
-	{
-		UE_LOG(LogLandscapeAutomationTests, Error, TEXT("Failed to create an empty level"));
-		return false;
-	}
-
-	ADD_LATENT_AUTOMATION_COMMAND(FCreateLandscapeCommand(/* NumComponents = */ 8, /* QuadsPerComponent = */ 7));
+	// Create a new map and landscape
+	ADD_LATENT_AUTOMATION_COMMAND(FCreateNewMapWithLandscapeCommand(/* NumComponents = */ 8, /* QuadsPerComponent = */ 7));
 
 	//For some reason the heightmap component takes a few ticks to register with the nav system.  We crash if we try to modify the heightmap before then.
 	ADD_LATENT_AUTOMATION_COMMAND(FWaitLatentCommand(1.0f));
@@ -224,22 +228,11 @@ bool FLandscapeEditorTest::RunTest(const FString& Parameters)
 	return true;
 }
 
-IMPLEMENT_SIMPLE_AUTOMATION_TEST(FLandscapeEditorTestFooBar, "Editor.Landscape Import", EAutomationTestFlags::EditorContext | EAutomationTestFlags::NonNullRHI | EAutomationTestFlags::EngineFilter);
+IMPLEMENT_SIMPLE_AUTOMATION_TEST(FLandscapeEditorTestFooBar, "Editor.Landscape.Import Landscape", EAutomationTestFlags::EditorContext | EAutomationTestFlags::NonNullRHI | EAutomationTestFlags::EngineFilter);
 bool FLandscapeEditorTestFooBar::RunTest(const FString& Parameters)
 {
-	//New level
-	UWorld* NewMap = FAutomationEditorCommonUtils::CreateNewMap();
-	if (NewMap)
-	{
-		UE_LOG(LogLandscapeAutomationTests, Display, TEXT("Created an empty level"));
-	}
-	else
-	{
-		UE_LOG(LogLandscapeAutomationTests, Error, TEXT("Failed to create an empty level"));
-		return false;
-	}
-
-	ADD_LATENT_AUTOMATION_COMMAND(FCreateLandscapeCommand(/* NumComponents = */ 8, /* QuadsPerComponent = */ 63));
+	// Create a new map and landscape
+	ADD_LATENT_AUTOMATION_COMMAND(FCreateNewMapWithLandscapeCommand(/* NumComponents = */ 8, /* QuadsPerComponent = */ 7));
 
 	//For some reason the heightmap component takes a few ticks to register with the nav system.  We crash if we try to modify the heightmap before then.
 	ADD_LATENT_AUTOMATION_COMMAND(FWaitLatentCommand(1.0f));
@@ -254,7 +247,30 @@ bool FLandscapeEditorTestFooBar::RunTest(const FString& Parameters)
 	return true;
 }
 
+/*
+ * Verify that we can create a new edit layer for a landscape
+ * See: TR = 290778158
+ */
+IMPLEMENT_SIMPLE_AUTOMATION_TEST(FLandscapeEditorCreateNewLayer, "Editor.Landscape.Create New Edit Layer", EAutomationTestFlags::EditorContext | EAutomationTestFlags::NonNullRHI | EAutomationTestFlags::EngineFilter);
+bool FLandscapeEditorCreateNewLayer::RunTest(const FString& Parameters)
+{
+	// Create a new map with landscape
+	LandscapeTestUtils::LandscapeTestCommands::CreateNewMapWithLandscape(8, 63);
 
+	UE_LOG(LogLandscapeAutomationTests, Display, TEXT("Verifying that new landscapes start with 1 layer..."));
+	FEdModeLandscape* LandscapeEditMode = (FEdModeLandscape*)GLevelEditorModeTools().GetActiveMode(FBuiltinEditorModes::EM_Landscape);
+	ALandscape* CurrentLandscape = LandscapeEditMode->GetLandscape();
+	TestNotNull("Current landscape should not be null", CurrentLandscape);
+	TestEqual("New landscapes should start with one edit layer", LandscapeEditMode->GetLayerCount(), 1);
+
+	UE_LOG(LogLandscapeAutomationTests, Display, TEXT("Verifying that we successfully added a new layer..."));
+	FName LayerName = TEXT("TestLayer");
+	CurrentLandscape->CreateLayer(LayerName);
+	TestEqual("The landscape should now have a new edit layer", LandscapeEditMode->GetLayerCount(), 2);
+	TestNotNull("There should be a new layer with the specified name", CurrentLandscape->GetLayerConst(LayerName));
+
+	return true;
+}
 
 #undef LOCTEXT_NAMESPACE
 
