Fixed landscape flatten eye dropper tool canceling when clicking on the landscape while keeping the mouse still.

diff --git a/Engine/Source/Editor/LandscapeEditor/Private/LandscapeEdMode.cpp b/Engine/Source/Editor/LandscapeEditor/Private/LandscapeEdMode.cpp
--- a/Engine/Source/Editor/LandscapeEditor/Private/LandscapeEdMode.cpp
+++ b/Engine/Source/Editor/LandscapeEditor/Private/LandscapeEdMode.cpp
@@ -855,6 +855,39 @@ void FEdModeLandscape::Tick(FEditorViewportClient* ViewportClient, float DeltaTi
 	}
 }
 
+/** FEdMode: Called when the mouse enters the viewport area */
+bool FEdModeLandscape::MouseEnter(FEditorViewportClient* InViewportClient, FViewport* InViewport, int32 MouseX, int32 MouseY)
+{
+	if (!IsEditingEnabled())
+	{
+		return false;
+	}
+
+	bool bHandled = false;
+	if (CurrentTool)
+	{
+		bHandled = CurrentTool->MouseEnter(InViewportClient, InViewport, MouseX, MouseY);
+	}
+
+	return bHandled;
+}
+
+/** FEdMode: Called when the mouse exits the viewport area */
+bool FEdModeLandscape::MouseLeave(FEditorViewportClient * InViewportClient, FViewport * InViewport)
+{
+	if (!IsEditingEnabled())
+	{
+		return false;
+	}
+
+	bool bHandled = false;
+	if (CurrentTool)
+	{
+		bHandled = CurrentTool->MouseLeave(InViewportClient, InViewport);
+	}
+
+	return bHandled;
+}
 
 /** FEdMode: Called when the mouse is moved over the viewport */
 bool FEdModeLandscape::MouseMove(FEditorViewportClient* InViewportClient, FViewport* InViewport, int32 MouseX, int32 MouseY)
diff --git a/Engine/Source/Editor/LandscapeEditor/Private/LandscapeEdMode.h b/Engine/Source/Editor/LandscapeEditor/Private/LandscapeEdMode.h
--- a/Engine/Source/Editor/LandscapeEditor/Private/LandscapeEdMode.h
+++ b/Engine/Source/Editor/LandscapeEditor/Private/LandscapeEdMode.h
@@ -381,6 +381,12 @@ public:
 	/** FEdMode: Called when the mode is exited */
 	virtual void Exit() override;
 
+	/** FEdMode: Called when the mouse enters the viewport area */
+	virtual bool MouseEnter(FEditorViewportClient* InViewportClient, FViewport* InViewport, int32 MouseX, int32 MouseY) override;
+
+	/** FEdMode: Called when the mouse exits the viewport area */
+	virtual bool MouseLeave(FEditorViewportClient* InViewportClient, FViewport* Viewport) override;
+
 	/** FEdMode: Called when the mouse is moved over the viewport */
 	virtual bool MouseMove(FEditorViewportClient* ViewportClient, FViewport* Viewport, int32 x, int32 y) override;
 
diff --git a/Engine/Source/Editor/LandscapeEditor/Private/LandscapeEdModePaintTools.cpp b/Engine/Source/Editor/LandscapeEditor/Private/LandscapeEdModePaintTools.cpp
--- a/Engine/Source/Editor/LandscapeEditor/Private/LandscapeEdModePaintTools.cpp
+++ b/Engine/Source/Editor/LandscapeEditor/Private/LandscapeEdModePaintTools.cpp
@@ -1189,11 +1189,6 @@ public:
 
 		Super::Tick(ViewportClient, DeltaTime);
 
-		if (this->EdMode->UISettings->bFlattenEyeDropperModeActivated)
-		{
-			this->EdMode->UISettings->bFlattenEyeDropperModeMousingOverViewport = false;
-		}
-
 		if (HeightmapFlattenPreviewComponent != nullptr)
 		{
 			bool bShowGrid = this->EdMode->UISettings->bUseFlattenTarget && this->EdMode->UISettings->bShowFlattenTargetPreview;
@@ -1207,17 +1202,30 @@ public:
 		}
 	}
 
+	virtual bool MouseEnter(FEditorViewportClient* InViewportClient, FViewport* Viewport, int32 InMouseX, int32 InMouseY) override
+	{
+		if (InViewportClient->IsLevelEditorClient() && this->EdMode->UISettings->bFlattenEyeDropperModeActivated)
+		{
+			this->EdMode->UISettings->bFlattenEyeDropperModeMousingOverViewport = true;
+		}
+		return true;
+	}
+
+	virtual bool MouseLeave(FEditorViewportClient* InViewportClient, FViewport* Viewport) override
+	{
+		if (InViewportClient->IsLevelEditorClient() && this->EdMode->UISettings->bFlattenEyeDropperModeActivated)
+		{
+			this->EdMode->UISettings->bFlattenEyeDropperModeMousingOverViewport = false;
+		}
+		return true;
+	}
+
 	virtual bool MouseMove(FEditorViewportClient* ViewportClient, FViewport* Viewport, int32 x, int32 y) override
 	{
 		bool bResult = Super::MouseMove(ViewportClient, Viewport, x, y);
 
 		if (ViewportClient->IsLevelEditorClient())
 		{
-			if (this->EdMode->UISettings->bFlattenEyeDropperModeActivated)
-			{
-				this->EdMode->UISettings->bFlattenEyeDropperModeMousingOverViewport = true;
-			}
-
 			if (HeightmapFlattenPreviewComponent != nullptr)
 			{
 				FVector MousePosition;
diff --git a/Engine/Source/Editor/LandscapeEditor/Public/LandscapeToolInterface.h b/Engine/Source/Editor/LandscapeEditor/Public/LandscapeToolInterface.h
--- a/Engine/Source/Editor/LandscapeEditor/Public/LandscapeToolInterface.h
+++ b/Engine/Source/Editor/LandscapeEditor/Public/LandscapeToolInterface.h
@@ -199,6 +199,8 @@ public:
 	virtual bool BeginTool(FEditorViewportClient* ViewportClient, const FLandscapeToolTarget& Target, const FVector& InHitLocation) = 0;
 	virtual void EndTool(FEditorViewportClient* ViewportClient) = 0;
 	virtual void Tick(FEditorViewportClient* ViewportClient, float DeltaTime) {};
+	virtual bool MouseEnter(FEditorViewportClient* InViewportClient, FViewport* InViewport, int32 MouseX, int32 MouseY) { return true; }
+	virtual bool MouseLeave(FEditorViewportClient* InViewportClient, FViewport* InViewport) { return true; }
 	virtual bool MouseMove(FEditorViewportClient* ViewportClient, FViewport* Viewport, int32 x, int32 y) = 0;
 	virtual bool HandleClick(HHitProxy* HitProxy, const FViewportClick& Click) { return false; }
 	virtual bool InputKey(FEditorViewportClient* InViewportClient, FViewport* InViewport, FKey InKey, EInputEvent InEvent) { return false; }
