Make default landscape creation location be the origin.

diff --git a/Engine/Source/Editor/LandscapeEditor/Private/LandscapeEditorObject.cpp b/Engine/Source/Editor/LandscapeEditor/Private/LandscapeEditorObject.cpp
--- a/Engine/Source/Editor/LandscapeEditor/Private/LandscapeEditorObject.cpp
+++ b/Engine/Source/Editor/LandscapeEditor/Private/LandscapeEditorObject.cpp
@@ -23,7 +23,7 @@ static TAutoConsoleVariable<bool> CVarLandscapeSimulateAlphaBrushTextureLoadFail
 	false,
 	TEXT("Debug utility to simulate a loading failure (e.g. invalid source data, which can happen in cooked editor or with a badly virtualized texture) when loading the alpha brush texture"));
 
-const FVector ULandscapeEditorObject::NewLandscape_DefaultLocation = FVector(0, 0, 100);
+const FVector ULandscapeEditorObject::NewLandscape_DefaultLocation = FVector(0, 0, 0);
 const FRotator ULandscapeEditorObject::NewLandscape_DefaultRotation = FRotator::ZeroRotator;
 const FVector ULandscapeEditorObject::NewLandscape_DefaultScale = FVector(100, 100, 100);
 
