Fix for failing landscape test

Test is creating landscape at the default location, which was changed to the origin

diff --git a/Engine/Source/Editor/LandscapeEditor/Private/Tests/LandscapeAutomationTests.cpp b/Engine/Source/Editor/LandscapeEditor/Private/Tests/LandscapeAutomationTests.cpp
--- a/Engine/Source/Editor/LandscapeEditor/Private/Tests/LandscapeAutomationTests.cpp
+++ b/Engine/Source/Editor/LandscapeEditor/Private/Tests/LandscapeAutomationTests.cpp
@@ -205,7 +205,7 @@ bool FCheckHeight::Update()
 	{
 		TOptional<float> Height = InfoMapPair.Value->GetLandscapeProxy()->GetHeightAtLocation(FVector(0, 0, 0), EHeightfieldSource::Editor);
 		CurrentTest->TestEqual("Has Height Value at 0,0", Height.IsSet(), true);
-		CurrentTest->TestNearlyEqual("Height Value at 0,0 is 0", Height.GetValue(), 100.0f, 1e-4f);
+		CurrentTest->TestNearlyEqual("Height Value at 0,0 is 0", Height.GetValue(), 0.0f, 1e-4f);
 		break;
 	}
 	return true;
