Added info message when there are no target layers in landscape mode, to help the user understand the workflow now requires him to add layers either using the "Add" or the "Populate from materials" button

diff --git a/Engine/Source/Editor/LandscapeEditor/Private/LandscapeEditorDetailCustomization_Layers.cpp b/Engine/Source/Editor/LandscapeEditor/Private/LandscapeEditorDetailCustomization_Layers.cpp
--- a/Engine/Source/Editor/LandscapeEditor/Private/LandscapeEditorDetailCustomization_Layers.cpp
+++ b/Engine/Source/Editor/LandscapeEditor/Private/LandscapeEditorDetailCustomization_Layers.cpp
@@ -74,7 +74,7 @@ void FLandscapeEditorDetailCustomization_Layers::CustomizeDetails(IDetailLayoutB
 			LayerCategory.AddCustomBuilder(MakeShareable(new FLandscapeEditorCustomNodeBuilder_Layers(DetailBuilder.GetThumbnailPool().ToSharedRef())));
 
 			LayerCategory.AddCustomRow(FText())
-				.Visibility(TAttribute<EVisibility>::Create(TAttribute<EVisibility>::FGetter::CreateLambda([]() { return ShoudShowLayersErrorMessageTip() ? EVisibility::Visible : EVisibility::Collapsed; })))
+				.Visibility(TAttribute<EVisibility>::Create(TAttribute<EVisibility>::FGetter::CreateLambda([]() { return ShouldShowLayersErrorMessageTip() ? EVisibility::Visible : EVisibility::Collapsed; })))
 				[
 					SNew(SMultiLineEditableTextBox)
 					.IsReadOnly(true)
@@ -88,7 +88,7 @@ void FLandscapeEditorDetailCustomization_Layers::CustomizeDetails(IDetailLayoutB
 }
 END_SLATE_FUNCTION_BUILD_OPTIMIZATION
 
-bool FLandscapeEditorDetailCustomization_Layers::ShoudShowLayersErrorMessageTip()
+bool FLandscapeEditorDetailCustomization_Layers::ShouldShowLayersErrorMessageTip()
 {
 	FEdModeLandscape* LandscapeEdMode = GetEditorMode();
 	if (LandscapeEdMode && LandscapeEdMode->DoesCurrentToolAffectEditLayers())
diff --git a/Engine/Source/Editor/LandscapeEditor/Private/LandscapeEditorDetailCustomization_Layers.h b/Engine/Source/Editor/LandscapeEditor/Private/LandscapeEditorDetailCustomization_Layers.h
--- a/Engine/Source/Editor/LandscapeEditor/Private/LandscapeEditorDetailCustomization_Layers.h
+++ b/Engine/Source/Editor/LandscapeEditor/Private/LandscapeEditorDetailCustomization_Layers.h
@@ -43,7 +43,7 @@ public:
 	virtual void CustomizeDetails(IDetailLayoutBuilder& DetailBuilder) override;
 
 protected:
-	static bool ShoudShowLayersErrorMessageTip();
+	static bool ShouldShowLayersErrorMessageTip();
 	static FText GetLayersErrorMessageText();
 };
 
diff --git a/Engine/Source/Editor/LandscapeEditor/Private/LandscapeEditorDetailCustomization_TargetLayers.cpp b/Engine/Source/Editor/LandscapeEditor/Private/LandscapeEditorDetailCustomization_TargetLayers.cpp
--- a/Engine/Source/Editor/LandscapeEditor/Private/LandscapeEditorDetailCustomization_TargetLayers.cpp
+++ b/Engine/Source/Editor/LandscapeEditor/Private/LandscapeEditorDetailCustomization_TargetLayers.cpp
@@ -60,7 +60,7 @@ void FLandscapeEditorDetailCustomization_TargetLayers::CustomizeDetails(IDetailL
 	PropertyHandle_TargetDisplayOrder->MarkHiddenByCustomization();
 
 	TSharedRef<IPropertyHandle> PropertyHandle_TargetShowUnusedLayers = DetailBuilder.GetProperty(GET_MEMBER_NAME_CHECKED(ULandscapeEditorObject, ShowUnusedLayers));
-	PropertyHandle_TargetShowUnusedLayers->MarkHiddenByCustomization();	
+	PropertyHandle_TargetShowUnusedLayers->MarkHiddenByCustomization();
 
 	if (!ShouldShowTargetLayers())
 	{
@@ -74,22 +74,33 @@ void FLandscapeEditorDetailCustomization_TargetLayers::CustomizeDetails(IDetailL
 	check(LandscapeEdMode);
 
 	TargetsCategory.AddProperty(PropertyHandle_PaintingRestriction)
-	.Visibility(TAttribute<EVisibility>::Create(TAttribute<EVisibility>::FGetter::CreateStatic(&FLandscapeEditorDetailCustomization_TargetLayers::GetVisibility_PaintingRestriction)))
-	.IsEnabled(TAttribute<bool>::Create(TAttribute<bool>::FGetter::CreateLambda([LandscapeEdMode](){ return LandscapeEdMode->HasValidLandscapeEditLayerSelection(); })));
-		
+		.Visibility(TAttribute<EVisibility>::Create(TAttribute<EVisibility>::FGetter::CreateStatic(&FLandscapeEditorDetailCustomization_TargetLayers::GetPaintingRestrictionVisibility)))
+		.IsEnabled(TAttribute<bool>::Create(TAttribute<bool>::FGetter::CreateLambda([LandscapeEdMode]() { return LandscapeEdMode->HasValidLandscapeEditLayerSelection(); })));
+
 	TargetsCategory.AddCustomRow(FText())
-	.Visibility(TAttribute<EVisibility>::Create(TAttribute<EVisibility>::FGetter::CreateStatic(&FLandscapeEditorDetailCustomization_TargetLayers::GetVisibility_VisibilityTip)))
-	[
-		SNew(SMultiLineEditableTextBox)
-		.IsReadOnly(true)
-		.Font(DetailBuilder.GetDetailFontBold())
-		.BackgroundColor(TAttribute<FSlateColor>::Create(TAttribute<FSlateColor>::FGetter::CreateLambda([]() { return FAppStyle::GetColor("ErrorReporting.WarningBackgroundColor"); })))
-		.Text(LOCTEXT("Visibility_Tip", "Note: There are some areas where visibility painting is disabled because Component/Proxy don't have a \"Landscape Visibility Mask\" node in their material."))
-		.AutoWrapText(true)
-		.IsEnabled(TAttribute<bool>::Create(TAttribute<bool>::FGetter::CreateLambda([LandscapeEdMode]() { return LandscapeEdMode->HasValidLandscapeEditLayerSelection(); })))
-	];
+		.Visibility(TAttribute<EVisibility>::Create(TAttribute<EVisibility>::FGetter::CreateStatic(&FLandscapeEditorDetailCustomization_TargetLayers::GetVisibilityMaskTipVisibility)))
+		[
+			SNew(SMultiLineEditableTextBox)
+				.IsReadOnly(true)
+				.Font(DetailBuilder.GetDetailFontBold())
+				.BackgroundColor(FAppStyle::GetColor("ErrorReporting.WarningBackgroundColor"))
+				.Text(LOCTEXT("Visibility_Tip", "Note: There are some areas where visibility painting is disabled because Component/Proxy don't have a \"Landscape Visibility Mask\" node in their material."))
+				.AutoWrapText(true)
+				.IsEnabled(TAttribute<bool>::Create(TAttribute<bool>::FGetter::CreateLambda([LandscapeEdMode]() { return LandscapeEdMode->HasValidLandscapeEditLayerSelection(); })))
+		];
 
 	TargetsCategory.AddCustomBuilder(MakeShareable(new FLandscapeEditorCustomNodeBuilder_TargetLayers(DetailBuilder.GetThumbnailPool().ToSharedRef(), PropertyHandle_TargetDisplayOrder, PropertyHandle_TargetShowUnusedLayers)));
+
+	TargetsCategory.AddCustomRow(FText())
+		.Visibility(TAttribute<EVisibility>::Create(TAttribute<EVisibility>::FGetter::CreateStatic(&FLandscapeEditorDetailCustomization_TargetLayers::GetPopulateTargetLayersInfoTipVisibility)))
+		[
+			SNew(SMultiLineEditableTextBox)
+				.IsReadOnly(true)
+				.Font(DetailBuilder.GetDetailFontBold())
+				.BackgroundColor(FAppStyle::GetColor("InfoReporting.BackgroundColor"))
+				.Text(LOCTEXT("PopulateTargetLayers_Tip", "There are currently no target layers assigned to this landscape. Use the buttons above to add new ones or populate them from the material(s) currently assigned to the landscape"))
+				.AutoWrapText(true)
+		];
 }
 END_SLATE_FUNCTION_BUILD_OPTIMIZATION
 
@@ -115,7 +126,7 @@ bool FLandscapeEditorDetailCustomization_TargetLayers::ShouldShowTargetLayers()
 	return false;
 }
 
-bool FLandscapeEditorDetailCustomization_TargetLayers::ShouldShowPaintingRestriction()
+EVisibility FLandscapeEditorDetailCustomization_TargetLayers::GetPaintingRestrictionVisibility()
 {
 	FEdModeLandscape* LandscapeEdMode = GetEditorMode();
 
@@ -128,19 +139,14 @@ bool FLandscapeEditorDetailCustomization_TargetLayers::ShouldShowPaintingRestric
 			|| (LandscapeEdMode->CurrentToolTarget.TargetType == ELandscapeToolTargetType::Invalid)
 			|| (LandscapeEdMode->CurrentToolTarget.TargetType == ELandscapeToolTargetType::Visibility))
 		{
-			return true;
+			return EVisibility::Visible;
 		}
 	}
 
-	return false;
-}
-
-EVisibility FLandscapeEditorDetailCustomization_TargetLayers::GetVisibility_PaintingRestriction()
-{
-	return ShouldShowPaintingRestriction() ? EVisibility::Visible : EVisibility::Collapsed;
+	return EVisibility::Collapsed;
 }
 
-bool FLandscapeEditorDetailCustomization_TargetLayers::ShouldShowVisibilityTip()
+EVisibility FLandscapeEditorDetailCustomization_TargetLayers::GetVisibilityMaskTipVisibility()
 {
 	FEdModeLandscape* LandscapeEdMode = GetEditorMode();
 	if (LandscapeEdMode && LandscapeEdMode->CurrentToolTarget.LandscapeInfo.IsValid())
@@ -154,18 +160,30 @@ bool FLandscapeEditorDetailCustomization_TargetLayers::ShouldShowVisibilityTip()
 				bHasValidHoleMaterial &= LandscapeComponent->IsLandscapeHoleMaterialValid();
 			});
 
-			return !bHasValidHoleMaterial;
+			return bHasValidHoleMaterial ? EVisibility::Collapsed : EVisibility::Visible;
 		}
 	}
 
-	return false;
+	return EVisibility::Collapsed;
 }
 
-EVisibility FLandscapeEditorDetailCustomization_TargetLayers::GetVisibility_VisibilityTip()
+EVisibility FLandscapeEditorDetailCustomization_TargetLayers::GetPopulateTargetLayersInfoTipVisibility()
 {
-	return ShouldShowVisibilityTip() ? EVisibility::Visible : EVisibility::Collapsed;
+	FEdModeLandscape* LandscapeEdMode = GetEditorMode();
+	if (LandscapeEdMode && LandscapeEdMode->CurrentToolTarget.LandscapeInfo.IsValid())
+	{
+		if ((LandscapeEdMode->CurrentToolTarget.TargetType == ELandscapeToolTargetType::Weightmap)
+			|| (LandscapeEdMode->CurrentToolTarget.TargetType == ELandscapeToolTargetType::Invalid)) // ELandscapeToolTargetType::Invalid means "weightmap with no valid paint layer" 
+		{
+			ULandscapeInfo* LandscapeInfo = LandscapeEdMode->CurrentToolTarget.LandscapeInfo.Get();
+			return LandscapeInfo->Layers.IsEmpty() ? EVisibility::Visible : EVisibility::Collapsed;
+		}
+	}
+
+	return EVisibility::Collapsed;
 }
 
+
 //////////////////////////////////////////////////////////////////////////
 
 FEdModeLandscape* FLandscapeEditorCustomNodeBuilder_TargetLayers::GetEditorMode()
@@ -1320,9 +1338,7 @@ void FLandscapeEditorCustomNodeBuilder_TargetLayers::OnTargetLayerSetObject(cons
 				if (ensure(Index != INDEX_NONE))
 				{
 					FLandscapeInfoLayerSettings& LayerSettings = LandscapeInfo->Layers[Index];
-
 					LandscapeInfo->ReplaceLayer(LayerSettings.LayerInfoObj, SelectedLayerInfo);
-
 					LayerSettings.LayerInfoObj = SelectedLayerInfo;
 				}
 			}
diff --git a/Engine/Source/Editor/LandscapeEditor/Private/LandscapeEditorDetailCustomization_TargetLayers.h b/Engine/Source/Editor/LandscapeEditor/Private/LandscapeEditorDetailCustomization_TargetLayers.h
--- a/Engine/Source/Editor/LandscapeEditor/Private/LandscapeEditorDetailCustomization_TargetLayers.h
+++ b/Engine/Source/Editor/LandscapeEditor/Private/LandscapeEditorDetailCustomization_TargetLayers.h
@@ -38,10 +38,9 @@ public:
 
 protected:
 	static bool ShouldShowTargetLayers();
-	static bool ShouldShowPaintingRestriction();
-	static EVisibility GetVisibility_PaintingRestriction();
-	static bool ShouldShowVisibilityTip();
-	static EVisibility GetVisibility_VisibilityTip();
+	static EVisibility GetPaintingRestrictionVisibility();
+	static EVisibility GetVisibilityMaskTipVisibility();
+	static EVisibility GetPopulateTargetLayersInfoTipVisibility();
 };
 
 class FLandscapeEditorCustomNodeBuilder_TargetLayers : public IDetailCustomNodeBuilder, public TSharedFromThis<FLandscapeEditorCustomNodeBuilder_TargetLayers>
