Fix for landscape visibility layer merge resulting in corrupted data when there is invalid layer infos assigned to layers
Also fixes a crash when assigning layer infos

diff --git a/Engine/Source/Editor/LandscapeEditor/Private/LandscapeEditorDetailCustomization_TargetLayers.cpp b/Engine/Source/Editor/LandscapeEditor/Private/LandscapeEditorDetailCustomization_TargetLayers.cpp
--- a/Engine/Source/Editor/LandscapeEditor/Private/LandscapeEditorDetailCustomization_TargetLayers.cpp
+++ b/Engine/Source/Editor/LandscapeEditor/Private/LandscapeEditorDetailCustomization_TargetLayers.cpp
@@ -1330,7 +1330,10 @@ void FLandscapeEditorCustomNodeBuilder_TargetLayers::OnTargetLayerSetObject(cons
 			ULandscapeInfo* LandscapeInfo = Target->LandscapeInfo.Get();
 			ALandscape* LandscapeActor = LandscapeInfo->LandscapeActor.Get();
 			
-			LandscapeActor->AddTargetLayer(Target->LayerName, FLandscapeTargetLayerSettings(SelectedLayerInfo));
+			if (!LandscapeActor->HasTargetLayer(Target->GetLayerName()))
+			{
+				LandscapeActor->AddTargetLayer(Target->GetLayerName(), FLandscapeTargetLayerSettings(SelectedLayerInfo));
+			}
 			
 			if (Target->LayerInfoObj.IsValid())
 			{
diff --git a/Engine/Source/Runtime/Landscape/Private/LandscapeEditLayers.cpp b/Engine/Source/Runtime/Landscape/Private/LandscapeEditLayers.cpp
--- a/Engine/Source/Runtime/Landscape/Private/LandscapeEditLayers.cpp
+++ b/Engine/Source/Runtime/Landscape/Private/LandscapeEditLayers.cpp
@@ -5967,6 +5967,8 @@ void ALandscape::PrepareComponentDataToExtractMaterialLayersCS(const TArray<ULan
 									break;
 								}
 							}
+							// didn't find a matching layer
+							continue;
 						}
 
 						OutComponentData.Add(Data);
