Landscape code cleanup : removed unused stuff

#rnx

diff --git a/Engine/Source/Editor/LandscapeEditor/Private/LandscapeEdMode.cpp b/Engine/Source/Editor/LandscapeEditor/Private/LandscapeEdMode.cpp
--- a/Engine/Source/Editor/LandscapeEditor/Private/LandscapeEdMode.cpp
+++ b/Engine/Source/Editor/LandscapeEditor/Private/LandscapeEdMode.cpp
@@ -3182,7 +3182,7 @@ void FEdModeLandscape::ReimportData(const FLandscapeTargetListInfo& TargetInfo)
 }
 
 template<class T>
-void ImportDataInternal(ULandscapeInfo* LandscapeInfo, const FString& Filename, FName LayerName, bool bSingleFile, bool bFlipYAxis, const FIntRect& ImportRegionVerts, ELandscapeImportTransformType TransformType, FIntPoint Offset, TFunctionRef<void(int32, int32, int32, int32, const TArray<T>&)> SetDataFunc)
+void ImportDataInternal(ULandscapeInfo* LandscapeInfo, const FString& Filename, bool bSingleFile, bool bFlipYAxis, const FIntRect& ImportRegionVerts, ELandscapeImportTransformType TransformType, FIntPoint Offset, TFunctionRef<void(int32, int32, int32, int32, const TArray<T>&)> SetDataFunc)
 {
 	if (!LandscapeInfo)
 	{
@@ -3255,7 +3255,7 @@ void FEdModeLandscape::ImportHeightData(ULandscapeInfo* LandscapeInfo, const FGu
 {
 	TRACE_CPUPROFILER_EVENT_SCOPE(FEdModeLandscape::ImportHeightData);
 
-	ImportDataInternal<uint16>(LandscapeInfo, Filename, NAME_None, UseSingleFileImport(), bFlipYAxis, ImportRegionVerts, TransformType, Offset, [LandscapeInfo, LayerGuid, PaintRestriction](int32 MinX, int32 MinY, int32 MaxX, int32 MaxY, const TArray<uint16>& Data)
+	ImportDataInternal<uint16>(LandscapeInfo, Filename, UseSingleFileImport(), bFlipYAxis, ImportRegionVerts, TransformType, Offset, [LandscapeInfo, LayerGuid, PaintRestriction](int32 MinX, int32 MinY, int32 MaxX, int32 MaxY, const TArray<uint16>& Data)
 	{
 		ALandscape* Landscape = LandscapeInfo->LandscapeActor.Get();
 		FScopedSetLandscapeEditingLayer Scope(Landscape, LayerGuid, [&] { check(Landscape); Landscape->RequestLayersContentUpdate(ELandscapeLayerUpdateMode::Update_Heightmap_All); });
@@ -3271,7 +3271,7 @@ void FEdModeLandscape::ImportWeightData(ULandscapeInfo* LandscapeInfo, const FGu
 {
 	if (LayerInfo)
 	{
-		ImportDataInternal<uint8>(LandscapeInfo, Filename, LayerInfo->LayerName, UseSingleFileImport(), bFlipYAxis, ImportRegionVerts, TransformType, Offset, [LandscapeInfo, LayerGuid, LayerInfo, PaintRestriction](int32 MinX, int32 MinY, int32 MaxX, int32 MaxY, const TArray<uint8>& Data)
+		ImportDataInternal<uint8>(LandscapeInfo, Filename, UseSingleFileImport(), bFlipYAxis, ImportRegionVerts, TransformType, Offset, [LandscapeInfo, LayerGuid, LayerInfo, PaintRestriction](int32 MinX, int32 MinY, int32 MaxX, int32 MaxY, const TArray<uint8>& Data)
 		{
 			ALandscape* Landscape = LandscapeInfo->LandscapeActor.Get();
 			FScopedSetLandscapeEditingLayer Scope(Landscape, LayerGuid, [&] { check(Landscape); Landscape->RequestLayersContentUpdate(ELandscapeLayerUpdateMode::Update_Weightmap_All); });
diff --git a/Engine/Source/Editor/LandscapeEditor/Private/LandscapeEditorDetailCustomization_Layers.cpp b/Engine/Source/Editor/LandscapeEditor/Private/LandscapeEditorDetailCustomization_Layers.cpp
--- a/Engine/Source/Editor/LandscapeEditor/Private/LandscapeEditorDetailCustomization_Layers.cpp
+++ b/Engine/Source/Editor/LandscapeEditor/Private/LandscapeEditorDetailCustomization_Layers.cpp
@@ -118,7 +118,6 @@ FEdModeLandscape* FLandscapeEditorCustomNodeBuilder_Layers::GetEditorMode()
 
 FLandscapeEditorCustomNodeBuilder_Layers::FLandscapeEditorCustomNodeBuilder_Layers(TSharedRef<FAssetThumbnailPool> InThumbnailPool)
 	: ThumbnailPool(InThumbnailPool)
-	, CurrentEditingInlineTextBlock(INDEX_NONE)
 	, CurrentSlider(INDEX_NONE)
 {
 }
@@ -304,8 +303,6 @@ TSharedPtr<SWidget> FLandscapeEditorCustomNodeBuilder_Layers::GenerateRow(int32
 								.ColorAndOpacity(TAttribute<FSlateColor>::Create(TAttribute<FSlateColor>::FGetter::CreateSP(this, &FLandscapeEditorCustomNodeBuilder_Layers::GetLayerTextColor, InLayerIndex)))
 								.ToolTipText(MakeAttributeSPLambda(this, [this, InLayerIndex] { FText Reason; CanRenameLayer(InLayerIndex, Reason); return Reason; }))
 								.OnVerifyTextChanged(FOnVerifyTextChanged::CreateSP(this, &FLandscapeEditorCustomNodeBuilder_Layers::CanRenameLayerTo, InLayerIndex))
-								.OnEnterEditingMode(this, &FLandscapeEditorCustomNodeBuilder_Layers::OnBeginNameTextEdit)
-								.OnExitEditingMode(this, &FLandscapeEditorCustomNodeBuilder_Layers::OnEndNameTextEdit)
 								.OnTextCommitted(FOnTextCommitted::CreateSP(this, &FLandscapeEditorCustomNodeBuilder_Layers::SetLayerName, InLayerIndex))
 						]
 				]
@@ -400,17 +397,6 @@ bool FLandscapeEditorCustomNodeBuilder_Layers::IsLayerSelected(int32 InLayerInde
 	return false;
 }
 
-void FLandscapeEditorCustomNodeBuilder_Layers::OnBeginNameTextEdit()
-{
-	FEdModeLandscape* LandscapeEdMode = GetEditorMode();
-	CurrentEditingInlineTextBlock = LandscapeEdMode ? LandscapeEdMode->GetCurrentLayerIndex() : INDEX_NONE;
-}
-
-void FLandscapeEditorCustomNodeBuilder_Layers::OnEndNameTextEdit()
-{
-	CurrentEditingInlineTextBlock = INDEX_NONE;
-}
-
 bool FLandscapeEditorCustomNodeBuilder_Layers::CanRenameLayerTo(const FText& InNewText, FText& OutErrorMessage, int32 InLayerIndex)
 {
 	FEdModeLandscape* LandscapeEdMode = GetEditorMode();
diff --git a/Engine/Source/Editor/LandscapeEditor/Private/LandscapeEditorDetailCustomization_Layers.h b/Engine/Source/Editor/LandscapeEditor/Private/LandscapeEditorDetailCustomization_Layers.h
--- a/Engine/Source/Editor/LandscapeEditor/Private/LandscapeEditorDetailCustomization_Layers.h
+++ b/Engine/Source/Editor/LandscapeEditor/Private/LandscapeEditorDetailCustomization_Layers.h
@@ -106,8 +106,6 @@ protected:
 
 	void ShowOnlySelectedLayer(int32 InLayerIndex);
 	void ShowAllLayers();
-	void OnBeginNameTextEdit();
-	void OnEndNameTextEdit();
 	void SetLayerName(const FText& InText, ETextCommit::Type InCommitType, int32 InLayerIndex);
 	FText GetLayerText(int32 InLayerIndex) const;
 	FSlateColor GetLayerTextColor(int32 InLayerIndex) const;
@@ -137,7 +135,6 @@ private:
 	/** Widgets for displaying and editing the layer name */
 	TArray< TSharedPtr< SInlineEditableTextBlock > > InlineTextBlocks;
 
-	int32 CurrentEditingInlineTextBlock;
 	int32 CurrentSlider;
 };
 
