* Improvement to the landscape target layers deprecation : if possible, try to deprecate the list of existing target layers instead of entirely relying on the analysis of the layer allocations within the landscape components. This allows not to lose landscape layer info pointers when loading map where a given layer is not used (i.e. there's no weightmap allocation associated with this layer)
* Prevent landscape target layers from being renamed to a layer name that's already in use

diff --git a/Engine/Source/Editor/LandscapeEditor/Private/LandscapeEditorDetailCustomization_TargetLayers.cpp b/Engine/Source/Editor/LandscapeEditor/Private/LandscapeEditorDetailCustomization_TargetLayers.cpp
--- a/Engine/Source/Editor/LandscapeEditor/Private/LandscapeEditorDetailCustomization_TargetLayers.cpp
+++ b/Engine/Source/Editor/LandscapeEditor/Private/LandscapeEditorDetailCustomization_TargetLayers.cpp
@@ -216,7 +216,7 @@ void FLandscapeEditorCustomNodeBuilder_TargetLayers::GenerateHeaderRowContent(FD
 {
 	FEdModeLandscape* LandscapeEdMode = GetEditorMode();
 	
-	if (LandscapeEdMode == NULL)
+	if (LandscapeEdMode == nullptr)
 	{
 		return;	
 	}
@@ -333,10 +333,10 @@ void FLandscapeEditorCustomNodeBuilder_TargetLayers::GenerateHeaderRowContent(FD
 
 FReply FLandscapeEditorCustomNodeBuilder_TargetLayers::HandleCreateLayersFromMaterials()
 {
-	FScopedTransaction Transaction(LOCTEXT("LandscapeWeightLayer_CreateFromMaterials", "Create Weight Layers from Assigned materials"));
+	FScopedTransaction Transaction(LOCTEXT("LandscapeTargetLayer_CreateFromMaterials", "Create Target Layers from Assigned materials"));
 	
 	FEdModeLandscape* LandscapeEdMode = GetEditorMode(); 
-	if (LandscapeEdMode == NULL)
+	if (LandscapeEdMode == nullptr)
 	{
 		return FReply::Handled(); 
 	}
@@ -346,7 +346,7 @@ FReply FLandscapeEditorCustomNodeBuilder_TargetLayers::HandleCreateLayersFromMat
 	TSet<FName> LayerNames;
 	LandscapeActor->GetLandscapeInfo()->ForEachLandscapeProxy([&LayerNames](ALandscapeProxy* Proxy)
 	{
-		LayerNames.Append(Proxy->RetrieveAllLayerNamesFromMaterials());
+		LayerNames.Append(Proxy->RetrieveTargetLayerNamesFromMaterials());
 		return true;
 	});
 
@@ -372,7 +372,7 @@ FReply FLandscapeEditorCustomNodeBuilder_TargetLayers::HandleCreateLayersFromMat
 void FLandscapeEditorCustomNodeBuilder_TargetLayers::HandleCreateLayer()
 {
 	FEdModeLandscape* LandscapeEdMode = GetEditorMode();
-	if (LandscapeEdMode == NULL)
+	if (LandscapeEdMode == nullptr)
 	{
 		return; 
 	}
@@ -384,7 +384,7 @@ void FLandscapeEditorCustomNodeBuilder_TargetLayers::HandleCreateLayer()
 		return; 
 	}
 
-	FScopedTransaction Transaction(LOCTEXT("LandscapeWeightLayer_Create", "Create a Weight Layer"));
+	FScopedTransaction Transaction(LOCTEXT("LandscapeTargetLayer_Create", "Create a Target Layer"));
 	
 	Landscape->AddTargetLayer();
 	
@@ -549,7 +549,7 @@ BEGIN_SLATE_FUNCTION_BUILD_OPTIMIZATION
 void FLandscapeEditorCustomNodeBuilder_TargetLayers::GenerateChildContent(IDetailChildrenBuilder& ChildrenBuilder)
 {
 	FEdModeLandscape* LandscapeEdMode = GetEditorMode();
-	if (LandscapeEdMode != NULL)
+	if (LandscapeEdMode != nullptr)
 	{
 		const TArray<TSharedRef<FLandscapeTargetListInfo>>& TargetList = LandscapeEdMode->GetTargetList();
 		const TArray<FName>* TargetDisplayOrderList = LandscapeEdMode->GetTargetDisplayOrderList();
@@ -737,19 +737,32 @@ TSharedPtr<SWidget> FLandscapeEditorCustomNodeBuilder_TargetLayers::GenerateRow(
 											.Font(IDetailLayoutBuilder::GetDetailFontBold())
 											.Text(Target->TargetLayerDisplayName)
 											.ColorAndOpacity_Static(&FLandscapeEditorCustomNodeBuilder_TargetLayers::GetTargetTextColor, Target)
-											.OnTextCommitted_Lambda([Target]( const FText& Text, ETextCommit::Type Type)
+											.OnVerifyTextChanged_Lambda([Target](const FText& InNewText, FText& OutErrorMessage)
+											{
+												const FName NewName(InNewText.ToString());
+												ALandscape* Landscape = Cast<ALandscape>(Target->Owner);
+												if ((Target->LayerName != NewName) && Landscape->HasTargetLayer(NewName))
+												{
+													OutErrorMessage = LOCTEXT("LandscapeTargetLayer_RenameFailed_AlreadyExists", "This target layer name already exists");
+													return false;
+												}
+
+												return true;
+											})
+											.OnTextCommitted_Lambda([Target](const FText& Text, ETextCommit::Type Type)
 											{
-												FScopedTransaction Transaction(LOCTEXT("LandscapeWeightLayer_Rename", "Rename Weight Layer"));
+												const FName NewName(Text.ToString());
+												if (Target->LayerName == NewName)
+												{
+													return;
+												}
+
+												FScopedTransaction Transaction(LOCTEXT("LandscapeTargetLayer_Rename", "Rename Target Layer"));
 												ALandscape* Landscape = Cast<ALandscape>(Target->Owner);
 												
 												const TMap<FName, FLandscapeTargetLayerSettings>& TargetLayers = Landscape->GetTargetLayers();
 												const FLandscapeTargetLayerSettings* LayerSettings = nullptr;
 												
-												if (Target->LayerName == FName(Text.ToString()))
-												{
-													LayerSettings = TargetLayers.Find(FName(Target->TargetLayerDisplayName.ToString()));	
-												}
-												
 												Landscape->RemoveTargetLayer(FName(Target->TargetLayerDisplayName.ToString()));
 											
 												Target->TargetLayerDisplayName = Text;
@@ -785,7 +798,7 @@ TSharedPtr<SWidget> FLandscapeEditorCustomNodeBuilder_TargetLayers::GenerateRow(
 									[
 										SNew(SObjectPropertyEntryBox)
 											.IsEnabled((bool)Target->bValid)
-											.ObjectPath(Target->LayerInfoObj != NULL ? Target->LayerInfoObj->GetPathName() : FString())
+											.ObjectPath(Target->LayerInfoObj != nullptr ? Target->LayerInfoObj->GetPathName() : FString())
 											.AllowedClass(ULandscapeLayerInfoObject::StaticClass())
 											.OnObjectChanged_Static(&FLandscapeEditorCustomNodeBuilder_TargetLayers::OnTargetLayerSetObject, Target)
 											.OnShouldFilterAsset_Static(&FLandscapeEditorCustomNodeBuilder_TargetLayers::ShouldFilterLayerInfo, Target->LayerName)
@@ -1036,7 +1049,7 @@ void FLandscapeEditorCustomNodeBuilder_TargetLayers::OnTargetSelectionChanged(co
 		LandscapeEdMode->CurrentToolTarget.TargetType = Target->TargetType;
 		if (Target->TargetType == ELandscapeToolTargetType::Heightmap)
 		{
-			checkSlow(Target->LayerInfoObj == NULL);
+			checkSlow(Target->LayerInfoObj == nullptr);
 			LandscapeEdMode->SetCurrentTargetLayer(NAME_None, nullptr);
 		}
 		else
@@ -1048,9 +1061,9 @@ void FLandscapeEditorCustomNodeBuilder_TargetLayers::OnTargetSelectionChanged(co
 
 TSharedPtr<SWidget> FLandscapeEditorCustomNodeBuilder_TargetLayers::OnTargetLayerContextMenuOpening(const TSharedRef<FLandscapeTargetListInfo> Target)
 {
-	if (Target->TargetType == ELandscapeToolTargetType::Heightmap || Target->LayerInfoObj != NULL)
+	if (Target->TargetType == ELandscapeToolTargetType::Heightmap || Target->LayerInfoObj != nullptr)
 	{
-		FMenuBuilder MenuBuilder(true, NULL);
+		FMenuBuilder MenuBuilder(true, nullptr);
 		
 		MenuBuilder.BeginSection("LandscapeEditorLayerActions", LOCTEXT("LayerContextMenu.Heading", "Layer Actions"));
 		{
@@ -1092,7 +1105,7 @@ TSharedPtr<SWidget> FLandscapeEditorCustomNodeBuilder_TargetLayers::OnTargetLaye
 		return MenuBuilder.MakeWidget();
 	}
 
-	return NULL;
+	return nullptr;
 }
 
 void FLandscapeEditorCustomNodeBuilder_TargetLayers::OnExportLayer(const TSharedRef<FLandscapeTargetListInfo> Target)
@@ -1104,7 +1117,7 @@ void FLandscapeEditorCustomNodeBuilder_TargetLayers::OnExportLayer(const TShared
 		IDesktopPlatform* DesktopPlatform = FDesktopPlatformModule::Get();
 
 		ULandscapeInfo* LandscapeInfo = Target->LandscapeInfo.Get();
-		ULandscapeLayerInfoObject* LayerInfoObj = Target->LayerInfoObj.Get(); // NULL for heightmaps
+		ULandscapeLayerInfoObject* LayerInfoObj = Target->LayerInfoObj.Get(); // nullptr for heightmaps
 
 		// Prompt for filename
 		FString SaveDialogTitle;
@@ -1167,7 +1180,7 @@ void FLandscapeEditorCustomNodeBuilder_TargetLayers::OnImportLayer(const TShared
 		IDesktopPlatform* DesktopPlatform = FDesktopPlatformModule::Get();
 
 		ULandscapeInfo* LandscapeInfo = Target->LandscapeInfo.Get();
-		ULandscapeLayerInfoObject* LayerInfoObj = Target->LayerInfoObj.Get(); // NULL for heightmaps
+		ULandscapeLayerInfoObject* LayerInfoObj = Target->LayerInfoObj.Get(); // nullptr for heightmaps
 
 		// Prompt for filename
 		FString OpenDialogTitle;
@@ -1314,7 +1327,7 @@ void FLandscapeEditorCustomNodeBuilder_TargetLayers::OnTargetLayerSetObject(cons
 {
 	// Can't assign null to a layer
 	UObject* Object = AssetData.GetAsset();
-	if (Object == NULL)
+	if (Object == nullptr)
 	{
 		return;
 	}
@@ -1416,7 +1429,7 @@ EVisibility FLandscapeEditorCustomNodeBuilder_TargetLayers::GetTargetLayerDelete
 
 TSharedRef<SWidget> FLandscapeEditorCustomNodeBuilder_TargetLayers::OnGetTargetLayerCreateMenu(const TSharedRef<FLandscapeTargetListInfo> Target)
 {
-	FMenuBuilder MenuBuilder(true, NULL);
+	FMenuBuilder MenuBuilder(true, nullptr);
 
 	MenuBuilder.AddMenuEntry(LOCTEXT("Menu_Create_Blended", "Weight-Blended Layer (normal)"), FText(), FSlateIcon(),
 		FUIAction(FExecuteAction::CreateStatic(&FLandscapeEditorCustomNodeBuilder_TargetLayers::OnTargetLayerCreateClicked, Target, false)));
@@ -1638,7 +1651,7 @@ void FLandscapeEditorCustomNodeBuilder_TargetLayers::OnDebugModeColorChannelChan
 	if (NewCheckedState == ECheckBoxState::Checked)
 	{
 		// Enable on us and disable colour channel on other targets
-		if (ensure(Target->LayerInfoObj != NULL))
+		if (ensure(Target->LayerInfoObj != nullptr))
 		{
 			ULandscapeInfo* LandscapeInfo = Target->LandscapeInfo.Get();
 			int32 Index = LandscapeInfo->GetLayerInfoIndex(Target->LayerInfoObj.Get(), Target->Owner.Get());
diff --git a/Engine/Source/Runtime/Landscape/Classes/LandscapeProxy.h b/Engine/Source/Runtime/Landscape/Classes/LandscapeProxy.h
--- a/Engine/Source/Runtime/Landscape/Classes/LandscapeProxy.h
+++ b/Engine/Source/Runtime/Landscape/Classes/LandscapeProxy.h
@@ -470,8 +470,8 @@ protected:
 
 public:
 #if WITH_EDITORONLY_DATA
-	UPROPERTY(Transient)
-	bool bInitWeightLayersFromMaterial;
+	UPROPERTY()
+	TMap<FName, TObjectPtr<ULandscapeLayerInfoObject>> TargetLayersForFixup;
 #endif //WITH_EDITORONLY_DATA
 	/** Offset in quads from global components grid origin (in quads) **/
 	UPROPERTY()
@@ -1285,11 +1285,11 @@ PRAGMA_ENABLE_DEPRECATION_WARNINGS
 	LANDSCAPE_API static const TArray<FName>& GetLayersFromMaterial(UMaterialInterface* Material);
 	LANDSCAPE_API const TArray<FName>& GetLayersFromMaterial() const;
 
-	// Get all Layer names used by the materials assigned the proxy and components ( Proxy Material, Component, Component LOD Material and Hole Material)
-	LANDSCAPE_API TArray<FName> RetrieveAllLayerNamesFromMaterials() const;
+	// Get all Layer names used by the materials assigned the proxy and components (Proxy Material, Component, Component LOD Material and Hole Material)
+	LANDSCAPE_API TArray<FName> RetrieveTargetLayerNamesFromMaterials() const;
 	
 	// Query all the LandscapeLayerInfo Objects from the weight layers allocated on this proxy.
-	LANDSCAPE_API TMap<FName, ULandscapeLayerInfoObject*> RetrieveAllocationInfos() const;
+	TMap<FName, ULandscapeLayerInfoObject*> RetrieveTargetLayerInfosFromAllocations() const;
 
 	/**
 	* Creates a new LandscapeLayerInfoObject
diff --git a/Engine/Source/Runtime/Landscape/Private/Landscape.cpp b/Engine/Source/Runtime/Landscape/Private/Landscape.cpp
--- a/Engine/Source/Runtime/Landscape/Private/Landscape.cpp
+++ b/Engine/Source/Runtime/Landscape/Private/Landscape.cpp
@@ -3924,42 +3924,38 @@ bool ULandscapeInfo::UpdateLayerInfoMapInternal(ALandscapeProxy* Proxy)
 		return false;
 	}
 
-	// Perform a delayed (see where bInitWeightLayersFromMaterial is assigned), one-time deprecation of the landscape layer data based on the content of the materials in the components (see FixupLandscapeTargetLayersInLandscapeActor)
-	if (GIsEditor && Proxy && Proxy->bInitWeightLayersFromMaterial)
+	// Perform a delayed (see where TargetLayersForFixup is set), one-time deprecation of the landscape layer data based on the content of the materials in the components (see FixupLandscapeTargetLayersInLandscapeActor)
+	if (GIsEditor && Proxy && !Proxy->TargetLayersForFixup.IsEmpty())
 	{
-		Proxy->bInitWeightLayersFromMaterial = false;
-		// add the landscape layers from the material on the proxy & LandscapeLayerInfos from the WeightMap Allocations
-		TArray<FName> LayerNames = Proxy->RetrieveAllLayerNamesFromMaterials();
-		const TMap<FName, ULandscapeLayerInfoObject*> InfoObjects = Proxy->RetrieveAllocationInfos();
-		for(const FName& LayerName : LayerNames)
+		// Go through the list of layer names / info to fixup and declare new or update existing layers in the main landscape actor if we have one that the main landscape doesn't know about :
+		for (const auto& It : Proxy->TargetLayersForFixup)
 		{
-			ULandscapeLayerInfoObject* const * LayerInfoObject = InfoObjects.Find(LayerName);
-			if (!LandscapeActor->HasTargetLayer(LayerName))
+			FName LayerName = It.Key;
+			ULandscapeLayerInfoObject* LayerInfo = It.Value;
+			check(LayerName.IsValid());
+
+			const FLandscapeTargetLayerSettings* LayerSettingsInLandscapeActor = LandscapeActor->GetTargetLayers().Find(LayerName);
+			// If the layer isn't known to the main landscape, add it now : 
+			if (LayerSettingsInLandscapeActor == nullptr)
 			{
 				// Mark the parent landscape actor dirty with bInForceResave == true so that the parent actor is put into the list of files to save even if we do this fixup on load :
 				MarkObjectDirty(/*InObject = */LandscapeActor.Get(), /*bInForceResave = */true);
 
-				LandscapeActor->AddTargetLayer(LayerName, FLandscapeTargetLayerSettings(LayerInfoObject ? *LayerInfoObject : nullptr), false);
+				LandscapeActor->AddTargetLayer(LayerName, FLandscapeTargetLayerSettings(LayerInfo), false);
 				bLayerInfoMapChanged = true;
 			}
-		}
-			
-		for (const auto& Layer : LandscapeActor->GetTargetLayers())
-		{
-			// Fixup the layers that point to an invalid layer info found in the proxy if possible
-			if (Layer.Value.LayerInfoObj == nullptr)
+			// If the layer name is known to the main landscape but it hasn't got a landscape info associated to it yet, update it to use this LayerInfo :
+			else if ((LayerInfo != nullptr) && (LayerSettingsInLandscapeActor->LayerInfoObj == nullptr))
 			{
-				ULandscapeLayerInfoObject* const * LayerInfoObject = InfoObjects.Find(Layer.Key);
-				if (LayerInfoObject && LandscapeActor->HasTargetLayer(Layer.Key))
-				{
-					// Mark the parent landscape actor dirty with bInForceResave == true so that the parent actor is put into the list of files to save even if we do this fixup on load :
-					MarkObjectDirty(/*InObject = */LandscapeActor.Get(), /*bInForceResave = */true);
+				// Mark the parent landscape actor dirty with bInForceResave == true so that the parent actor is put into the list of files to save even if we do this fixup on load :
+				MarkObjectDirty(/*InObject = */LandscapeActor.Get(), /*bInForceResave = */true);
 
-					LandscapeActor->UpdateTargetLayer(Layer.Key, FLandscapeTargetLayerSettings(*LayerInfoObject), false);
-					bLayerInfoMapChanged = true;
-				}
+				LandscapeActor->UpdateTargetLayer(LayerName, FLandscapeTargetLayerSettings(LayerInfo), false);
+				bLayerInfoMapChanged = true;
 			}
 		}
+
+		Proxy->TargetLayersForFixup.Empty();
 	}
 
 	Layers.Empty();
@@ -4115,12 +4111,6 @@ void ALandscapeProxy::PostLoad()
 	}
 
 #if WITH_EDITOR
-	const int32 LinkerVersion = GetLinkerCustomVersion(FFortniteMainBranchObjectVersion::GUID);
-	// With and before LandscapeTargetLayersInLandscapeActor and until FixupLandscapeTargetLayersInLandscapeActor, some layer info objects have been incorrectly unassigned and we now have to go through 
-	//  all materials of all streaming proxies to gather the missing landscape layer info objects and this can only be done after proxies are united with their parent landscape by their ULandscapeInfo 
-	//  (since TargetLayers is a LandscapeInherited property of the parent landscape, propagated to the child proxies), so we set a boolean to delay this operation until then :
-	bInitWeightLayersFromMaterial = (LinkerVersion < FFortniteMainBranchObjectVersion::FixupLandscapeTargetLayersInLandscapeActor);
-	
 PRAGMA_DISABLE_DEPRECATION_WARNINGS
 	if (!LandscapeMaterialsOverride_DEPRECATED.IsEmpty())
 	{
@@ -4131,8 +4121,54 @@ PRAGMA_DISABLE_DEPRECATION_WARNINGS
 		}
 		LandscapeMaterialsOverride_DEPRECATED.Reset();
 	}
+
+	if (!EditorLayerSettings_DEPRECATED.IsEmpty())
+	{
+		// If we still have access to EditorLayerSettings_DEPRECATED because it's the first time we deprecate this proxy since FFortniteMainBranchObjectVersion::LandscapeTargetLayersInLandscapeActor, 
+		//  fill the list of target layers to fixup based on the original property because it's the most accurate (it has layer info assignment even if there's no weightmap allocation for a given layer) :
+		TargetLayersForFixup.Reserve(EditorLayerSettings_DEPRECATED.Num());
+		for (const FLandscapeEditorLayerSettings& EditorLayerSetting : EditorLayerSettings_DEPRECATED)
+		{
+			if (EditorLayerSetting.LayerInfoObj != nullptr)
+			{
+				TargetLayersForFixup.Add(EditorLayerSetting.LayerInfoObj->LayerName, EditorLayerSetting.LayerInfoObj);
+			}
+		}
+		EditorLayerSettings_DEPRECATED.Reset();
+	}
 PRAGMA_ENABLE_DEPRECATION_WARNINGS
 
+	const int32 LinkerVersion = GetLinkerCustomVersion(FFortniteMainBranchObjectVersion::GUID);
+	// With and before LandscapeTargetLayersInLandscapeActor and until FixupLandscapeTargetLayersInLandscapeActor, some layer info objects have been incorrectly unassigned and we now have to go through 
+	//  all materials of all streaming proxies to gather the missing landscape layer info objects and this can only be done after proxies are united with their parent landscape by their ULandscapeInfo 
+	//  (since TargetLayers is a LandscapeInherited property of the parent landscape, propagated to the child proxies), so we must delay this operation until then. What we do here is prepare a list of 
+	//  layers to fixup in the main landscape actor (TargetLayersForFixup) and when the proxy is registered to the parent landscape, we'll go through that list and update the landscape's TargetLayers list,
+	//  which will then be synchronized with all proxies if necessary:
+	if (LinkerVersion < FFortniteMainBranchObjectVersion::FixupLandscapeTargetLayersInLandscapeActor)
+	{
+		// Go through the list of materials and weightmap allocations to gather potential layer name / layer info associations :
+		{
+			TMap<FName, ULandscapeLayerInfoObject*> LayerInfosFromAllocations = RetrieveTargetLayerInfosFromAllocations();
+			for (const auto& It : LayerInfosFromAllocations)
+			{
+				FName LayerName = It.Key;
+				ULandscapeLayerInfoObject* LayerInfo = It.Value;
+				TObjectPtr<ULandscapeLayerInfoObject>* LayerInfoInFixupMap = TargetLayersForFixup.Find(LayerName);
+				// Unknown layer name yet, let's add a layer name / info association :
+				if (LayerInfoInFixupMap == nullptr)
+				{
+					TargetLayersForFixup.Add(LayerName, LayerInfo);
+				}
+				// Known layer name, but we have no valid layer info associated with it yet, update it :
+				else if (*LayerInfoInFixupMap == nullptr)
+				{
+					*LayerInfoInFixupMap = LayerInfo;
+				}
+				// Otherwise, don't touch it, we consider that TargetLayersForFixup has the authority over this layer already
+			}
+		}
+	}
+
 	if (GIsEditor)
 	{
 		// We may not have run PostLoad on LandscapeComponents yet
@@ -5458,7 +5494,7 @@ void ULandscapeInfo::RegisterLandscapeActorWithProxyInternal(ALandscapeProxy* Pr
 	// generic proxy setup (that requires ALandscape actor) here
 	if (bool bLayerInfoMapChanged = UpdateLayerInfoMap(Proxy))
 	{
-		// The layer info map is part of the main landscape so if it has changed, we needto do another round of shared data fixup on all proxies, so all proxies have their TargetLayers list synchronized. 
+		// The layer info map is part of the main landscape so if it has changed, we need to do another round of shared data fixup on all proxies, so all proxies have their TargetLayers list synchronized. 
 		//  This is a one-time thing because at some point during development, the target layer data was deprecated and the deprecation turned somewhat sour :S
 		ForEachLandscapeProxy([this, Landscape, bMapCheck](ALandscapeProxy* Proxy)
 		{
diff --git a/Engine/Source/Runtime/Landscape/Private/LandscapeEdit.cpp b/Engine/Source/Runtime/Landscape/Private/LandscapeEdit.cpp
--- a/Engine/Source/Runtime/Landscape/Private/LandscapeEdit.cpp
+++ b/Engine/Source/Runtime/Landscape/Private/LandscapeEdit.cpp
@@ -2982,7 +2982,7 @@ const TArray<FName>& ALandscapeProxy::GetLayersFromMaterial() const
 	return GetLayersFromMaterial(LandscapeMaterial);
 }
 
-TArray<FName> ALandscapeProxy::RetrieveAllLayerNamesFromMaterials() const
+TArray<FName> ALandscapeProxy::RetrieveTargetLayerNamesFromMaterials() const
 {
 	TSet<FName> LayerNames;
 	LayerNames.Append(GetLayersFromMaterial());
@@ -3009,7 +3009,7 @@ TArray<FName> ALandscapeProxy::RetrieveAllLayerNamesFromMaterials() const
 	return LayerNames.Array();
 }
 
-TMap<FName, ULandscapeLayerInfoObject*> ALandscapeProxy::RetrieveAllocationInfos() const
+TMap<FName, ULandscapeLayerInfoObject*> ALandscapeProxy::RetrieveTargetLayerInfosFromAllocations() const
 {
 	TMap<FName, ULandscapeLayerInfoObject*> InfoObjects;
 	for (const TObjectPtr<ULandscapeComponent>& Component : LandscapeComponents)
@@ -7223,7 +7223,7 @@ void ULandscapeComponent::ReallocateWeightmapsInternal(FLandscapeEditDataInterfa
 			int32 WeightmapSize = (SubsectionSizeQuads + 1) * NumSubsections;
 
 			// We need a new weightmap texture
-			CurrentWeightmapTexture = TargetProxy->CreateLandscapeTexture(WeightmapSize, WeightmapSize, TEXTUREGROUP_Terrain_Weightmap, TSF_BGRA8, nullptr, /* bCompress = */ false, /* bMipChain= */ bIsFinalWeightmap); //dmb-nomips
+			CurrentWeightmapTexture = TargetProxy->CreateLandscapeTexture(WeightmapSize, WeightmapSize, TEXTUREGROUP_Terrain_Weightmap, TSF_BGRA8, nullptr, /* bCompress = */ false, /* bMipChain= */ bIsFinalWeightmap); 
 
 			// Alloc dummy mips
 			if (bIsFinalWeightmap)
