Prevent user from accidently using a landscape target layer named after the reserved weightmap name for visibility

diff --git a/Engine/Source/Editor/LandscapeEditor/Private/LandscapeEditorDetailCustomization_TargetLayers.cpp b/Engine/Source/Editor/LandscapeEditor/Private/LandscapeEditorDetailCustomization_TargetLayers.cpp
--- a/Engine/Source/Editor/LandscapeEditor/Private/LandscapeEditorDetailCustomization_TargetLayers.cpp
+++ b/Engine/Source/Editor/LandscapeEditor/Private/LandscapeEditorDetailCustomization_TargetLayers.cpp
@@ -33,6 +33,7 @@
 #include "PropertyCustomizationHelpers.h"
 #include "Subsystems/AssetEditorSubsystem.h"
 #include "AssetRegistry/AssetRegistryModule.h"
+#include "Materials/MaterialExpressionLandscapeVisibilityMask.h"
 
 #include "SLandscapeEditor.h"
 #include "Dialogs/DlgPickAssetPath.h"
@@ -740,8 +741,20 @@ TSharedPtr<SWidget> FLandscapeEditorCustomNodeBuilder_TargetLayers::GenerateRow(
 											.OnVerifyTextChanged_Lambda([Target](const FText& InNewText, FText& OutErrorMessage)
 											{
 												const FName NewName(InNewText.ToString());
+
+												if (Target->LayerName == NewName)
+												{
+													return true;
+												}
+
+												if (NewName == UMaterialExpressionLandscapeVisibilityMask::ParameterName)
+												{
+													OutErrorMessage = LOCTEXT("LandscapeTargetLayer_RenameFailed_ReservedName", "This target layer name is reserved for internal usage");
+													return false;
+												}
+
 												ALandscape* Landscape = Cast<ALandscape>(Target->Owner);
-												if ((Target->LayerName != NewName) && Landscape->HasTargetLayer(NewName))
+												if (Landscape->HasTargetLayer(NewName))
 												{
 													OutErrorMessage = LOCTEXT("LandscapeTargetLayer_RenameFailed_AlreadyExists", "This target layer name already exists");
 													return false;
