Use class display names in the lanscape edit layer class picker
#rnx

diff --git a/Engine/Source/Editor/LandscapeEditor/Private/LandscapeEditorDetailCustomization_Layers.cpp b/Engine/Source/Editor/LandscapeEditor/Private/LandscapeEditorDetailCustomization_Layers.cpp
--- a/Engine/Source/Editor/LandscapeEditor/Private/LandscapeEditorDetailCustomization_Layers.cpp
+++ b/Engine/Source/Editor/LandscapeEditor/Private/LandscapeEditorDetailCustomization_Layers.cpp
@@ -1058,6 +1058,7 @@ TSubclassOf<ULandscapeEditLayerBase> FLandscapeEditorCustomNodeBuilder_Layers::P
 	// Fill in options
 	FClassViewerInitializationOptions Options;
 	Options.Mode = EClassViewerMode::ClassPicker;
+	Options.NameTypeToDisplay = EClassViewerNameTypeToDisplay::DisplayName;
 
 	TSharedPtr<FLandscapeEditLayerClassFilter> Filter = MakeShareable(new FLandscapeEditLayerClassFilter());
 	Options.ClassFilters.Add(Filter.ToSharedRef());
