Fix for crash creating new landscape
Modify landscape creation code to handle failed region creation gracefully

This crash has only ever happened in 29.30, and not since mid-may, however the cause seems to be SpawnActor() failing and returning null.
It is unknown why this failure mode only happens in 29.30.

This code change was tested against forced null results.

diff --git a/Engine/Source/Editor/LandscapeEditor/Private/LandscapeEditorDetailCustomization_NewLandscape.cpp b/Engine/Source/Editor/LandscapeEditor/Private/LandscapeEditorDetailCustomization_NewLandscape.cpp
--- a/Engine/Source/Editor/LandscapeEditor/Private/LandscapeEditorDetailCustomization_NewLandscape.cpp
+++ b/Engine/Source/Editor/LandscapeEditor/Private/LandscapeEditorDetailCustomization_NewLandscape.cpp
@@ -1223,8 +1223,11 @@ FReply FLandscapeEditorDetailCustomization_NewLandscape::OnCreateButtonClicked()
 			double RegionSizeX = RegionSizeXTexels * LandscapeProxy->GetActorScale3D().X;
 			double RegionSizeY = RegionSizeYTexels * LandscapeProxy->GetActorScale3D().Y;
 			ALocationVolume* RegionVolume = LandscapeRegionUtils::CreateLandscapeRegionVolume(World, LandscapeProxy, RegionCoordinate, RegionSizeX);
-			RegionVolumes.Add(RegionVolume);
-			
+			if (RegionVolume)
+			{
+				RegionVolumes.Add(RegionVolume);
+			}
+
 			TArray<ALandscapeProxy*> CreatedStreamingProxies;
 			AddComponents(LandscapeInfo, LandscapeSubsystem, NewComponents, CreatedStreamingProxies);
 			
diff --git a/Engine/Source/Editor/LandscapeEditor/Private/LandscapeRegionUtils.cpp b/Engine/Source/Editor/LandscapeEditor/Private/LandscapeRegionUtils.cpp
--- a/Engine/Source/Editor/LandscapeEditor/Private/LandscapeRegionUtils.cpp
+++ b/Engine/Source/Editor/LandscapeEditor/Private/LandscapeRegionUtils.cpp
@@ -32,6 +32,10 @@ ALocationVolume* CreateLandscapeRegionVolume(UWorld* InWorld, ALandscapeProxy* I
 	SpawnParameters.NameMode = FActorSpawnParameters::ESpawnActorNameMode::Requested;
 
 	ALocationVolume* LocationVolume = InWorld->SpawnActor<ALocationVolume>(Location, Rotation, SpawnParameters);
+	if (LocationVolume == nullptr)
+	{
+		return nullptr;
+	}
 	LocationVolume->SetActorLabel(Label);
 
 	LocationVolume->AttachToActor(InParentLandscapeActor, FAttachmentTransformRules::KeepWorldTransform);
