UX fix for landscape edit layer names overlapping with the "Alpha" text box

diff --git a/Engine/Source/Editor/LandscapeEditor/Private/LandscapeEditorDetailCustomization_Layers.cpp b/Engine/Source/Editor/LandscapeEditor/Private/LandscapeEditorDetailCustomization_Layers.cpp
--- a/Engine/Source/Editor/LandscapeEditor/Private/LandscapeEditorDetailCustomization_Layers.cpp
+++ b/Engine/Source/Editor/LandscapeEditor/Private/LandscapeEditorDetailCustomization_Layers.cpp
@@ -291,9 +291,11 @@ TSharedPtr<SWidget> FLandscapeEditorCustomNodeBuilder_Layers::GenerateRow(int32
 				.Padding(4, 0)
 				[
 					SNew(SHorizontalBox)
+						.Clipping(EWidgetClipping::OnDemand)
+
 						+ SHorizontalBox::Slot()
 						.Padding(0)
-						.AutoWidth()
+						.FillWidth(1.0f)
 						.VAlign(VAlign_Center)
 						.HAlign(HAlign_Left)
 						[
