Unshelved from pending changelist '35350462':

Several fixes with landscape brush

Details:
* Reverted incorrect change (29959327) for computation of valid area in FLandscapeEditDataInterface::GetHeightDataInternal
* In TLandscapeEditCache, when extending the cache data (in the "interpolation case"), we were incorrectly using the "valid bounds" returned by GetWeightData because this function mutates its input attributes likewise but without actually caching more than what was requested. We now do the same as in the "no-interpolation case" (where GetWeightDataFast doesn't do a similar modification to its inputs...), and only use the cached bounds to extend the cache
* Unlike TLandscapeEditCache's GetCachedData, GetOriginalData didn't check whether the data was present or not before storing it, which could introduce random values to OriginalData
* Pretty much all landscape brushes were using mismatched inclusive/exclusive bounds, which led to a missing row/column on the the bottom/right

Misc changes :
* Added landscape.VisualLog.ShowBrushPhysics (0 by default) to avoid polluting the visual logger whenever we move the mouse when landscape sculpting/painting, since that is only relevant to debug the landscape brush
* Added visual logs when extending the landscape data cache when using the landscape brush + fixed those logs to extend by (-0.5, +0.5) beyond the cached region, in order to properly show the vertices that were sampled (i.e. that allows to show what was cached when the cached region contains a single landscape vertex in either the X or Y direction)
* Use standard namespace name in landscape editor

diff --git a/Engine/Source/Editor/LandscapeEditor/Private/LandscapeEdMode.cpp b/Engine/Source/Editor/LandscapeEditor/Private/LandscapeEdMode.cpp
--- a/Engine/Source/Editor/LandscapeEditor/Private/LandscapeEdMode.cpp
+++ b/Engine/Source/Editor/LandscapeEditor/Private/LandscapeEdMode.cpp
@@ -75,7 +75,7 @@ void FLandscapeTool::SetEditRenderType()
 	GLandscapeEditRenderMode = ELandscapeEditRenderMode::SelectRegion | (GLandscapeEditRenderMode & ELandscapeEditRenderMode::BitMaskForMask);
 }
 
-namespace LandscapeTool
+namespace UE::Landscape::Editor::Tool
 {
 	UMaterialInstance* CreateMaterialInstance(UMaterialInterface* BaseMaterial)
 	{
@@ -88,12 +88,22 @@ namespace LandscapeTool
 		MaterialInstance->PostEditChange();
 		return MaterialInstance;
 	}
+}
+// namespace UE::Landscape::Editor::Tool
 
+namespace UE::Landscape::Editor::Private
+{
 	/** Indicates the user is currently moving the landscape gizmo object by dragging the mouse. */
 	bool GIsGizmoDragging = false;
 	/** Indicates the user is currently changing the landscape brush radius/falloff by dragging the mouse. */
 	bool GIsAdjustingBrush = false;
+
+	static FAutoConsoleVariable CVarVisualLogShowBrushPhysics(
+		TEXT("landscape.VisualLog.ShowBrushPhysics"),
+		false,
+		TEXT("When visual log is active, allows to leave a visual log of the physics queries made for the purpose of the landscape brush"));
 }
+// namespace UE::Landscape::Editor::Private
 
 //
 // FEdModeLandscape
@@ -122,14 +132,16 @@ FEdModeLandscape::FEdModeLandscape()
 	, bNeedsUpdateShownLayerList(false)
 	, bUpdatingLandscapeInfo(false)
 {
-	GLayerDebugColorMaterial = LandscapeTool::CreateMaterialInstance(LoadObject<UMaterial>(nullptr, TEXT("/Engine/EditorLandscapeResources/LayerVisMaterial.LayerVisMaterial")));
-	GSelectionColorMaterial  = LandscapeTool::CreateMaterialInstance(LoadObject<UMaterialInstanceConstant>(nullptr, TEXT("/Engine/EditorLandscapeResources/SelectBrushMaterial_Selected.SelectBrushMaterial_Selected")));
-	GSelectionRegionMaterial = LandscapeTool::CreateMaterialInstance(LoadObject<UMaterialInstanceConstant>(nullptr, TEXT("/Engine/EditorLandscapeResources/SelectBrushMaterial_SelectedRegion.SelectBrushMaterial_SelectedRegion")));
-	GMaskRegionMaterial      = LandscapeTool::CreateMaterialInstance(LoadObject<UMaterialInstanceConstant>(nullptr, TEXT("/Engine/EditorLandscapeResources/MaskBrushMaterial_MaskedRegion.MaskBrushMaterial_MaskedRegion")));
-	GColorMaskRegionMaterial = LandscapeTool::CreateMaterialInstance(LoadObject<UMaterialInstanceConstant>(nullptr, TEXT("/Engine/EditorLandscapeResources/ColorMaskBrushMaterial_MaskedRegion.ColorMaskBrushMaterial_MaskedRegion")));
-	GLandscapeDirtyMaterial		 = LandscapeTool::CreateMaterialInstance(LoadObject<UMaterial>(nullptr, TEXT("/Engine/EditorLandscapeResources/LandscapeDirtyMaterial.LandscapeDirtyMaterial")));
-	GLandscapeBlackTexture   = LoadObject<UTexture2D>(nullptr, TEXT("/Engine/EngineResources/Black.Black"));
-	GLandscapeLayerUsageMaterial = LandscapeTool::CreateMaterialInstance(LoadObject<UMaterial>(nullptr, TEXT("/Engine/EditorLandscapeResources/LandscapeLayerUsageMaterial.LandscapeLayerUsageMaterial")));
+	using namespace UE::Landscape::Editor::Tool;
+
+	GLayerDebugColorMaterial = CreateMaterialInstance(LoadObject<UMaterial>(nullptr, TEXT("/Engine/EditorLandscapeResources/LayerVisMaterial.LayerVisMaterial")));
+	GSelectionColorMaterial = CreateMaterialInstance(LoadObject<UMaterialInstanceConstant>(nullptr, TEXT("/Engine/EditorLandscapeResources/SelectBrushMaterial_Selected.SelectBrushMaterial_Selected")));
+	GSelectionRegionMaterial = CreateMaterialInstance(LoadObject<UMaterialInstanceConstant>(nullptr, TEXT("/Engine/EditorLandscapeResources/SelectBrushMaterial_SelectedRegion.SelectBrushMaterial_SelectedRegion")));
+	GMaskRegionMaterial = CreateMaterialInstance(LoadObject<UMaterialInstanceConstant>(nullptr, TEXT("/Engine/EditorLandscapeResources/MaskBrushMaterial_MaskedRegion.MaskBrushMaterial_MaskedRegion")));
+	GColorMaskRegionMaterial = CreateMaterialInstance(LoadObject<UMaterialInstanceConstant>(nullptr, TEXT("/Engine/EditorLandscapeResources/ColorMaskBrushMaterial_MaskedRegion.ColorMaskBrushMaterial_MaskedRegion")));
+	GLandscapeDirtyMaterial = CreateMaterialInstance(LoadObject<UMaterial>(nullptr, TEXT("/Engine/EditorLandscapeResources/LandscapeDirtyMaterial.LandscapeDirtyMaterial")));
+	GLandscapeBlackTexture = LoadObject<UTexture2D>(nullptr, TEXT("/Engine/EngineResources/Black.Black"));
+	GLandscapeLayerUsageMaterial = CreateMaterialInstance(LoadObject<UMaterial>(nullptr, TEXT("/Engine/EditorLandscapeResources/LandscapeLayerUsageMaterial.LandscapeLayerUsageMaterial")));
 	
 	// Initialize modes
 	UpdateToolModes();
@@ -432,6 +444,8 @@ int32 FEdModeLandscape::GetNewLandscapeResolutionY() const
 /** FEdMode: Called when the mode is entered */
 void FEdModeLandscape::Enter()
 {
+	using namespace UE::Landscape::Editor::Private;
+
 	ErrorReasonOnMouseUp = FText::GetEmpty();
 
 	// Call parent implementation
@@ -643,16 +657,18 @@ void FEdModeLandscape::Enter()
 	}
 
 	// Reset mouse tracking info : 
-	LandscapeTool::GIsGizmoDragging = false;
-	LandscapeTool::GIsAdjustingBrush  = false;
+	GIsGizmoDragging = false;
+	GIsAdjustingBrush  = false;
 }
 
 /** FEdMode: Called when the mode is exited */
 void FEdModeLandscape::Exit()
 {
+	using namespace UE::Landscape::Editor::Private;
+
 	// Reset mouse tracking info : 
-	LandscapeTool::GIsGizmoDragging = false;
-	LandscapeTool::GIsAdjustingBrush  = false;
+	GIsGizmoDragging = false;
+	GIsAdjustingBrush  = false;
 
 	if (UWorld* World = GetWorld())
 	{
@@ -1074,14 +1090,16 @@ bool FEdModeLandscape::CapturedMouseMove(FEditorViewportClient* ViewportClient,
 /** FEdMode: Called when a mouse button is pressed */
 bool FEdModeLandscape::StartTracking(FEditorViewportClient* InViewportClient, FViewport* InViewport)
 {
+	using namespace UE::Landscape::Editor::Private;
+
 	if (CurrentGizmoActor.IsValid() && CurrentGizmoActor->IsSelected() && GLandscapeEditRenderMode & ELandscapeEditRenderMode::Gizmo)
 	{
-		LandscapeTool::GIsGizmoDragging = true;
+		GIsGizmoDragging = true;
 		return true;
 	}
 	else if (IsAdjustingBrush(InViewportClient))
 	{ 
-		LandscapeTool::GIsAdjustingBrush = true;
+		GIsAdjustingBrush = true;
 		// We're adjusting the brush via mouse tracking, return true in order to prevent the viewport client from doing any mouse dragging operation while we're doing it
 		return true;
 	}
@@ -1093,14 +1111,16 @@ bool FEdModeLandscape::StartTracking(FEditorViewportClient* InViewportClient, FV
 /** FEdMode: Called when the a mouse button is released */
 bool FEdModeLandscape::EndTracking(FEditorViewportClient* InViewportClient, FViewport* InViewport)
 {
-	if (LandscapeTool::GIsGizmoDragging)
+	using namespace UE::Landscape::Editor::Private;
+
+	if (GIsGizmoDragging)
 	{
-		LandscapeTool::GIsGizmoDragging = false;
+		GIsGizmoDragging = false;
 		return true;
 	}
-	if (LandscapeTool::GIsAdjustingBrush)
+	if (GIsAdjustingBrush)
 	{
-		LandscapeTool::GIsAdjustingBrush = false;
+		GIsAdjustingBrush = false;
 		return true;
 	}
 	return false;
@@ -1188,6 +1208,8 @@ bool FEdModeLandscape::ProcessLandscapeTraceHits(const TArray<FHitResult>& InRes
 
 bool FEdModeLandscape::LandscapeTrace(const FVector& InRayOrigin, const FVector& InRayEnd, const FVector& InDirection, FVector& OutHitLocation)
 {
+	using namespace UE::Landscape::Editor::Private;
+
 	TRACE_CPUPROFILER_EVENT_SCOPE(FEdModeLandscape_LandscapeTrace);
 	if (!CurrentTool || !CurrentToolTarget.LandscapeInfo.IsValid())
 	{
@@ -1220,14 +1242,14 @@ bool FEdModeLandscape::LandscapeTrace(const FVector& InRayOrigin, const FVector&
 			{
 				OutHitLocation = ProcessResult.LandscapeProxy->LandscapeActorToWorld().InverseTransformPosition(ProcessResult.HitLocation);
 			
-				UE_VLOG_SEGMENT_THICK(World, LogLandscapeEdMode, VeryVerbose, InRayOrigin, InRayEnd, FColor(100,255,100), 4, TEXT("landscape:ray-hit"));
-				UE_VLOG_LOCATION(World, LogLandscapeEdMode, VeryVerbose,  ProcessResult.LandscapeProxy->LandscapeActorToWorld().TransformPosition(OutHitLocation), 20.0,  FColor(100,100,255), TEXT("landscape:point-hit"));
+				UE_CVLOG_SEGMENT_THICK(CVarVisualLogShowBrushPhysics->GetBool(), World, LogLandscapeEdMode, VeryVerbose, InRayOrigin, InRayEnd, FColor(100,255,100), 4, TEXT("landscape:ray-hit"));
+				UE_CVLOG_LOCATION(CVarVisualLogShowBrushPhysics->GetBool(), World, LogLandscapeEdMode, VeryVerbose,  ProcessResult.LandscapeProxy->LandscapeActorToWorld().TransformPosition(OutHitLocation), 20.0,  FColor(100,100,255), TEXT("landscape:point-hit"));
 				return true;
 			}
 		}
 	}
 
-	UE_VLOG_SEGMENT_THICK(World, LogLandscapeEdMode, VeryVerbose, InRayOrigin, InRayEnd, FColor(255,100,100), 2, TEXT("landscape:ray-miss"));
+	UE_CVLOG_SEGMENT_THICK(CVarVisualLogShowBrushPhysics->GetBool(), World, LogLandscapeEdMode, VeryVerbose, InRayOrigin, InRayEnd, FColor(255,100,100), 2, TEXT("landscape:ray-miss"));
 		
 	if (CurrentTool->UseSphereTrace())
 	{
@@ -1262,8 +1284,8 @@ bool FEdModeLandscape::LandscapeTrace(const FVector& InRayOrigin, const FVector&
 			{
 				if (FProcessLandscapeTraceHitsResult ProcessResult; ProcessLandscapeTraceHits(Results, ProcessResult))
 				{
-					UE_VLOG_SEGMENT_THICK(World, LogLandscapeEdMode, VeryVerbose, AdjustedRayStart, AdjustedRayEnd, FColor(100, 255, 100), 4, TEXT("landscape:ray-hit"));
-					UE_VLOG_LOCATION(World, LogLandscapeEdMode, VeryVerbose, ProcessResult.HitLocation, 20.0, FColor(100, 100, 255), TEXT("landscape:point-hit"));
+					UE_CVLOG_SEGMENT_THICK(CVarVisualLogShowBrushPhysics->GetBool(), World, LogLandscapeEdMode, VeryVerbose, AdjustedRayStart, AdjustedRayEnd, FColor(100, 255, 100), 4, TEXT("landscape:ray-hit"));
+					UE_CVLOG_LOCATION(CVarVisualLogShowBrushPhysics->GetBool(), World, LogLandscapeEdMode, VeryVerbose, ProcessResult.HitLocation, 20.0, FColor(100, 100, 255), TEXT("landscape:point-hit"));
 
 					HitDistance = (ProcessResult.HitLocation - Start).Length();
 					return ProcessResult;
@@ -1271,7 +1293,7 @@ bool FEdModeLandscape::LandscapeTrace(const FVector& InRayOrigin, const FVector&
 			}
 			else
 			{
-				UE_VLOG_SEGMENT_THICK(World, LogLandscapeEdMode, VeryVerbose, AdjustedRayStart, AdjustedRayEnd, FColor(255, 100, 100), 2, TEXT("landscape:ray-miss"));
+				UE_CVLOG_SEGMENT_THICK(CVarVisualLogShowBrushPhysics->GetBool(), World, LogLandscapeEdMode, VeryVerbose, AdjustedRayStart, AdjustedRayEnd, FColor(255, 100, 100), 2, TEXT("landscape:ray-miss"));
 			}
 			return TOptional<FProcessLandscapeTraceHitsResult>();
 		};
@@ -1323,7 +1345,7 @@ bool FEdModeLandscape::LandscapeTrace(const FVector& InRayOrigin, const FVector&
 			MeanHeight /= (float)Count;
 			PointOnPlane.Z = MeanHeight;
 
-			UE_VLOG_LOCATION(World, LogLandscapeEdMode, VeryVerbose, PointOnPlane, 10.0, FColor(100, 100, 255), TEXT("landscape:point-on-plane"));
+			UE_CVLOG_LOCATION(CVarVisualLogShowBrushPhysics->GetBool(), World, LogLandscapeEdMode, VeryVerbose, PointOnPlane, 10.0, FColor(100, 100, 255), TEXT("landscape:point-on-plane"));
 
 			if (FMath::Abs(FVector::DotProduct(InDirection, FVector::ZAxisVector)) < SMALL_NUMBER)
 			{
@@ -1335,7 +1357,7 @@ bool FEdModeLandscape::LandscapeTrace(const FVector& InRayOrigin, const FVector&
 			FVector EstimatedHitLocation = FMath::RayPlaneIntersection(Start, InDirection, Plane);
 			check(!EstimatedHitLocation.ContainsNaN());
 
-			UE_VLOG_LOCATION(World, LogLandscapeEdMode, VeryVerbose, EstimatedHitLocation, 10.0, FColor(100, 100, 255), TEXT("landscape:estimated-hit-location"));
+			UE_CVLOG_LOCATION(CVarVisualLogShowBrushPhysics->GetBool(), World, LogLandscapeEdMode, VeryVerbose, EstimatedHitLocation, 10.0, FColor(100, 100, 255), TEXT("landscape:estimated-hit-location"));
 
 			OutHitLocation = BestHitResult->LandscapeProxy->LandscapeActorToWorld().InverseTransformPosition(EstimatedHitLocation);
 			return true;
diff --git a/Engine/Source/Editor/LandscapeEditor/Private/LandscapeEdModeBrushes.cpp b/Engine/Source/Editor/LandscapeEditor/Private/LandscapeEdModeBrushes.cpp
--- a/Engine/Source/Editor/LandscapeEditor/Private/LandscapeEdModeBrushes.cpp
+++ b/Engine/Source/Editor/LandscapeEditor/Private/LandscapeEdModeBrushes.cpp
@@ -66,7 +66,7 @@ protected:
 	/** Protected so that only subclasses can create instances of this class. */
 	FLandscapeBrushCircle(FEdModeLandscape* InEdMode, UMaterialInterface* InBrushMaterial)
 		: LastMousePosition(0, 0)
-		, BrushMaterial(LandscapeTool::CreateMaterialInstance(InBrushMaterial))
+		, BrushMaterial(UE::Landscape::Editor::Tool::CreateMaterialInstance(InBrushMaterial))
 		, bCanPaint(false)
 		, EdMode(InEdMode)
 	{
@@ -324,23 +324,23 @@ public:
 			InteractorPositions = InInteractorPositions;
 		}
 
-		FIntRect Bounds;
+		FIntRect InclusiveBounds;
 		for (const FLandscapeToolInteractorPosition& InteractorPosition : InteractorPositions)
 		{
-			FIntRect SpotBounds;
-			SpotBounds.Min.X = FMath::FloorToInt32(InteractorPosition.Position.X - TotalRadius);
-			SpotBounds.Min.Y = FMath::FloorToInt32(InteractorPosition.Position.Y - TotalRadius);
-			SpotBounds.Max.X = FMath::CeilToInt32( InteractorPosition.Position.X + TotalRadius);
-			SpotBounds.Max.Y = FMath::CeilToInt32( InteractorPosition.Position.Y + TotalRadius);
+			FIntRect SpotInclusiveBounds;
+			SpotInclusiveBounds.Min.X = FMath::FloorToInt32(InteractorPosition.Position.X - TotalRadius);
+			SpotInclusiveBounds.Min.Y = FMath::FloorToInt32(InteractorPosition.Position.Y - TotalRadius);
+			SpotInclusiveBounds.Max.X = FMath::CeilToInt32(InteractorPosition.Position.X + TotalRadius);
+			SpotInclusiveBounds.Max.Y = FMath::CeilToInt32(InteractorPosition.Position.Y + TotalRadius);
 
-			if (Bounds.IsEmpty())
+			if (InclusiveBounds.IsEmpty())
 			{
-				Bounds = SpotBounds;
+				InclusiveBounds = SpotInclusiveBounds;
 			}
 			else
 			{
-				Bounds.Min = Bounds.Min.ComponentMin(SpotBounds.Min);
-				Bounds.Max = Bounds.Max.ComponentMax(SpotBounds.Max);
+				InclusiveBounds.Min = InclusiveBounds.Min.ComponentMin(SpotInclusiveBounds.Min);
+				InclusiveBounds.Max = InclusiveBounds.Max.ComponentMax(SpotInclusiveBounds.Max);
 			}
 		}
 
@@ -351,23 +351,24 @@ public:
 			// Landscape has no components somehow
 			return FLandscapeBrushData();
 		}
-		Bounds.Clip(FIntRect(MinX, MinY, MaxX + 1, MaxY + 1));
+		InclusiveBounds.Clip(FIntRect(MinX, MinY, MaxX, MaxY));
 
-		FLandscapeBrushData BrushData(Bounds);
+		// FLandscapeBrushData requires exclusive bounds : 
+		FLandscapeBrushData BrushData(FIntRect(InclusiveBounds.Min, InclusiveBounds.Max + 1));
 
 		for (const FLandscapeToolInteractorPosition& InteractorPosition : InteractorPositions)
 		{
-			FIntRect SpotBounds;
-			SpotBounds.Min.X = FMath::Max(FMath::FloorToInt32(InteractorPosition.Position.X - TotalRadius), Bounds.Min.X);
-			SpotBounds.Min.Y = FMath::Max(FMath::FloorToInt32(InteractorPosition.Position.Y - TotalRadius), Bounds.Min.Y);
-			SpotBounds.Max.X = FMath::Min(FMath::CeilToInt32( InteractorPosition.Position.X + TotalRadius), Bounds.Max.X);
-			SpotBounds.Max.Y = FMath::Min(FMath::CeilToInt32( InteractorPosition.Position.Y + TotalRadius), Bounds.Max.Y);
+			FIntRect SpotInclusiveBounds;
+			SpotInclusiveBounds.Min.X = FMath::Max(FMath::FloorToInt32(InteractorPosition.Position.X - TotalRadius), InclusiveBounds.Min.X);
+			SpotInclusiveBounds.Min.Y = FMath::Max(FMath::FloorToInt32(InteractorPosition.Position.Y - TotalRadius), InclusiveBounds.Min.Y);
+			SpotInclusiveBounds.Max.X = FMath::Min(FMath::CeilToInt32(InteractorPosition.Position.X + TotalRadius), InclusiveBounds.Max.X);
+			SpotInclusiveBounds.Max.Y = FMath::Min(FMath::CeilToInt32(InteractorPosition.Position.Y + TotalRadius), InclusiveBounds.Max.Y);
 
-			for (int32 Y = SpotBounds.Min.Y; Y < SpotBounds.Max.Y; Y++)
+			for (int32 Y = SpotInclusiveBounds.Min.Y; Y <= SpotInclusiveBounds.Max.Y; Y++)
 			{
 				float* Scanline = BrushData.GetDataPtr(FIntPoint(0, Y));
 
-				for (int32 X = SpotBounds.Min.X; X < SpotBounds.Max.X; X++)
+				for (int32 X = SpotInclusiveBounds.Min.X; X <= SpotInclusiveBounds.Max.X; X++)
 				{
 					float PrevAmount = Scanline[X];
 					if (PrevAmount < 1.0f)
@@ -431,7 +432,7 @@ public:
 		, EdMode(InEdMode)
 	{
 		UMaterial* BaseBrushMaterial = LoadObject<UMaterial>(nullptr, TEXT("/Engine/EditorLandscapeResources/SelectBrushMaterial.SelectBrushMaterial"));
-		BrushMaterial = LandscapeTool::CreateMaterialInstance(BaseBrushMaterial);
+		BrushMaterial = UE::Landscape::Editor::Tool::CreateMaterialInstance(BaseBrushMaterial);
 	}
 
 	// FGCObject interface
@@ -542,7 +543,7 @@ public:
 		// Selection Brush only works for 
 		ULandscapeInfo* LandscapeInfo = EdMode->CurrentToolTarget.LandscapeInfo.Get();
 
-		FIntRect Bounds;
+		FIntRect ExclusiveBounds(INT_MAX, INT_MAX, INT_MIN, INT_MIN);
 
 		// The add component tool needs the raw bounds of the brush rather than the bounds of the actually existing components under the brush
 		if (EdMode->CurrentTool->GetToolName() == FName("AddComponent"))
@@ -554,10 +555,10 @@ public:
 			const int32 ComponentIndexX = FMath::FloorToInt(BrushOriginX);
 			const int32 ComponentIndexY = FMath::FloorToInt(BrushOriginY);
 
-			Bounds.Min.X = (ComponentIndexX) * LandscapeInfo->ComponentSizeQuads;
-			Bounds.Min.Y = (ComponentIndexY) * LandscapeInfo->ComponentSizeQuads;
-			Bounds.Max.X = (ComponentIndexX + BrushSize) * LandscapeInfo->ComponentSizeQuads + 1;
-			Bounds.Max.Y = (ComponentIndexY + BrushSize) * LandscapeInfo->ComponentSizeQuads + 1;
+			ExclusiveBounds.Min.X = (ComponentIndexX) * LandscapeInfo->ComponentSizeQuads;
+			ExclusiveBounds.Min.Y = (ComponentIndexY) * LandscapeInfo->ComponentSizeQuads;
+			ExclusiveBounds.Max.X = (ComponentIndexX + BrushSize) * LandscapeInfo->ComponentSizeQuads + 1;
+			ExclusiveBounds.Max.Y = (ComponentIndexY + BrushSize) * LandscapeInfo->ComponentSizeQuads + 1;
 		}
 		else
 		{
@@ -567,30 +568,25 @@ public:
 			}
 
 			// Get extent for all components
-			Bounds.Min.X = INT_MAX;
-			Bounds.Min.Y = INT_MAX;
-			Bounds.Max.X = INT_MIN;
-			Bounds.Max.Y = INT_MIN;
-
 			for (ULandscapeComponent* Component : BrushMaterialComponents)
 			{
 				if (ensure(Component))
 				{
-					Component->GetComponentExtent(Bounds.Min.X, Bounds.Min.Y, Bounds.Max.X, Bounds.Max.Y);
+					Component->GetComponentExtent(ExclusiveBounds.Min.X, ExclusiveBounds.Min.Y, ExclusiveBounds.Max.X, ExclusiveBounds.Max.Y);
 				}
 			}
 
 			// GetComponentExtent returns an inclusive max bound
-			Bounds.Max += FIntPoint(1, 1);
+			ExclusiveBounds.Max += FIntPoint(1, 1);
 		}
 
-		FLandscapeBrushData BrushData(Bounds);
+		FLandscapeBrushData BrushData(ExclusiveBounds);
 
-		for (int32 Y = Bounds.Min.Y; Y < Bounds.Max.Y; Y++)
+		for (int32 Y = ExclusiveBounds.Min.Y; Y < ExclusiveBounds.Max.Y; Y++)
 		{
 			float* Scanline = BrushData.GetDataPtr(FIntPoint(0, Y));
 
-			for (int32 X = Bounds.Min.X; X < Bounds.Max.X; X++)
+			for (int32 X = ExclusiveBounds.Min.X; X < ExclusiveBounds.Max.X; X++)
 			{
 				const int32 ComponentIndexX = FMath::FloorToInt(static_cast<float>(X) / LandscapeInfo->ComponentSizeQuads);
 				const int32 ComponentIndexY = FMath::FloorToInt(static_cast<float>(Y) / LandscapeInfo->ComponentSizeQuads);
@@ -646,7 +642,7 @@ public:
 		, EdMode(InEdMode)
 	{
 		UMaterialInterface* GizmoMaterial = LoadObject<UMaterialInstanceConstant>(nullptr, TEXT("/Engine/EditorLandscapeResources/MaskBrushMaterial_Gizmo.MaskBrushMaterial_Gizmo"));
-		BrushMaterial = UMaterialInstanceDynamic::Create(LandscapeTool::CreateMaterialInstance(GizmoMaterial), nullptr);
+		BrushMaterial = UMaterialInstanceDynamic::Create(UE::Landscape::Editor::Tool::CreateMaterialInstance(GizmoMaterial), nullptr);
 	}
 
 	// FGCObject interface
@@ -799,21 +795,16 @@ public:
 		float ScaleXY = static_cast<float>(FMath::Abs(LandscapeInfo->DrawScale.X));
 
 		// Get extent for all components
-		FIntRect Bounds;
-		Bounds.Min.X = INT_MAX;
-		Bounds.Min.Y = INT_MAX;
-		Bounds.Max.X = INT_MIN;
-		Bounds.Max.Y = INT_MIN;
-
+		FIntRect ExclusiveBounds(INT_MAX, INT_MAX, INT_MIN, INT_MIN);
 		for (ULandscapeComponent* Component : BrushMaterialComponents)
 		{
 			if (ensure(Component))
 			{
-				Component->GetComponentExtent(Bounds.Min.X, Bounds.Min.Y, Bounds.Max.X, Bounds.Max.Y);
+				Component->GetComponentExtent(ExclusiveBounds.Min.X, ExclusiveBounds.Min.Y, ExclusiveBounds.Max.X, ExclusiveBounds.Max.Y);
 			}
 		}
 
-		FLandscapeBrushData BrushData(Bounds);
+		FLandscapeBrushData BrushData(ExclusiveBounds);
 
 		//FMatrix LandscapeToGizmoLocal = Landscape->LocalToWorld() * Gizmo->WorldToLocal();
 		const float LW = Gizmo->GetWidth() / (2 * ScaleXY);
@@ -827,11 +818,11 @@ public:
 		float W = Gizmo->GetWidth() / ScaleXY; //Gizmo->GetWidth() / (Gizmo->DrawScale * Gizmo->DrawScale3D.X);
 		float H = Gizmo->GetHeight() / ScaleXY; //Gizmo->GetHeight() / (Gizmo->DrawScale * Gizmo->DrawScale3D.Y);
 
-		for (int32 Y = Bounds.Min.Y; Y < Bounds.Max.Y; Y++)
+		for (int32 Y = ExclusiveBounds.Min.Y; Y < ExclusiveBounds.Max.Y; Y++)
 		{
 			float* Scanline = BrushData.GetDataPtr(FIntPoint(0, Y));
 
-			for (int32 X = Bounds.Min.X; X < Bounds.Max.X; X++)
+			for (int32 X = ExclusiveBounds.Min.X; X < ExclusiveBounds.Max.X; X++)
 			{
 				FVector GizmoLocal = LandscapeToGizmoLocal.TransformPosition(FVector(X, Y, 0));
 				if (GizmoLocal.X < W && GizmoLocal.X > 0 && GizmoLocal.Y < H && GizmoLocal.Y > 0)
@@ -1167,12 +1158,11 @@ public:
 		int32 SizeY = EdMode->UISettings->AlphaTextureSizeY;
 		check(EdMode->UISettings->HasValidAlphaTextureData() && (SizeX > 0) && (SizeY > 0)); // See CanPaint() above : if we can paint, these must all be valid
 
-		FIntRect Bounds;
-		Bounds.Min.X = FMath::FloorToInt(LastMousePosition.X - TotalRadius);
-		Bounds.Min.Y = FMath::FloorToInt(LastMousePosition.Y - TotalRadius);
-		Bounds.Max.X = FMath::CeilToInt( LastMousePosition.X + TotalRadius);
-		Bounds.Max.Y = FMath::CeilToInt( LastMousePosition.Y + TotalRadius);
-
+		FIntRect InclusiveBounds;
+		InclusiveBounds.Min.X = FMath::FloorToInt(LastMousePosition.X - TotalRadius);
+		InclusiveBounds.Min.Y = FMath::FloorToInt(LastMousePosition.Y - TotalRadius);
+		InclusiveBounds.Max.X = FMath::CeilToInt(LastMousePosition.X + TotalRadius);
+		InclusiveBounds.Max.Y = FMath::CeilToInt(LastMousePosition.Y + TotalRadius);
 
 		// Clamp to landscape bounds
 		int32 MinX, MaxX, MinY, MaxY;
@@ -1181,15 +1171,16 @@ public:
 			// Landscape has no components somehow
 			return FLandscapeBrushData();
 		}
-		Bounds.Clip(FIntRect(MinX, MinY, MaxX + 1, MaxY + 1));
+		InclusiveBounds.Clip(FIntRect(MinX, MinY, MaxX, MaxY));
 
-		FLandscapeBrushData BrushData(Bounds);
+		// FLandscapeBrushData requires exclusive bounds : 
+		FLandscapeBrushData BrushData(FIntRect(InclusiveBounds.Min, InclusiveBounds.Max + 1));
 
-		for (int32 Y = Bounds.Min.Y; Y < Bounds.Max.Y; Y++)
+		for (int32 Y = InclusiveBounds.Min.Y; Y <= InclusiveBounds.Max.Y; Y++)
 		{
 			float* Scanline = BrushData.GetDataPtr(FIntPoint(0, Y));
 
-			for (int32 X = Bounds.Min.X; X < Bounds.Max.X; X++)
+			for (int32 X = InclusiveBounds.Min.X; X <= InclusiveBounds.Max.X; X++)
 			{
 				float Angle;
 				FVector2D Scale;
@@ -1445,11 +1436,11 @@ public:
 			check(bIsValid && !FMath::IsNearlyZero(AlphaBrushScale) && !FMath::IsNearlyZero(Radius) && (SizeX > 0) && (SizeY > 0)); // See CanPaint() function above : if bCanPaint, AlphaBrushScale should be non-zero
 
 			const float BrushAngle = EdMode->UISettings->bAlphaBrushAutoRotate ? LastMouseAngle : FMath::DegreesToRadians(EdMode->UISettings->AlphaBrushRotation);
-			FIntRect Bounds;
-			Bounds.Min.X = FMath::FloorToInt(LastMousePosition.X - Radius);
-			Bounds.Min.Y = FMath::FloorToInt(LastMousePosition.Y - Radius);
-			Bounds.Max.X = FMath::CeilToInt( LastMousePosition.X + Radius);
-			Bounds.Max.Y = FMath::CeilToInt( LastMousePosition.Y + Radius);
+			FIntRect InclusiveBounds;
+			InclusiveBounds.Min.X = FMath::FloorToInt(LastMousePosition.X - Radius);
+			InclusiveBounds.Min.Y = FMath::FloorToInt(LastMousePosition.Y - Radius);
+			InclusiveBounds.Max.X = FMath::CeilToInt(LastMousePosition.X + Radius);
+			InclusiveBounds.Max.Y = FMath::CeilToInt(LastMousePosition.Y + Radius);
 
 			// Clamp to landscape bounds
 			int32 MinX, MaxX, MinY, MaxY;
@@ -1458,15 +1449,16 @@ public:
 				// Landscape has no components somehow
 				return FLandscapeBrushData();
 			}
-			Bounds.Clip(FIntRect(MinX, MinY, MaxX + 1, MaxY + 1));
+			InclusiveBounds.Clip(FIntRect(MinX, MinY, MaxX, MaxY));
 
-			FLandscapeBrushData BrushData(Bounds);
+			// FLandscapeBrushData requires exclusive bounds : 
+			FLandscapeBrushData BrushData(FIntRect(InclusiveBounds.Min, InclusiveBounds.Max + 1));
 
-			for (int32 Y = Bounds.Min.Y; Y < Bounds.Max.Y; Y++)
+			for (int32 Y = InclusiveBounds.Min.Y; Y <= InclusiveBounds.Max.Y; Y++)
 			{
 				float* Scanline = BrushData.GetDataPtr(FIntPoint(0, Y));
 
-				for (int32 X = Bounds.Min.X; X < Bounds.Max.X; X++)
+				for (int32 X = InclusiveBounds.Min.X; X <= InclusiveBounds.Max.X; X++)
 				{
 					// Find alphamap sample location
 					float ScaleSampleX = ((float)X - LastMousePosition.X) / AlphaBrushScale;
diff --git a/Engine/Source/Editor/LandscapeEditor/Private/LandscapeEdModeTools.h b/Engine/Source/Editor/LandscapeEditor/Private/LandscapeEdModeTools.h
--- a/Engine/Source/Editor/LandscapeEditor/Private/LandscapeEdModeTools.h
+++ b/Engine/Source/Editor/LandscapeEditor/Private/LandscapeEdModeTools.h
@@ -306,11 +306,11 @@ public:
 		{
 			if (Accessor::bUseInterp)
 			{
-				ValidX1 = CachedX1 = X1;
-				ValidY1 = CachedY1 = Y1;
-				ValidX2 = CachedX2 = X2;
-				ValidY2 = CachedY2 = Y2;
-
+				// GetData alters its args, so make temp copies to avoid screwing things up
+				int32 ValidX1 = X1;
+				int32 ValidY1 = Y1;
+				int32 ValidX2 = X2;
+				int32 ValidY2 = Y2;
 				DataAccess.GetData(ValidX1, ValidY1, ValidX2, ValidY2, CachedData);
 				if (!ensureMsgf(ValidX1 <= ValidX2 && ValidY1 <= ValidY2, TEXT("Invalid cache area: X(%d-%d), Y(%d-%d) from region X(%d-%d), Y(%d-%d)"), ValidX1, ValidX2, ValidY1, ValidY2, X1, X2, Y1, Y2))
 				{
@@ -320,16 +320,15 @@ public:
 			}
 			else
 			{
-				CachedX1 = X1;
-				CachedY1 = Y1;
-				CachedX2 = X2;
-				CachedY2 = Y2;
-
 				DataAccess.GetDataFast(CachedX1, CachedY1, CachedX2, CachedY2, CachedData);
 			}
+			CachedX1 = X1;
+			CachedY1 = Y1;
+			CachedX2 = X2;
+			CachedY2 = Y2;
 
 			// Drop a visual log to indicate the area covered by this cache region extension :
-			VisualizeLandscapeRegion(CachedX1, CachedY1, CachedX2, CachedY2, FColor::Red, TEXT("Cache Data"));
+			VisualizeLandscapeRegion(CachedX1, CachedY1, CachedX2, CachedY2, FColor::Red, FString::Printf(TEXT("Cache Data (X1:%i, Y1:%i, X2:%i, Y2:%i)"), CachedX1, CachedY1, CachedX2, CachedY2));
 
 			if (bCacheOriginalData)
 			{
@@ -339,114 +338,64 @@ public:
 		}
 		else
 		{
-			bool bCacheExtended = false;
-
-			// Extend the cache area if needed
-			if (X1 < CachedX1)
+			auto ExtendCache = [this, bCacheOriginalData](const int32 InX1, const int32 InY1, const int32 InX2, const int32 InY2)
 			{
 				if (Accessor::bUseInterp)
 				{
-					int32 x1 = X1;
-					int32 x2 = ValidX1;
-					int32 y1 = FMath::Min<int32>(Y1, CachedY1);
-					int32 y2 = FMath::Max<int32>(Y2, CachedY2);
-
-					DataAccess.GetData(x1, y1, x2, y2, CachedData);
-					ValidX1 = FMath::Min<int32>(x1, ValidX1);
+					// GetData alters its args, so make temp copies to avoid screwing things up :
+					int32 ValidX1 = InX1;
+					int32 ValidY1 = InY1;
+					int32 ValidX2 = InX2;
+					int32 ValidY2 = InY2;
+					DataAccess.GetData(ValidX1, ValidY1, ValidX2, ValidY2, CachedData);
 				}
 				else
 				{
-					DataAccess.GetDataFast(X1, CachedY1, CachedX1 - 1, CachedY2, CachedData);
+					DataAccess.GetDataFast(InX1, InY1, InX2, InY2, CachedData);
 				}
 
 				if (bCacheOriginalData)
 				{
-					CacheOriginalData(X1, CachedY1, CachedX1 - 1, CachedY2);
+					CacheOriginalData(InX1, InY1, InX2, InY2);
+					// Drop a visual log to indicate the area covered by this cache region extension :
+					VisualizeLandscapeRegion(InX1, InY1, InX2, InY2, FColor::Purple, FString::Printf(TEXT("Cache Original Data (X1:%i, Y1:%i, X2:%i, Y2:%i)"), InX1, InY1, InX2, InY2), /*InZOffset = */10.0);
 				}
-				CachedX1 = X1;
+			};
 
+			// Extend the cache area if needed
+			bool bCacheExtended = false;
+			if (X1 < CachedX1)
+			{
+				ExtendCache(X1, CachedY1, CachedX1 - 1, CachedY2);
+				CachedX1 = X1;
 				bCacheExtended = true;
 			}
 
 			if (X2 > CachedX2)
 			{
-				if (Accessor::bUseInterp)
-				{
-					int32 x1 = ValidX2;
-					int32 x2 = X2;
-					int32 y1 = FMath::Min<int32>(Y1, CachedY1);
-					int32 y2 = FMath::Max<int32>(Y2, CachedY2);
-
-					DataAccess.GetData(x1, y1, x2, y2, CachedData);
-					ValidX2 = FMath::Max<int32>(x2, ValidX2);
-				}
-				else
-				{
-					DataAccess.GetDataFast(CachedX2 + 1, CachedY1, X2, CachedY2, CachedData);
-				}
-				if (bCacheOriginalData)
-				{
-					CacheOriginalData(CachedX2 + 1, CachedY1, X2, CachedY2);
-				}
+				ExtendCache(CachedX2 + 1, CachedY1, X2, CachedY2);
 				CachedX2 = X2;
-
 				bCacheExtended = true;
 			}
 
 			if (Y1 < CachedY1)
 			{
-				if (Accessor::bUseInterp)
-				{
-					int32 x1 = CachedX1;
-					int32 x2 = CachedX2;
-					int32 y1 = Y1;
-					int32 y2 = ValidY1;
-
-					DataAccess.GetData(x1, y1, x2, y2, CachedData);
-					ValidY1 = FMath::Min<int32>(y1, ValidY1);
-				}
-				else
-				{
-					DataAccess.GetDataFast(CachedX1, Y1, CachedX2, CachedY1 - 1, CachedData);
-				}
-				if (bCacheOriginalData)
-				{
-					CacheOriginalData(CachedX1, Y1, CachedX2, CachedY1 - 1);
-				}
+				ExtendCache(CachedX1, Y1, CachedX2, CachedY1 - 1);
 				CachedY1 = Y1;
-
 				bCacheExtended = true;
 			}
 
 			if (Y2 > CachedY2)
 			{
-				if (Accessor::bUseInterp)
-				{
-					int32 x1 = CachedX1;
-					int32 x2 = CachedX2;
-					int32 y1 = ValidY2;
-					int32 y2 = Y2;
-
-					DataAccess.GetData(x1, y1, x2, y2, CachedData);
-					ValidY2 = FMath::Max<int32>(y2, ValidY2);
-				}
-				else
-				{
-					DataAccess.GetDataFast(CachedX1, CachedY2 + 1, CachedX2, Y2, CachedData);
-				}
-				if (bCacheOriginalData)
-				{
-					CacheOriginalData(CachedX1, CachedY2 + 1, CachedX2, Y2);
-				}
+				ExtendCache(CachedX1, CachedY2 + 1, CachedX2, Y2);
 				CachedY2 = Y2;
-
 				bCacheExtended = true;
 			}
 
 			if (bCacheExtended)
 			{
 				// Drop a visual log to indicate the area covered by this cache region extension :
-				VisualizeLandscapeRegion(CachedX1, CachedY1, CachedX2, CachedY2, FColor::Red, TEXT("Cache Data"));
+				VisualizeLandscapeRegion(CachedX1, CachedY1, CachedX2, CachedY2, FColor::Red, FString::Printf(TEXT("Cache Data (X1:%i, Y1:%i, X2:%i, Y2:%i)"), CachedX1, CachedY1, CachedX2, CachedY2));
 			}
 		}
 	}
@@ -627,10 +576,7 @@ public:
 			for (int32 X = X1; X <= X2; X++)
 			{
 				AccessorType* Ptr = OriginalData.Find(FIntPoint(X, Y));
-				if (Ptr)
-				{
-					OutOriginalData[(X - X1) + (Y - Y1)*(1 + X2 - X1)] = *Ptr;
-				}
+				OutOriginalData[(X - X1) + (Y - Y1) * (1 + X2 - X1)] = Ptr ? *Ptr : (AccessorType)0;
 			}
 		}
 	}
@@ -660,15 +606,21 @@ private:
 		}
 	}
 
-	void VisualizeLandscapeRegion(int32 InX1, int32 InY1, int32 InX2, int32 InY2, const FColor& InColor, const FString& InDescription)
+	// X2/Y2 Coordinates are "inclusive" max values
+	void VisualizeLandscapeRegion(int32 InX1, int32 InY1, int32 InX2, int32 InY2, const FColor& InColor, const FString& InDescription, double InZOffset = 0.0f)
 	{
 		check(LandscapeInfo != nullptr);
 		ALandscapeProxy* LandscapeProxy = LandscapeInfo->GetLandscapeProxy();
 		check(LandscapeProxy != nullptr);
 		const FTransform& LandscapeTransform = LandscapeProxy->GetTransform();
-		FVector Min = LandscapeTransform.TransformPosition(FVector(InX1, InY1, 0));
-		FVector Max = LandscapeTransform.TransformPosition(FVector(InX2, InY2, 0));
-		UE_VLOG_BOX(LandscapeProxy, LogLandscapeTools, Log, FBox(Min, Max), InColor, TEXT("%s"), *InDescription);
+		static constexpr double BaseZOffset = 10.0; // Add a small Z offset to avoid Z-flickering for flat landscapes
+		FVector ZOffset(0.0, 0.0, BaseZOffset + InZOffset);
+		// the offset is given in world space so unapply the scale before applying the transform
+		ZOffset /= LandscapeTransform.GetScale3D();
+		FTransform WorldTransformForVisLog = FTransform(ZOffset) * LandscapeTransform;
+		// Use (-0.5, -0.5) for min and (+0.5, +0.5) for max because both are inclusive (i.e. caching at 0, 0, 0, 0 (X1, Y1, X2, Y2)) will draw a box of size (1, 1) around (0, 0))
+		UE_VLOG_OBOX(LandscapeProxy, LogLandscapeTools, Log, FBox(FVector(InX1, InY1, 0.0) - FVector(0.5, 0.5, 0.0), FVector(InX2, InY2, 0.0) + FVector(0.5, 0.5, 0.0)),
+			WorldTransformForVisLog.ToMatrixWithScale(), InColor, TEXT("%s"), *InDescription);
 	}
 
 	TMap<FIntPoint, AccessorType> CachedData;
@@ -678,16 +630,11 @@ private:
 
 	bool bIsValid = false;
 
+	// Inclusive bounds of the current cached region (CachedData)
 	int32 CachedX1 = INDEX_NONE;
 	int32 CachedY1 = INDEX_NONE;
 	int32 CachedX2 = INDEX_NONE;
 	int32 CachedY2 = INDEX_NONE;
-
-	// To store valid region....
-	int32 ValidX1 = INDEX_NONE;
-	int32 ValidX2 = INDEX_NONE;
-	int32 ValidY1 = INDEX_NONE;
-	int32 ValidY2 = INDEX_NONE;
 };
 
 template<bool bInUseInterp>
@@ -895,14 +842,6 @@ struct FLandscapeAlphaCache : public TLandscapeEditCache<FAlphamapAccessorTool<t
 	}
 };
 
-struct FVisibilityAccessor : public FAlphamapAccessorTool<false, false>
-{
-	FVisibilityAccessor(const FLandscapeToolTarget& InTarget)
-		: FAlphamapAccessorTool<false, false>(InTarget.LandscapeInfo.Get(), ALandscapeProxy::VisibilityLayer)
-	{
-	}
-};
-
 struct FLandscapeVisCache : public TLandscapeEditCache<FAlphamapAccessorTool<false, false>, uint8>
 {
 	static uint8 ClampValue(int32 Value) { return static_cast<uint8>(FMath::Clamp(Value, 0, 255)); }
diff --git a/Engine/Source/Editor/LandscapeEditor/Public/LandscapeToolInterface.h b/Engine/Source/Editor/LandscapeEditor/Public/LandscapeToolInterface.h
--- a/Engine/Source/Editor/LandscapeEditor/Public/LandscapeToolInterface.h
+++ b/Engine/Source/Editor/LandscapeEditor/Public/LandscapeToolInterface.h
@@ -273,7 +273,8 @@ public:
 	TArray<FName> ValidBrushes;
 };
 
-namespace LandscapeTool
+namespace UE::Landscape::Editor::Tool
 {
 	UMaterialInstance* CreateMaterialInstance(UMaterialInterface* BaseMaterial);
 }
+// namespace UE::Landscape::Editor::Tool
diff --git a/Engine/Source/Runtime/Landscape/Private/LandscapeEditInterface.cpp b/Engine/Source/Runtime/Landscape/Private/LandscapeEditInterface.cpp
--- a/Engine/Source/Runtime/Landscape/Private/LandscapeEditInterface.cpp
+++ b/Engine/Source/Runtime/Landscape/Private/LandscapeEditInterface.cpp
@@ -1471,10 +1471,10 @@ void FLandscapeEditDataInterface::GetHeightDataInternal(int32& ValidX1, int32& V
 			ComponentSizeX, ComponentSizeY, CornerValues,
 			NoBorderY1, NoBorderY2, ComponentDataExist, StoreData);
 		// Update valid region
-		ValidX1 = FMath::Max<int32>(X1, ValidX1);
-		ValidX2 = FMath::Min<int32>(X2, ValidX2);
-		ValidY1 = FMath::Max<int32>(Y1, ValidY1);
-		ValidY2 = FMath::Min<int32>(Y2, ValidY2);
+		ValidX1 = FMath::Min<int32>(X1, ValidX1);
+		ValidX2 = FMath::Max<int32>(X2, ValidX2);
+		ValidY1 = FMath::Min<int32>(Y1, ValidY1);
+		ValidY2 = FMath::Max<int32>(Y2, ValidY2);
 	}
 	else
 	{
