Fixed crash in landscape alpha brush thumbnails which wrongly assumed that the texture's resource was loaded. Now the thumbnail viewport will keep on being invalidated until it can render the texture

diff --git a/Engine/Source/Editor/LandscapeEditor/Private/LandscapeEditorDetailCustomization_AlphaBrush.cpp b/Engine/Source/Editor/LandscapeEditor/Private/LandscapeEditorDetailCustomization_AlphaBrush.cpp
--- a/Engine/Source/Editor/LandscapeEditor/Private/LandscapeEditorDetailCustomization_AlphaBrush.cpp
+++ b/Engine/Source/Editor/LandscapeEditor/Private/LandscapeEditorDetailCustomization_AlphaBrush.cpp
@@ -102,6 +102,9 @@ public:
 
 		// The viewport widget needs an interface so it knows what should render
 		ViewportWidget->SetViewportInterface( Viewport.ToSharedRef() );
+
+		UTexture2D* Texture2D = Texture.Get();
+		bCanRender = (Texture2D && !Texture2D->IsCompiling() && Texture2D->GetResource());
 	}
 
 	FReply OnAssetThumbnailDoubleClick(const FGeometry& InMyGeometry, const FPointerEvent& InMouseEvent)
@@ -124,11 +127,10 @@ public:
 
 	virtual int32 OnPaint(const FPaintArgs& Args, const FGeometry& AllottedGeometry, const FSlateRect& MyCullingRect, FSlateWindowElementList& OutDrawElements, int32 LayerId, const FWidgetStyle& InWidgetStyle, bool bParentEnabled) const
 	{
+		UTexture2D* NewTexture = Texture.Get();
+		int32 NewTextureChannel = TextureChannel.Get();
 		if (Texture.IsBound() || TextureChannel.IsBound())
 		{
-			UTexture2D* NewTexture = Texture.Get();
-			int32 NewTextureChannel = TextureChannel.Get();
-
 			if (NewTexture != CachedTexture
 				|| NewTextureChannel != CachedTextureChannel)
 			{
@@ -138,6 +140,14 @@ public:
 			}
 		}
 
+		bool bCanRenderBefore = bCanRender;
+		bCanRender = (NewTexture && !NewTexture->IsCompiling() && NewTexture->GetResource());
+		// If we were waiting for the texture to become ready and it now is, invalidate the viewport to trigger a draw :
+		if (!bCanRenderBefore && bCanRender)
+		{
+			Viewport->Invalidate();
+		}
+
 		return SCompoundWidget::OnPaint(Args, AllottedGeometry, MyCullingRect, OutDrawElements, LayerId, InWidgetStyle, bParentEnabled);
 	}
 
@@ -147,6 +157,7 @@ private:
 
 	mutable UTexture2D* CachedTexture;
 	mutable uint8 CachedTextureChannel;
+	mutable bool bCanRender = false;
 
 	TSharedPtr<FTextureMaskThumbnailViewportClient> ViewportClient;
 	TSharedPtr<FSceneViewport> Viewport;
@@ -163,17 +174,19 @@ void FTextureMaskThumbnailViewportClient::Draw(FViewport* Viewport, FCanvas* Can
 
 	Canvas->Clear( FLinearColor::Black);
 
-	if (UTexture2D* Texture = PinnedParent->Texture.Get())
+	if (UTexture2D* Texture = PinnedParent->Texture.Get(); Texture && !Texture->IsCompiling())
 	{
 		// Fully stream in the texture before drawing it.
-		FTextureCompilingManager::Get().FinishCompilation({ Texture });
 		Texture->SetForceMipLevelsToBeResident(30.0f);
 		Texture->WaitForStreaming();
 
-		//Draw the selected texture, uses ColourChannelBlend mode parameter to filter colour channels and apply grayscale
-		FCanvasTileItem TileItem( FVector2D( 0.0f, 0.0f ), Texture->GetResource(), Viewport->GetSizeXY(), FLinearColor::White );
-		TileItem.BlendMode = (ESimpleElementBlendMode)(SE_BLEND_RGBA_MASK_START + (1<<PinnedParent->TextureChannel.Get()) + 16);
-		Canvas->DrawItem( TileItem );
+		if (FTextureResource* TextureResource = Texture->GetResource())
+		{
+			//Draw the selected texture, uses ColourChannelBlend mode parameter to filter colour channels and apply grayscale
+			FCanvasTileItem TileItem( FVector2D( 0.0f, 0.0f ), TextureResource, Viewport->GetSizeXY(), FLinearColor::White );
+			TileItem.BlendMode = (ESimpleElementBlendMode)(SE_BLEND_RGBA_MASK_START + (1<<PinnedParent->TextureChannel.Get()) + 16);
+			Canvas->DrawItem( TileItem );
+		}
 	}
 }
 
