Fix for non-interpolated landscape brush tools

diff --git a/Engine/Source/Editor/LandscapeEditor/Private/LandscapeEdModeTools.h b/Engine/Source/Editor/LandscapeEditor/Private/LandscapeEdModeTools.h
--- a/Engine/Source/Editor/LandscapeEditor/Private/LandscapeEdModeTools.h
+++ b/Engine/Source/Editor/LandscapeEditor/Private/LandscapeEdModeTools.h
@@ -320,7 +320,7 @@ public:
 			}
 			else
 			{
-				DataAccess.GetDataFast(CachedX1, CachedY1, CachedX2, CachedY2, CachedData);
+				DataAccess.GetDataFast(X1, Y1, X2, Y2, CachedData);
 			}
 			CachedX1 = X1;
 			CachedY1 = Y1;
