Removed useless cache in landscape "select components" brush

diff --git a/Engine/Source/Editor/LandscapeEditor/Private/LandscapeEdModeComponentTools.cpp b/Engine/Source/Editor/LandscapeEditor/Private/LandscapeEdModeComponentTools.cpp
--- a/Engine/Source/Editor/LandscapeEditor/Private/LandscapeEdModeComponentTools.cpp
+++ b/Engine/Source/Editor/LandscapeEditor/Private/LandscapeEdModeComponentTools.cpp
@@ -42,7 +42,6 @@ public:
 		: Super(InEdMode, InViewportClient, InTarget)
 		, bInitializedComponentInvert(false)
 		, bNeedsSelectionUpdate(false)
-		, Cache(InTarget)
 	{
 	}
 
@@ -122,9 +121,6 @@ public:
 			bNeedsSelectionUpdate = true;
 		}
 	}
-
-protected:
-	FLandscapeDataCache Cache;
 };
 
 class FLandscapeToolSelect : public FLandscapeToolBase<FLandscapeToolStrokeSelect>
