Various UX improvements for landscape "Select" tool :
* Don't display target layers, edit layers or BP brushes when using the Select tool, since it makes no sense there
* Display the Select tool in the Paint panel since it applies there too
* Refresh the landscape mode details panel when clearing the selection in order to hide the Select tool options immediately
* Moved the Select tool-specific options into a "Select Mask" category
* Moved the Select button to the end of the tool list for both Sculpt and Paint panel, since this is not really an editing tool and acts in both modes

diff --git a/Engine/Source/Editor/LandscapeEditor/Private/LandscapeEdMode.cpp b/Engine/Source/Editor/LandscapeEditor/Private/LandscapeEdMode.cpp
--- a/Engine/Source/Editor/LandscapeEditor/Private/LandscapeEdMode.cpp
+++ b/Engine/Source/Editor/LandscapeEditor/Private/LandscapeEdMode.cpp
@@ -304,6 +304,7 @@ void FEdModeLandscape::UpdateToolModes()
 	ToolMode_Paint->ValidTools.Add(TEXT("Flatten"));
 	ToolMode_Paint->ValidTools.Add(TEXT("Noise"));
 	ToolMode_Paint->ValidTools.Add(TEXT("Visibility"));
+	ToolMode_Paint->ValidTools.Add(TEXT("Mask"));
 
 	if (CanHaveLandscapeLayersContent())
 	{
@@ -3736,7 +3737,7 @@ ALandscape* FEdModeLandscape::ChangeComponentSetting(int32 NumComponentsX, int32
 
 			// LandscapeGuid is stomped by CopySharedProperties, but original guid is not -- fix the mismatch or it will complain during Import
  			NewLandscape->SetLandscapeGuid(FGuid(), /* bValidateGuid= */ false);
-			
+
 			// Copy settings that are not copied by CopySharedProperties
 			NewLandscape->ExportLOD = OldLandscape->ExportLOD;
 			NewLandscape->StaticLightingLOD = OldLandscape->StaticLightingLOD;
diff --git a/Engine/Source/Editor/LandscapeEditor/Private/LandscapeEditorDetailCustomization_Layers.cpp b/Engine/Source/Editor/LandscapeEditor/Private/LandscapeEditorDetailCustomization_Layers.cpp
--- a/Engine/Source/Editor/LandscapeEditor/Private/LandscapeEditorDetailCustomization_Layers.cpp
+++ b/Engine/Source/Editor/LandscapeEditor/Private/LandscapeEditorDetailCustomization_Layers.cpp
@@ -67,23 +67,23 @@ void FLandscapeEditorDetailCustomization_Layers::CustomizeDetails(IDetailLayoutB
 	IDetailCategoryBuilder& LayerCategory = DetailBuilder.EditCategory("Edit Layers");
 
 	FEdModeLandscape* LandscapeEdMode = GetEditorMode();
-	if (LandscapeEdMode && LandscapeEdMode->CurrentToolMode != nullptr)
+	if (LandscapeEdMode 
+		&& LandscapeEdMode->GetLandscape() 
+		&& (LandscapeEdMode->CurrentToolMode != nullptr)
+		&& (FName(LandscapeEdMode->CurrentTool->GetToolName()) != TEXT("Mask")))
 	{
-		if (LandscapeEdMode->GetLandscape())
-		{
-			LayerCategory.AddCustomBuilder(MakeShareable(new FLandscapeEditorCustomNodeBuilder_Layers(DetailBuilder.GetThumbnailPool().ToSharedRef())));
+		LayerCategory.AddCustomBuilder(MakeShareable(new FLandscapeEditorCustomNodeBuilder_Layers(DetailBuilder.GetThumbnailPool().ToSharedRef())));
 
-			LayerCategory.AddCustomRow(FText())
-				.Visibility(TAttribute<EVisibility>::Create(TAttribute<EVisibility>::FGetter::CreateLambda([]() { return ShouldShowLayersErrorMessageTip() ? EVisibility::Visible : EVisibility::Collapsed; })))
-				[
-					SNew(SMultiLineEditableTextBox)
-					.IsReadOnly(true)
-					.Font(DetailBuilder.GetDetailFontBold())
-					.BackgroundColor(TAttribute<FSlateColor>::Create(TAttribute<FSlateColor>::FGetter::CreateLambda([]() { return FAppStyle::GetColor("ErrorReporting.WarningBackgroundColor"); })))
-					.Text(TAttribute<FText>::Create(TAttribute<FText>::FGetter::CreateStatic(&FLandscapeEditorDetailCustomization_Layers::GetLayersErrorMessageText)))
-					.AutoWrapText(true)
-				];
-		}
+		LayerCategory.AddCustomRow(FText())
+			.Visibility(TAttribute<EVisibility>::Create(TAttribute<EVisibility>::FGetter::CreateLambda([]() { return ShouldShowLayersErrorMessageTip() ? EVisibility::Visible : EVisibility::Collapsed; })))
+			[
+				SNew(SMultiLineEditableTextBox)
+				.IsReadOnly(true)
+				.Font(DetailBuilder.GetDetailFontBold())
+				.BackgroundColor(TAttribute<FSlateColor>::Create(TAttribute<FSlateColor>::FGetter::CreateLambda([]() { return FAppStyle::GetColor("ErrorReporting.WarningBackgroundColor"); })))
+				.Text(TAttribute<FText>::Create(TAttribute<FText>::FGetter::CreateStatic(&FLandscapeEditorDetailCustomization_Layers::GetLayersErrorMessageText)))
+				.AutoWrapText(true)
+			];
 	}
 }
 END_SLATE_FUNCTION_BUILD_OPTIMIZATION
diff --git a/Engine/Source/Editor/LandscapeEditor/Private/LandscapeEditorDetailCustomization_LayersBrushStack.cpp b/Engine/Source/Editor/LandscapeEditor/Private/LandscapeEditorDetailCustomization_LayersBrushStack.cpp
--- a/Engine/Source/Editor/LandscapeEditor/Private/LandscapeEditorDetailCustomization_LayersBrushStack.cpp
+++ b/Engine/Source/Editor/LandscapeEditor/Private/LandscapeEditorDetailCustomization_LayersBrushStack.cpp
@@ -42,7 +42,10 @@ BEGIN_SLATE_FUNCTION_BUILD_OPTIMIZATION
 void FLandscapeEditorDetailCustomization_LayersBrushStack::CustomizeDetails(IDetailLayoutBuilder& DetailBuilder)
 {
 	FEdModeLandscape* LandscapeEdMode = GetEditorMode();
-	if (LandscapeEdMode && LandscapeEdMode->CurrentToolMode != nullptr)
+	if (LandscapeEdMode
+		&& LandscapeEdMode->GetLandscape()
+		&& (LandscapeEdMode->CurrentToolMode != nullptr)
+		&& (FName(LandscapeEdMode->CurrentTool->GetToolName()) != TEXT("Mask")))
 	{
 		IDetailCategoryBuilder& LayerCategory = DetailBuilder.EditCategory(FName("Edit Layer Blueprint Brushes"));
 
diff --git a/Engine/Source/Editor/LandscapeEditor/Private/LandscapeEditorDetailCustomization_MiscTools.cpp b/Engine/Source/Editor/LandscapeEditor/Private/LandscapeEditorDetailCustomization_MiscTools.cpp
--- a/Engine/Source/Editor/LandscapeEditor/Private/LandscapeEditorDetailCustomization_MiscTools.cpp
+++ b/Engine/Source/Editor/LandscapeEditor/Private/LandscapeEditorDetailCustomization_MiscTools.cpp
@@ -45,9 +45,9 @@ void FLandscapeEditorDetailCustomization_MiscTools::CustomizeDetails(IDetailLayo
 		];
 	}
 
-	//if (IsToolActive("Mask"))
+	IDetailCategoryBuilder& SelectMaskCategory = DetailBuilder.EditCategory("Select Mask");
 	{
-		ToolsCategory.AddCustomRow(LOCTEXT("Mask.ClearSelection", "Clear Region Selection"))
+		SelectMaskCategory.AddCustomRow(LOCTEXT("Mask.ClearSelection", "Clear Region Selection"))
 		.Visibility(TAttribute<EVisibility>::Create(TAttribute<EVisibility>::FGetter::CreateStatic(&FLandscapeEditorDetailCustomization_MiscTools::GetClearRegionSelectionVisibility)))
 		[
 			SNew(SButton)
diff --git a/Engine/Source/Editor/LandscapeEditor/Private/LandscapeEditorDetailCustomization_TargetLayers.cpp b/Engine/Source/Editor/LandscapeEditor/Private/LandscapeEditorDetailCustomization_TargetLayers.cpp
--- a/Engine/Source/Editor/LandscapeEditor/Private/LandscapeEditorDetailCustomization_TargetLayers.cpp
+++ b/Engine/Source/Editor/LandscapeEditor/Private/LandscapeEditorDetailCustomization_TargetLayers.cpp
@@ -118,7 +118,9 @@ bool FLandscapeEditorDetailCustomization_TargetLayers::ShouldShowTargetLayers()
 
 		//// Visible if there are possible choices
 		//if (bSupportsWeightmap || bSupportsHeightmap || bSupportsVisibility)
-		if (LandscapeEdMode->CurrentToolMode->SupportedTargetTypes != 0 && CurrentToolName != TEXT("BlueprintBrush"))
+		if (LandscapeEdMode->CurrentToolMode->SupportedTargetTypes != 0 
+			&& CurrentToolName != TEXT("BlueprintBrush")
+			&& CurrentToolName != TEXT("Mask"))
 		{
 			return true;
 		}
diff --git a/Engine/Source/Editor/LandscapeEditor/Private/LandscapeEditorDetails.cpp b/Engine/Source/Editor/LandscapeEditor/Private/LandscapeEditorDetails.cpp
--- a/Engine/Source/Editor/LandscapeEditor/Private/LandscapeEditorDetails.cpp
+++ b/Engine/Source/Editor/LandscapeEditor/Private/LandscapeEditorDetails.cpp
@@ -27,7 +27,6 @@
 
 #include "SLandscapeEditor.h"
 #include "LandscapeEditorCommands.h"
-#include "LandscapeEditorDetailWidgets.h"
 #include "LandscapeEditorDetailCustomization_LayersBrushStack.h"
 #include "LandscapeEditorObject.h"
 #include "Landscape.h"
@@ -75,6 +74,7 @@ void FLandscapeEditorDetails::CustomizeDetails(IDetailLayoutBuilder& DetailBuild
 		
 		SafeSetOrder(FName("Tool Settings"));
 		SafeSetOrder(FName("Brush Settings"));
+		SafeSetOrder(FName("Select Mask"));
 		
 		SafeSetOrder(FName("Edit Layers"));
 		SafeSetOrder(FName("Edit Layer Blueprint Brushes"));
diff --git a/Engine/Source/Editor/LandscapeEditor/Private/SLandscapeEditor.cpp b/Engine/Source/Editor/LandscapeEditor/Private/SLandscapeEditor.cpp
--- a/Engine/Source/Editor/LandscapeEditor/Private/SLandscapeEditor.cpp
+++ b/Engine/Source/Editor/LandscapeEditor/Private/SLandscapeEditor.cpp
@@ -241,10 +241,9 @@ void FLandscapeToolKit::BuildToolPalette(FName PaletteName, class FToolBarBuilde
 		}
 
 		ToolBarBuilder.AddToolBarButton(Commands.MirrorTool);
-
-		ToolBarBuilder.AddToolBarButton(Commands.RegionSelectTool);
 		ToolBarBuilder.AddToolBarButton(Commands.RegionCopyPasteTool);
 
+		ToolBarBuilder.AddToolBarButton(Commands.RegionSelectTool);
 	}
 
 	else if (PaletteName == LandscapeEditorNames::Paint)
@@ -253,11 +252,13 @@ void FLandscapeToolKit::BuildToolPalette(FName PaletteName, class FToolBarBuilde
 		ToolBarBuilder.AddToolBarButton(Commands.SmoothTool);
 		ToolBarBuilder.AddToolBarButton(Commands.FlattenTool);
 		ToolBarBuilder.AddToolBarButton(Commands.NoiseTool);
-		
+
 		if (LandscapeEdMode->CanHaveLandscapeLayersContent() && Settings->AreBlueprintToolsAllowed())
 		{
 			ToolBarBuilder.AddToolBarButton(Commands.BlueprintBrushTool);
 		}
+
+		ToolBarBuilder.AddToolBarButton(Commands.RegionSelectTool);
 	}
 
 }
diff --git a/Engine/Source/Editor/LandscapeEditor/Public/LandscapeEditorObject.h b/Engine/Source/Editor/LandscapeEditor/Public/LandscapeEditorObject.h
--- a/Engine/Source/Editor/LandscapeEditor/Public/LandscapeEditorObject.h
+++ b/Engine/Source/Editor/LandscapeEditor/Public/LandscapeEditorObject.h
@@ -444,12 +444,12 @@ class ULandscapeEditorObject : public UObject
 	// Mask Tool:
 
 	// Uses selected region as a mask for other tools
-	UPROPERTY(Category="Tool Settings", EditAnywhere, NonTransactional, meta=(DisplayName="Use Region as Mask", ShowForTools="Mask", ShowForMask))
+	UPROPERTY(Category="Select Mask", EditAnywhere, NonTransactional, meta=(DisplayName="Use Region as Mask", ShowForTools="Mask", ShowForMask))
 	bool bUseSelectedRegion;
 
 	// If enabled, protects the selected region from changes
 	// If disabled, only allows changes in the selected region
-	UPROPERTY(Category="Tool Settings", EditAnywhere, NonTransactional, meta=(DisplayName="Negative Mask", ShowForTools="Mask", ShowForMask))
+	UPROPERTY(Category="Select Mask", EditAnywhere, NonTransactional, meta=(DisplayName="Negative Mask", ShowForTools="Mask", ShowForMask))
 	bool bUseNegativeMask;
 
 	// Copy/Paste Tool:
diff --git a/Engine/Source/Runtime/Landscape/Private/LandscapeEdit.cpp b/Engine/Source/Runtime/Landscape/Private/LandscapeEdit.cpp
--- a/Engine/Source/Runtime/Landscape/Private/LandscapeEdit.cpp
+++ b/Engine/Source/Runtime/Landscape/Private/LandscapeEdit.cpp
@@ -19,6 +19,7 @@ LandscapeEdit.cpp: Landscape editing
 #include "LandscapeStreamingProxy.h"
 #include "LandscapeInfo.h"
 #include "LandscapeComponent.h"
+#include "LandscapeModule.h"
 #include "LandscapeLayerInfoObject.h"
 #include "MaterialDomain.h"
 #include "Materials/Material.h"
@@ -7100,6 +7101,10 @@ void ULandscapeInfo::ClearSelectedRegion(bool bIsComponentwise /*= true*/)
 	{
 		SelectedRegion.Empty();
 	}
+
+	// Refresh the detail panel since the selected region affects the visibility of some landscape tool settings : 
+	ILandscapeModule& LandscapeModule = FModuleManager::GetModuleChecked<ILandscapeModule>("Landscape");
+	LandscapeModule.GetLandscapeEditorServices()->RefreshDetailPanel();
 }
 
 void ULandscapeComponent::ReallocateWeightmaps(FLandscapeEditDataInterface* DataInterface, bool InCanUseEditingWeightmap, bool InSaveToTransactionBuffer, bool InForceReallocate, 
