Fixed assert when resizing landscape : the target layers list would get manually copied but since it's a LandscapeInherited property, it's already been copied before.
All the other such properties that were manually copied in FEdModeLandscape::ChangeComponentSetting have been removed since their copy is now automatically handled by the call to CopySharedProperties
Also, the new landscape actor's label is now set to be the same as the original actor's, so that changing the size of a landscape doesn't end up giving actor numbered names

diff --git a/Engine/Source/Editor/LandscapeEditor/Private/LandscapeEdMode.cpp b/Engine/Source/Editor/LandscapeEditor/Private/LandscapeEdMode.cpp
--- a/Engine/Source/Editor/LandscapeEditor/Private/LandscapeEdMode.cpp
+++ b/Engine/Source/Editor/LandscapeEditor/Private/LandscapeEdMode.cpp
@@ -3738,37 +3738,6 @@ ALandscape* FEdModeLandscape::ChangeComponentSetting(int32 NumComponentsX, int32
 			// LandscapeGuid is stomped by CopySharedProperties, but original guid is not -- fix the mismatch or it will complain during Import
  			NewLandscape->SetLandscapeGuid(FGuid(), /* bValidateGuid= */ false);
 
-			// Copy settings that are not copied by CopySharedProperties
-			NewLandscape->ExportLOD = OldLandscape->ExportLOD;
-			NewLandscape->StaticLightingLOD = OldLandscape->StaticLightingLOD;
-			NewLandscape->StreamingDistanceMultiplier = OldLandscape->StreamingDistanceMultiplier;
-			NewLandscape->StaticLightingResolution = OldLandscape->StaticLightingResolution;
-			NewLandscape->ShadowCacheInvalidationBehavior = OldLandscape->ShadowCacheInvalidationBehavior;
-			NewLandscape->bUseMaterialPositionOffsetInStaticLighting = OldLandscape->bUseMaterialPositionOffsetInStaticLighting;
-			NewLandscape->bUseDynamicMaterialInstance = OldLandscape->bUseDynamicMaterialInstance;
-			NewLandscape->bGenerateOverlapEvents = OldLandscape->bGenerateOverlapEvents;
-			NewLandscape->bBakeMaterialPositionOffsetIntoCollision = OldLandscape->bBakeMaterialPositionOffsetIntoCollision;
-			NewLandscape->bFillCollisionUnderLandscapeForNavmesh = OldLandscape->bFillCollisionUnderLandscapeForNavmesh;
-			NewLandscape->NavigationGeometryGatheringMode = OldLandscape->NavigationGeometryGatheringMode;
-			NewLandscape->bUseLandscapeForCullingInvisibleHLODVertices = OldLandscape->bUseLandscapeForCullingInvisibleHLODVertices;
-			NewLandscape->NonNaniteVirtualShadowMapConstantDepthBias = OldLandscape->NonNaniteVirtualShadowMapConstantDepthBias;
-			NewLandscape->NonNaniteVirtualShadowMapInvalidationHeightErrorThreshold = OldLandscape->NonNaniteVirtualShadowMapInvalidationHeightErrorThreshold;
-			NewLandscape->NonNaniteVirtualShadowMapInvalidationScreenSizeLimit = OldLandscape->NonNaniteVirtualShadowMapInvalidationScreenSizeLimit;
-
-			NewLandscape->BodyInstance.SetCollisionProfileName(OldLandscape->BodyInstance.GetCollisionProfileName());
-			if (NewLandscape->BodyInstance.DoesUseCollisionProfile() == false)
-			{
-				NewLandscape->BodyInstance.SetCollisionEnabled(OldLandscape->BodyInstance.GetCollisionEnabled());
-				NewLandscape->BodyInstance.SetObjectType(OldLandscape->BodyInstance.GetObjectType());
-				NewLandscape->BodyInstance.SetResponseToChannels(OldLandscape->BodyInstance.GetResponseToChannels());
-			}
-			for (const TTuple<FName, FLandscapeTargetLayerSettings>& Layer : OldLandscape->GetTargetLayers())
-			{
-				NewLandscape->AddTargetLayer(Layer.Key, Layer.Value);
-			}
-			NewLandscape->bUsedForNavigation = OldLandscape->bUsedForNavigation;
-			NewLandscape->MaxPaintedLayersPerComponent = OldLandscape->MaxPaintedLayersPerComponent;
-
 			TArrayView<const FLandscapeLayer> LandscapeLayers;
 			if (CanHaveLandscapeLayersContent())
 			{
@@ -3906,6 +3875,7 @@ ALandscape* FEdModeLandscape::ChangeComponentSetting(int32 NumComponentsX, int32
 				}
 			}
 
+			FString OldLandscapeActorLabel = OldLandscape->GetActorLabel();
 			// Delete the old Landscape and all its proxies
 			for (ALandscapeStreamingProxy* Proxy : TActorRange<ALandscapeStreamingProxy>(OldLandscape->GetWorld()))
 			{
@@ -3915,6 +3885,8 @@ ALandscape* FEdModeLandscape::ChangeComponentSetting(int32 NumComponentsX, int32
 				}
 			}
 			OldLandscape->Destroy();
+
+			NewLandscape->SetActorLabel(OldLandscapeActorLabel);
 		}
 	}
 
