Fixed use after free when reassigning landscape layer info

diff --git a/Engine/Source/Editor/LandscapeEditor/Private/LandscapeEditorDetailCustomization_TargetLayers.cpp b/Engine/Source/Editor/LandscapeEditor/Private/LandscapeEditorDetailCustomization_TargetLayers.cpp
--- a/Engine/Source/Editor/LandscapeEditor/Private/LandscapeEditorDetailCustomization_TargetLayers.cpp
+++ b/Engine/Source/Editor/LandscapeEditor/Private/LandscapeEditorDetailCustomization_TargetLayers.cpp
@@ -1370,7 +1370,9 @@ void FLandscapeEditorCustomNodeBuilder_TargetLayers::OnTargetLayerSetObject(cons
 				{
 					FLandscapeInfoLayerSettings& LayerSettings = LandscapeInfo->Layers[Index];
 					LandscapeInfo->ReplaceLayer(LayerSettings.LayerInfoObj, SelectedLayerInfo);
-					LayerSettings.LayerInfoObj = SelectedLayerInfo;
+					// Important : don't use LayerSettings after the call to ReplaceLayer as it will have been reallocated. 
+					//  Validate that the replacement happened as expected : 
+					check(LandscapeInfo->GetLayerInfoIndex(SelectedLayerInfo, Target->Owner.Get()) != INDEX_NONE);
 				}
 			}
 			else
