Landscape mode paint panel :
* Added filter to target layers list (helps when too many layers are being used)
* Minor UX improvements

diff --git a/Engine/Source/Editor/LandscapeEditor/Private/LandscapeEdMode.cpp b/Engine/Source/Editor/LandscapeEditor/Private/LandscapeEdMode.cpp
--- a/Engine/Source/Editor/LandscapeEditor/Private/LandscapeEdMode.cpp
+++ b/Engine/Source/Editor/LandscapeEditor/Private/LandscapeEdMode.cpp
@@ -2852,9 +2852,29 @@ void FEdModeLandscape::UpdateLayerUsageInformation(TWeakObjectPtr<ULandscapeLaye
 
 bool FEdModeLandscape::ShouldShowLayer(TSharedRef<FLandscapeTargetListInfo> Target) const
 {
-	if (!UISettings->ShowUnusedLayers)
+	if (!UISettings->ShowUnusedLayers && (!Target->LayerInfoObj.IsValid() || !Target->LayerInfoObj.Get()->IsReferencedFromLoadedData))
 	{
-		return Target->LayerInfoObj.IsValid() && Target->LayerInfoObj.Get()->IsReferencedFromLoadedData;
+		return false;
+	}
+
+	// check each string in the filter strings list against our layer name
+	if (!UISettings->TargetLayersFilterString.IsEmpty())
+	{
+		// Build a list of strings that must be matched
+		TArray<FString> FilterStrings;
+
+		FString FilterString = UISettings->TargetLayersFilterString;
+		FilterString.TrimStartAndEndInline();
+		FilterString.ParseIntoArray(FilterStrings, TEXT(" "), true /*bCullEmpty*/);
+
+		const FString LayerName = Target->GetLayerName().ToString();
+		for (const FString& String : FilterStrings)
+		{
+			if (!LayerName.Contains(String))
+			{
+				return false;
+			}
+		}
 	}
 
 	return true;
diff --git a/Engine/Source/Editor/LandscapeEditor/Private/LandscapeEditorDetailCustomization_TargetLayers.cpp b/Engine/Source/Editor/LandscapeEditor/Private/LandscapeEditorDetailCustomization_TargetLayers.cpp
--- a/Engine/Source/Editor/LandscapeEditor/Private/LandscapeEditorDetailCustomization_TargetLayers.cpp
+++ b/Engine/Source/Editor/LandscapeEditor/Private/LandscapeEditorDetailCustomization_TargetLayers.cpp
@@ -17,6 +17,7 @@
 #include "Widgets/Input/SButton.h"
 #include "Widgets/Input/SComboButton.h"
 #include "Widgets/Input/SCheckBox.h"
+#include "Widgets/Input/SSearchBox.h"
 #include "Widgets/Input/SMultiLineEditableTextBox.h"
 #include "Widgets/SBoxPanel.h"
 #include "Widgets/Text/SInlineEditableTextBlock.h"
@@ -243,10 +244,9 @@ void FLandscapeEditorCustomNodeBuilder_TargetLayers::GenerateHeaderRowContent(FD
 			.Padding(0.0f, 0.0f, 0.0f, 0.0f)
 			[
 				SNew(SComboButton)
-				.ComboButtonStyle(FAppStyle::Get(), "ToolbarComboButton")
+				.ComboButtonStyle(FAppStyle::Get(), "SimpleComboButtonWithIcon")
 				.ForegroundColor(FSlateColor::UseForeground())
 				.HasDownArrow(true)
-				.ContentPadding(FMargin(1, 0))
 				.VAlign(VAlign_Center)
 				.HAlign(HAlign_Center)
 				.ToolTipText(LOCTEXT("TargetLayerSortButtonTooltip", "Define how we want to sort the displayed layers"))
@@ -279,10 +279,9 @@ void FLandscapeEditorCustomNodeBuilder_TargetLayers::GenerateHeaderRowContent(FD
 			.Padding(5.0f, 0.0f, 0.0f, 0.0f)
 			[
 				SNew(SComboButton)
-				.ComboButtonStyle(FAppStyle::Get(), "ToolbarComboButton")
+				.ComboButtonStyle(FAppStyle::Get(), "SimpleComboButtonWithIcon")
 				.ForegroundColor(FSlateColor::UseForeground())
 				.HasDownArrow(true)
-				.ContentPadding(FMargin(1, 0))
 				.VAlign(VAlign_Center)
 				.HAlign(HAlign_Center)
 				.ToolTipText(LOCTEXT("TargetLayerUnusedLayerButtonTooltip", "Define if we want to display unused layers"))
@@ -548,6 +547,50 @@ EVisibility FLandscapeEditorCustomNodeBuilder_TargetLayers::ShouldShowLayer(TSha
 	return EVisibility::Visible;
 }
 
+void FLandscapeEditorCustomNodeBuilder_TargetLayers::OnFilterTextChanged(const FText& InFilterText)
+{
+	if (FEdModeLandscape* LandscapeEdMode = GetEditorMode())
+	{
+		LandscapeEdMode->UISettings->TargetLayersFilterString = InFilterText.ToString();
+	}
+}
+
+void FLandscapeEditorCustomNodeBuilder_TargetLayers::OnFilterTextCommitted(const FText& InFilterText, ETextCommit::Type InCommitType)
+{
+	if (InCommitType == ETextCommit::OnCleared)
+	{
+		LayersFilterSearchBox->SetText(FText::GetEmpty());
+		OnFilterTextChanged(FText::GetEmpty());
+		FSlateApplication::Get().ClearKeyboardFocus(EFocusCause::Cleared);
+	}
+}
+
+EVisibility FLandscapeEditorCustomNodeBuilder_TargetLayers::GetLayersFilterVisibility() const
+{
+	FEdModeLandscape* LandscapeEdMode = GetEditorMode();
+
+	if (LandscapeEdMode && LandscapeEdMode->CurrentToolMode)
+	{
+		if (LandscapeEdMode->CurrentToolTarget.TargetType == ELandscapeToolTargetType::Weightmap)
+		{
+			return EVisibility::Visible;
+		}
+	}
+
+	return EVisibility::Collapsed;
+}
+
+FText FLandscapeEditorCustomNodeBuilder_TargetLayers::GetLayersFilterText() const
+{
+	if (FEdModeLandscape* LandscapeEdMode = GetEditorMode())
+	{
+		return FText::FromString(LandscapeEdMode->UISettings->TargetLayersFilterString);
+	}
+
+	return FText();
+}
+
+
 BEGIN_SLATE_FUNCTION_BUILD_OPTIMIZATION
 void FLandscapeEditorCustomNodeBuilder_TargetLayers::GenerateChildContent(IDetailChildrenBuilder& ChildrenBuilder)
 {
@@ -575,7 +618,30 @@ void FLandscapeEditorCustomNodeBuilder_TargetLayers::GenerateChildContent(IDetai
 		ChildrenBuilder.AddCustomRow(LOCTEXT("LayersLabel", "Layers"))
 			.Visibility(EVisibility::Visible)
 			[
-				TargetLayerList.ToSharedRef()
+				SNew(SVerticalBox)
+					+ SVerticalBox::Slot()
+					.AutoHeight()
+					.VAlign(VAlign_Top)
+					.HAlign(HAlign_Fill)
+					.Padding(FMargin(2))
+					[
+						SAssignNew(LayersFilterSearchBox, SSearchBox)
+							.InitialText(this, &FLandscapeEditorCustomNodeBuilder_TargetLayers::GetLayersFilterText)
+							.SelectAllTextWhenFocused(true)
+							.HintText(LOCTEXT("LayersSearch", "Filter Target Layers"))
+							.OnTextChanged(this, &FLandscapeEditorCustomNodeBuilder_TargetLayers::OnFilterTextChanged)
+							.OnTextCommitted(this, &FLandscapeEditorCustomNodeBuilder_TargetLayers::OnFilterTextCommitted)
+							.Visibility(this, &FLandscapeEditorCustomNodeBuilder_TargetLayers::GetLayersFilterVisibility)
+					]
+
+					+ SVerticalBox::Slot()
+					.AutoHeight()
+					.VAlign(VAlign_Top)
+					.HAlign(HAlign_Fill)
+					.Padding(0.0f, 0.0f)
+					[
+						TargetLayerList.ToSharedRef()
+					]
 			];
 
 		for (int32 i = 0; i < TargetDisplayOrderList->Num(); ++i)
diff --git a/Engine/Source/Editor/LandscapeEditor/Private/LandscapeEditorDetailCustomization_TargetLayers.h b/Engine/Source/Editor/LandscapeEditor/Private/LandscapeEditorDetailCustomization_TargetLayers.h
--- a/Engine/Source/Editor/LandscapeEditor/Private/LandscapeEditorDetailCustomization_TargetLayers.h
+++ b/Engine/Source/Editor/LandscapeEditor/Private/LandscapeEditorDetailCustomization_TargetLayers.h
@@ -61,6 +61,7 @@ protected:
 	TSharedRef<FAssetThumbnailPool> ThumbnailPool;
 	TSharedRef<IPropertyHandle> TargetDisplayOrderPropertyHandle;
 	TSharedRef<IPropertyHandle> TargetShowUnusedLayersPropertyHandle;
+	TSharedPtr<SSearchBox> LayersFilterSearchBox;
 
 	static class FEdModeLandscape* GetEditorMode();
 
@@ -115,6 +116,11 @@ protected:
 	void ShowUnusedLayers(bool Result);
 	bool ShouldShowUnusedLayers(bool Result) const;
 	EVisibility ShouldShowLayer(TSharedRef<FLandscapeTargetListInfo> Target) const;
+
+	void OnFilterTextChanged(const FText& InFilterText);
+	void OnFilterTextCommitted(const FText& InFilterText, ETextCommit::Type InCommitType);
+	EVisibility GetLayersFilterVisibility() const;
+	FText GetLayersFilterText() const;
 };
 
 class SLandscapeEditorSelectableBorder : public SBorder
diff --git a/Engine/Source/Editor/LandscapeEditor/Public/LandscapeEditorObject.h b/Engine/Source/Editor/LandscapeEditor/Public/LandscapeEditorObject.h
--- a/Engine/Source/Editor/LandscapeEditor/Public/LandscapeEditorObject.h
+++ b/Engine/Source/Editor/LandscapeEditor/Public/LandscapeEditorObject.h
@@ -699,13 +699,18 @@ public:
 	// Target Layer Settings:
 
 	// Limits painting to only the components that already have the selected layer
-	UPROPERTY(Category="Target Layers", EditAnywhere, NonTransactional, meta=(ShowForTargetTypes="Weightmap,Visibility"))
+	UPROPERTY(Category = "Target Layers", EditAnywhere, NonTransactional, meta=(ShowForTargetTypes="Weightmap,Visibility"))
 	ELandscapeLayerPaintingRestriction PaintingRestriction;
 
+	// This allows to hide layer names that do not match the filter
+	UPROPERTY(Transient)
+	FString TargetLayersFilterString;
+
 	// Display order of the targets
 	UPROPERTY(Category = "Target Layers", EditAnywhere)
 	ELandscapeLayerDisplayMode TargetDisplayOrder;
 
+	// Show layers that are not currently being painted anywhere
 	UPROPERTY(Category = "Target Layers", EditAnywhere)
 	bool ShowUnusedLayers;	
 
