Disable Retopologize tool button when unsupported or when CVar landscape.EnableRetopologizeTool is false (default is false).

diff --git a/Engine/Source/Editor/LandscapeEditor/Private/LandscapeEdMode.cpp b/Engine/Source/Editor/LandscapeEditor/Private/LandscapeEdMode.cpp
--- a/Engine/Source/Editor/LandscapeEditor/Private/LandscapeEdMode.cpp
+++ b/Engine/Source/Editor/LandscapeEditor/Private/LandscapeEdMode.cpp
@@ -4299,7 +4299,7 @@ bool FEdModeLandscape::CanEditLayer(FText* Reason /*=nullptr*/, const FLandscape
 			{
 				if (Reason)
 				{
-					*Reason = FText::Format(NSLOCTEXT("UnrealEd", "LandscapeLayersNoSupportForRetopologize", "{0} Tool is not available with the Landscape Edit Layer System."), CurrentTool->GetDisplayName());
+					*Reason = FText::Format(NSLOCTEXT("UnrealEd", "LandscapeLayersNoSupportForRetopologize", "{0} Tool is not available with the Landscape Edit Layer System and will be entirely deprecated for all landscape types in UE5.6."), CurrentTool->GetDisplayName());
 				}
 				return false;
 			}
diff --git a/Engine/Source/Editor/LandscapeEditor/Private/LandscapeEditorCommands.cpp b/Engine/Source/Editor/LandscapeEditor/Private/LandscapeEditorCommands.cpp
--- a/Engine/Source/Editor/LandscapeEditor/Private/LandscapeEditorCommands.cpp
+++ b/Engine/Source/Editor/LandscapeEditor/Private/LandscapeEditorCommands.cpp
@@ -73,7 +73,9 @@ void FLandscapeEditorCommands::RegisterCommands()
 	UI_COMMAND(NoiseTool, "Noise", "Adds noise to the heightmap or blend layer", EUserInterfaceActionType::RadioButton, FInputChord());
 	NameToCommandMap.Add("Tool_Noise", NoiseTool);
 
-	UI_COMMAND(RetopologizeTool, "Retop", "Automatically adjusts landscape vertices with an X/Y offset map to improve vertex density on cliffs, reducing texture stretching.\nNote: An X/Y offset map makes the landscape slower to render and paint on with other tools, so only use if needed", EUserInterfaceActionType::RadioButton, FInputChord());
+	UI_COMMAND(RetopologizeTool, "Retop", "Automatically adjusts landscape vertices with an X/Y offset map to improve vertex density on cliffs, reducing texture stretching.\n"
+								 "Note: An X/Y offset map makes the landscape slower to render and paint on with other tools, so only use if needed.\n\n"
+							     "Retop is not available with the Landscape Edit Layer System and will be deprecated for all landscape types in UE5.6.  If not using Edit Layers, you can set the CVar landscape.EnableRetopologizeTool to True to enable it.", EUserInterfaceActionType::RadioButton, FInputChord());
 	NameToCommandMap.Add("Tool_Retopologize", RetopologizeTool);
 
 	UI_COMMAND(VisibilityTool, "Visibility", "Mask out individual quads in the landscape, leaving a hole.", EUserInterfaceActionType::RadioButton, FInputChord());
diff --git a/Engine/Source/Editor/LandscapeEditor/Private/SLandscapeEditor.cpp b/Engine/Source/Editor/LandscapeEditor/Private/SLandscapeEditor.cpp
--- a/Engine/Source/Editor/LandscapeEditor/Private/SLandscapeEditor.cpp
+++ b/Engine/Source/Editor/LandscapeEditor/Private/SLandscapeEditor.cpp
@@ -19,9 +19,20 @@
 #include "IDetailsView.h"
 #include "PropertyEditorModule.h"
 #include "LandscapeSettings.h"
+#include "HAL/IConsoleManager.h"
 
 #define LOCTEXT_NAMESPACE "LandscapeEditor"
 
+namespace UE::Landscape::Private {
+	static bool bEnableRetopoTool = false;
+	static FAutoConsoleVariableRef CVarEnableRetopoTool(
+		TEXT("landscape.EnableRetopologizeTool"),
+		bEnableRetopoTool,
+		TEXT("Enable the Retopologize tool.  The tool will be fully deprecated in UE5.6, but this cvar will enable it for 5.5"),
+		ECVF_Default
+	);
+}
+
 void SLandscapeAssetThumbnail::Construct(const FArguments& InArgs, UObject* Asset, TSharedRef<FAssetThumbnailPool> ThumbnailPool)
 {
 	FIntPoint ThumbnailSize = InArgs._ThumbnailSize;
@@ -378,6 +389,11 @@ bool FLandscapeToolKit::IsToolEnabled(FName ToolName) const
 			return true;
 		}
 
+		if (ToolName == "Retopologize")
+		{
+			return UE::Landscape::Private::bEnableRetopoTool && !LandscapeEdMode->CanHaveLandscapeLayersContent();
+		}
+
 		// Other tools are available if there is an existing landscape.
 		if (LandscapeEdMode->GetLandscapeList().Num() > 0)
 		{
