Fix for landscape import sometimes choosing non-valid QuadsPerSection

#12338

diff --git a/Engine/Source/Editor/LandscapeEditor/Private/LandscapeImportHelper.cpp b/Engine/Source/Editor/LandscapeEditor/Private/LandscapeImportHelper.cpp
--- a/Engine/Source/Editor/LandscapeEditor/Private/LandscapeImportHelper.cpp
+++ b/Engine/Source/Editor/LandscapeEditor/Private/LandscapeImportHelper.cpp
@@ -417,7 +417,7 @@ void FLandscapeImportHelper::ChooseBestComponentSizeForImport(int32 Width, int32
 
 		if (!bValidQuadsPerSectionParam)
 		{
-			InOutSectionsPerComponent = FLandscapeConfig::SubsectionSizeQuadsValues[0];
+			InOutQuadsPerSection = FLandscapeConfig::SubsectionSizeQuadsValues[0];
 		}
 
 		// if there was no exact match, try increasing the section size until we encompass the whole heightmap
