Fix wrong warning when importing landscape weightmap from 8bit source image.
Prevent cached result from attempted 16-bit import from preventing the 8bit attempt from running at all.  Make more of the import code aware of the intended usage of an import request, 8bit for weightmaps and 16bit for heightmaps.  Store those separately in the import cache to stop errors when the same file is imported as a weightmap and heightmap at the same time.

diff --git a/Engine/Source/Editor/LandscapeEditor/Private/LandscapeEdMode.cpp b/Engine/Source/Editor/LandscapeEditor/Private/LandscapeEdMode.cpp
--- a/Engine/Source/Editor/LandscapeEditor/Private/LandscapeEdMode.cpp
+++ b/Engine/Source/Editor/LandscapeEditor/Private/LandscapeEdMode.cpp
@@ -3233,7 +3233,7 @@ void ImportDataInternal(ULandscapeInfo* LandscapeInfo, const FString& Filename,
 	}
 
 	FLandscapeTiledImage TiledImage;
-	TiledImage.Load(*Filename);
+	TiledImage.Load<T>(*Filename);
 	FIntPoint ImportResolution = TiledImage.GetResolution();
 
 	bool bResolutionMismatch = false;
diff --git a/Engine/Source/Editor/LandscapeEditor/Private/LandscapeEditorDetailCustomization_ImportExport.cpp b/Engine/Source/Editor/LandscapeEditor/Private/LandscapeEditorDetailCustomization_ImportExport.cpp
--- a/Engine/Source/Editor/LandscapeEditor/Private/LandscapeEditorDetailCustomization_ImportExport.cpp
+++ b/Engine/Source/Editor/LandscapeEditor/Private/LandscapeEditorDetailCustomization_ImportExport.cpp
@@ -37,7 +37,6 @@
 #include "LandscapeSubsystem.h"
 #include "LandscapeStreamingProxy.h"
 #include "LandscapeImageFileCache.h"
-#include "LandscapeTiledImage.h"
 #include "LandscapeRegionUtils.h"
 #include "LandscapeEditorUtils.h"
 
diff --git a/Engine/Source/Editor/LandscapeEditor/Private/LandscapeEditorDetailCustomization_ImportLayers.cpp b/Engine/Source/Editor/LandscapeEditor/Private/LandscapeEditorDetailCustomization_ImportLayers.cpp
--- a/Engine/Source/Editor/LandscapeEditor/Private/LandscapeEditorDetailCustomization_ImportLayers.cpp
+++ b/Engine/Source/Editor/LandscapeEditor/Private/LandscapeEditorDetailCustomization_ImportLayers.cpp
@@ -8,7 +8,6 @@
 #include "LandscapeEditorModule.h"
 #include "LandscapeEditorObject.h"
 #include "LandscapeUtils.h"
-#include "LandscapeTiledImage.h"
 #include "LandscapeEditorUtils.h"
 #include "SLandscapeEditor.h"
 
diff --git a/Engine/Source/Editor/LandscapeEditor/Private/LandscapeEditorDetailCustomization_NewLandscape.cpp b/Engine/Source/Editor/LandscapeEditor/Private/LandscapeEditorDetailCustomization_NewLandscape.cpp
--- a/Engine/Source/Editor/LandscapeEditor/Private/LandscapeEditorDetailCustomization_NewLandscape.cpp
+++ b/Engine/Source/Editor/LandscapeEditor/Private/LandscapeEditorDetailCustomization_NewLandscape.cpp
@@ -22,7 +22,6 @@
 #include "LandscapeEditorObject.h"
 #include "Landscape.h"
 #include "LandscapeStreamingProxy.h"
-#include "LandscapeTiledImage.h"
 #include "LandscapeRegionUtils.h"
 #include "LandscapeEditorPrivate.h"
 #include "LandscapeEditorUtils.h"
diff --git a/Engine/Source/Editor/LandscapeEditor/Private/LandscapeEditorObject.cpp b/Engine/Source/Editor/LandscapeEditor/Private/LandscapeEditorObject.cpp
--- a/Engine/Source/Editor/LandscapeEditor/Private/LandscapeEditorObject.cpp
+++ b/Engine/Source/Editor/LandscapeEditor/Private/LandscapeEditorObject.cpp
@@ -669,7 +669,7 @@ void ULandscapeEditorObject::RefreshImports()
 	{
 		FLandscapeTiledImage TiledImage;
 		
-		FLandscapeFileInfo FileInfo = TiledImage.Load(*ImportLandscape_HeightmapFilename);
+		FLandscapeFileInfo FileInfo = TiledImage.Load<uint16>(*ImportLandscape_HeightmapFilename);
 
 		if (FileInfo.PossibleResolutions.Num() > 0)
 		{
@@ -726,7 +726,7 @@ void ULandscapeEditorObject::RefreshLayerImport(FLandscapeImportLayer& ImportLay
 		else
 		{
 			FLandscapeTiledImage TiledImage;
-			FLandscapeFileInfo FileInfo = TiledImage.Load(*ImportLayer.SourceFilePath);
+			FLandscapeFileInfo FileInfo = TiledImage.Load<uint8>(*ImportLayer.SourceFilePath);
 			ImportLayer.ImportResult = FileInfo.ResultCode;
 			ImportLayer.ErrorMessage = FileInfo.ErrorMessage;
 			if (FileInfo.ResultCode == ELandscapeImportResult::Success)
@@ -755,7 +755,7 @@ void ULandscapeEditorObject::OnChangeImportLandscapeResolution(int32 DescriptorI
 void ULandscapeEditorObject::ImportLandscapeData()
 {
 	FLandscapeTiledImage TiledImage;
-	FLandscapeFileInfo FileInfo = TiledImage.Load(*ImportLandscape_HeightmapFilename);
+	FLandscapeFileInfo FileInfo = TiledImage.Load<uint16>(*ImportLandscape_HeightmapFilename);
 
 	if (FileInfo.ResultCode == ELandscapeImportResult::Error)
 	{
@@ -787,7 +787,7 @@ ELandscapeImportResult ULandscapeEditorObject::CreateImportLayersInfo(TArray<FLa
 		if (ImportLayer.LayerInfo != nullptr && !ImportLayer.SourceFilePath.IsEmpty())
 		{
 			FLandscapeTiledImage LayerImage;
-			FLandscapeFileInfo LayerFileInfo = LayerImage.Load(*ImportLayer.SourceFilePath);
+			FLandscapeFileInfo LayerFileInfo = LayerImage.Load<uint8>(*ImportLayer.SourceFilePath);
 
 			UIImportLayer.ImportResult = LayerFileInfo.ResultCode;
 
diff --git a/Engine/Source/Editor/LandscapeEditor/Private/LandscapeImageFileCache.cpp b/Engine/Source/Editor/LandscapeEditor/Private/LandscapeImageFileCache.cpp
--- a/Engine/Source/Editor/LandscapeEditor/Private/LandscapeImageFileCache.cpp
+++ b/Engine/Source/Editor/LandscapeEditor/Private/LandscapeImageFileCache.cpp
@@ -77,22 +77,27 @@ void FLandscapeImageFileCache::UnmonitorFile(const FString& Filename)
 	}
 }
 
+template <typename T>
 void FLandscapeImageFileCache::Add(const FString& Filename, FLandscapeImageDataRef NewImageData)
 {
-	check(CachedImages.Find(Filename) == nullptr);
+	CacheType& Cache = ChooseCache<T>();
+	check(Cache.Find(Filename) == nullptr);
 
-	CachedImages.Add(FString(Filename), FCacheEntry(NewImageData));
+	Cache.Add(FString(Filename), FCacheEntry(NewImageData));
 	MonitorFile(Filename);
 	CacheSize += NewImageData.Data->Num();
 }
 
 void FLandscapeImageFileCache::Remove(const FString& Filename)
 {
-	if (FCacheEntry* CacheEntry = CachedImages.Find(Filename))
+	for (CacheType& Cache : CachedImages)
 	{
-		CacheSize -= CacheEntry->ImageData.Data->Num();
-		CachedImages.Remove(Filename);
-		UnmonitorFile(Filename);
+		if (FCacheEntry* CacheEntry = Cache.Find(Filename))
+		{
+			CacheSize -= CacheEntry->ImageData.Data->Num();
+			Cache.Remove(Filename);
+			UnmonitorFile(Filename);
+		}
 	}
 }
 
@@ -116,7 +121,10 @@ void FLandscapeImageFileCache::SetMaxSize(uint64 InNewMaxSize)
 
 void FLandscapeImageFileCache::Clear()
 {
-	CachedImages.Empty();
+	for (CacheType& Cache : CachedImages)
+	{
+		Cache.Empty();
+	}
 	FDirectoryWatcherModule& DirectoryWatcherModule = FModuleManager::LoadModuleChecked<FDirectoryWatcherModule>("DirectoryWatcher");
 	
 	for (const TPair<FString, FDirectoryMonitor>& MonitoredDir : MonitoredDirs)
@@ -134,13 +142,26 @@ void FLandscapeImageFileCache::Trim()
 		return;
 	}
 
-	CachedImages.ValueSort([](const FCacheEntry& LHS, const FCacheEntry& RHS) { return  LHS.UsageCount < RHS.UsageCount; });
+	using RemovalPriorityData = TTuple<uint32, FString, uint32>;
+	TArray<RemovalPriorityData> AllCacheEntries;
+
+	// Make an array of all cache entries from both maps, sorted by UsageCount.
+	for (CacheType& Cache : CachedImages)
+	{
+		for (const TPair<FString, FCacheEntry>& It : Cache)
+		{
+			AllCacheEntries.Push(RemovalPriorityData(It.Value.UsageCount, It.Key, It.Value.ImageData.Data->Num()));
+		}
+	}
+	AllCacheEntries.Sort([](const RemovalPriorityData& LHS, const RemovalPriorityData& RHS) { return  LHS.Get<0>() < RHS.Get<0>(); });
+
+	// Starting with the leased used, remove entries until we are under MaxCacheSize.
 	TArray<FString> ToRemove;
 	uint64 Size = CacheSize;
-	for (const TPair<FString, FCacheEntry>& It : CachedImages)
+	for (const RemovalPriorityData& Entry : AllCacheEntries)
 	{
-		ToRemove.Add(It.Key);
-		uint64 ImageSize = It.Value.ImageData.Data->Num();
+		ToRemove.Add(Entry.Get<1>());
+		uint64 ImageSize = Entry.Get<2>();
 		Size -= ImageSize;
 		if (Size <= MaxCacheSize)
 		{
@@ -148,8 +169,23 @@ void FLandscapeImageFileCache::Trim()
 		}
 	}
 
+	// Note that the file is removed from both caches, regardless of which entry UsageCount got it on this list.  Other
+	// callers of Remove must target all entries of the file name, but this one could go either way.  This rare case might
+	// overshoot the MaxCacheSize target.
 	for (const FString& Filename : ToRemove)
 	{
 		Remove(Filename);
 	}
-}
\ No newline at end of file
+}
+
+template <>
+FLandscapeImageFileCache::CacheType& FLandscapeImageFileCache::ChooseCache<uint8>()
+{
+	return CachedImages[Cache8];
+}
+
+template <>
+FLandscapeImageFileCache::CacheType& FLandscapeImageFileCache::ChooseCache<uint16>()
+{
+	return CachedImages[Cache16];
+}
diff --git a/Engine/Source/Editor/LandscapeEditor/Private/LandscapeImageFileCache.h b/Engine/Source/Editor/LandscapeEditor/Private/LandscapeImageFileCache.h
--- a/Engine/Source/Editor/LandscapeEditor/Private/LandscapeImageFileCache.h
+++ b/Engine/Source/Editor/LandscapeEditor/Private/LandscapeImageFileCache.h
@@ -13,7 +13,7 @@
 #include "DirectoryWatcherModule.h"
 #include "LandscapeFileFormatInterface.h"
 #include "LandscapeEditorModule.h"
-
+#include "Containers/StaticArray.h"
 #include "Delegates/IDelegateInstance.h"
 
 #define LOCTEXT_NAMESPACE "LandscapeEditor"
@@ -33,11 +33,10 @@ public:
 	FLandscapeImageFileCache();
 	~FLandscapeImageFileCache();
 
-
 	template<typename T>
 	FLandscapeFileInfo FindImage(const TCHAR* InImageFilename, FLandscapeImageDataRef& OutImageData)
 	{
-		FCacheEntry* CacheEntry = CachedImages.Find(InImageFilename);
+		FCacheEntry* CacheEntry = ChooseCache<T>().Find(InImageFilename);
 		FLandscapeFileInfo Result;
 
 		if (CacheEntry)
@@ -83,7 +82,7 @@ public:
 		
 		Trim();
 		OutImageData = NewImageData;
-		Add(FString(InImageFilename), OutImageData);
+		Add<T>(FString(InImageFilename), OutImageData);
 		
 		Result.PossibleResolutions.Add(FLandscapeFileResolution(OutImageData.Resolution.X, OutImageData.Resolution.Y));
 		Result.ResultCode = FileInfo.ResultCode;
@@ -120,10 +119,22 @@ private:
 
 	void MonitorCallback(const TArray<struct FFileChangeData>& Changes);
 
-	TMap<FString, FCacheEntry> CachedImages;
+	using CacheType = TMap<FString, FCacheEntry>;
+
+	template<typename T>
+	CacheType& ChooseCache();
+
+	// Two cache maps.  Hold 8 bit and 16 bit data separately.  Otherwise importing the same file as weightmap and heightmap
+	// will use the cached 8bit version for the heightmap.
+	enum {
+		Cache8 = 0,
+		Cache16 = 1
+	};
+	TStaticArray<CacheType, 2> CachedImages;
 	
 	TMap<FString, FDirectoryMonitor> MonitoredDirs;
 
+	template <typename T>
 	void Add(const FString& Filename, FLandscapeImageDataRef NewImageData);
 	void Remove(const FString& Filename);
 
diff --git a/Engine/Source/Editor/LandscapeEditor/Private/LandscapeTiledImage.cpp b/Engine/Source/Editor/LandscapeEditor/Private/LandscapeTiledImage.cpp
--- a/Engine/Source/Editor/LandscapeEditor/Private/LandscapeTiledImage.cpp
+++ b/Engine/Source/Editor/LandscapeEditor/Private/LandscapeTiledImage.cpp
@@ -28,6 +28,7 @@ FLandscapeTiledImage::FLandscapeTiledImage()
 {
 }
 
+template <typename T>
 FLandscapeFileInfo FLandscapeTiledImage::Load(const TCHAR* Filename)
 {
 	FLandscapeFileInfo Result;
@@ -115,21 +116,10 @@ FLandscapeFileInfo FLandscapeTiledImage::Load(const TCHAR* Filename)
 	{
 		FLandscapeImageDataRef ImageData;
 		FLandscapeImageFileCache& LandscapeImageFileCache = FModuleManager::GetModuleChecked<ILandscapeEditorModule>("LandscapeEditor").GetImageFileCache();
-		FLandscapeFileInfo TileLoadResult;
-		if (TileLoadResult = LandscapeImageFileCache.FindImage<uint16>(*Tile.Value, ImageData); TileLoadResult.ResultCode != ELandscapeImportResult::Success)
+		Result = LandscapeImageFileCache.FindImage<T>(*Tile.Value, ImageData);
+		if (Result.ResultCode == ELandscapeImportResult::Error)
 		{
-			if (TileLoadResult = LandscapeImageFileCache.FindImage<uint8>(*Tile.Value, ImageData);  TileLoadResult.ResultCode != ELandscapeImportResult::Success)
-			{
-				if (TileLoadResult.ResultCode == ELandscapeImportResult::Warning)
-				{
-					Result.ErrorMessage = TileLoadResult.ErrorMessage;
-					Result.ResultCode = TileLoadResult.ResultCode;
-				}
-				else if (TileLoadResult.ResultCode == ELandscapeImportResult::Error)
-				{
-					return TileLoadResult;
-				}
-			}
+			return Result;
 		}
 
 		SizeInTiles.X = FMath::Max(SizeInTiles.X, Tile.Key.X + 1);
@@ -146,7 +136,7 @@ FLandscapeFileInfo FLandscapeTiledImage::Load(const TCHAR* Filename)
 		else if (Width != TileResolution.X && Height != TileResolution.Y)
 		{
 			Result.ResultCode = ELandscapeImportResult::Error;
-			Result.ErrorMessage = LOCTEXT("FileReadErrorTiledResolutionMismatch", "Mimatched resolution found in tiled image");
+			Result.ErrorMessage = LOCTEXT("FileReadErrorTiledResolutionMismatch", "Mismatched resolution found in tiled image");
 			return Result;
 		}
 	}
@@ -163,6 +153,10 @@ FLandscapeFileInfo FLandscapeTiledImage::Load(const TCHAR* Filename)
 	return Result;
 }
 
+template FLandscapeFileInfo FLandscapeTiledImage::Load<uint8>(const TCHAR* Filename);
+template FLandscapeFileInfo FLandscapeTiledImage::Load<uint16>(const TCHAR* Filename);
+
+
 void FLandscapeTiledImage::FindFiles(const TCHAR* FilenamePattern, TArray<FString>& OutPaths)
 {
 	FString GlobFilename(FilenamePattern);
diff --git a/Engine/Source/Editor/LandscapeEditor/Private/LandscapeTiledImage.h b/Engine/Source/Editor/LandscapeEditor/Private/LandscapeTiledImage.h
--- a/Engine/Source/Editor/LandscapeEditor/Private/LandscapeTiledImage.h
+++ b/Engine/Source/Editor/LandscapeEditor/Private/LandscapeTiledImage.h
@@ -26,6 +26,8 @@ class FLandscapeTiledImage
 {
 public:
 	FLandscapeTiledImage();
+
+	template<typename T>
 	FLandscapeFileInfo Load(const TCHAR* Filename);
 	
 	FIntPoint GetTileResolution() const { return TileResolution;  }
