Fixed crash when unloading a map while landscape mode is active

diff --git a/Engine/Source/Editor/LandscapeEditor/Private/LandscapeEdMode.cpp b/Engine/Source/Editor/LandscapeEditor/Private/LandscapeEdMode.cpp
--- a/Engine/Source/Editor/LandscapeEditor/Private/LandscapeEdMode.cpp
+++ b/Engine/Source/Editor/LandscapeEditor/Private/LandscapeEdMode.cpp
@@ -4069,6 +4069,7 @@ const ULandscapeEditLayerBase* FEdModeLandscape::GetEditLayer(int32 InLayerIndex
 	{
 		if (const FLandscapeLayer* EditLayer = Landscape->GetLayerConst(InLayerIndex))
 		{
+			check(EditLayer->EditLayer != nullptr);
 			return EditLayer->EditLayer;
 		}
 	}
diff --git a/Engine/Source/Editor/LandscapeEditor/Private/LandscapeEditorDetailCustomization_Layers.cpp b/Engine/Source/Editor/LandscapeEditor/Private/LandscapeEditorDetailCustomization_Layers.cpp
--- a/Engine/Source/Editor/LandscapeEditor/Private/LandscapeEditorDetailCustomization_Layers.cpp
+++ b/Engine/Source/Editor/LandscapeEditor/Private/LandscapeEditorDetailCustomization_Layers.cpp
@@ -1250,8 +1250,7 @@ TSharedPtr<IToolTip> FLandscapeEditorCustomNodeBuilder_Layers::GetEditLayerTypeT
 {
 	FEdModeLandscape* LandscapeEdMode = GetEditorMode();
 	const ULandscapeEditLayerBase* EditLayer = (LandscapeEdMode != nullptr) ? LandscapeEdMode->GetEditLayer(InLayerIndex) : nullptr;
-	check(EditLayer != nullptr);
-	return FEditorClassUtils::GetTooltip(EditLayer->GetClass());
+	return (EditLayer != nullptr) ? TSharedPtr<IToolTip>(FEditorClassUtils::GetTooltip(EditLayer->GetClass())) : nullptr;
 }
 
 const FSlateBrush* FLandscapeEditorCustomNodeBuilder_Layers::GetLockBrushForLayer(int32 InLayerIndex) const
@@ -1263,13 +1262,9 @@ const FSlateBrush* FLandscapeEditorCustomNodeBuilder_Layers::GetLockBrushForLaye
 
 const FSlateBrush* FLandscapeEditorCustomNodeBuilder_Layers::GetEditLayerIconBrush(int32 InLayerIndex) const
 {
-	if (FEdModeLandscape* LandscapeEdMode = GetEditorMode())
-	{
-		const ULandscapeEditLayerBase* EditLayer = LandscapeEdMode->GetEditLayer(InLayerIndex);
-		check(EditLayer != nullptr);
-		return FSlateIconFinder::FindIconBrushForClass(EditLayer->GetClass());
-	}
-	return nullptr;
+	FEdModeLandscape* LandscapeEdMode = GetEditorMode();
+	const ULandscapeEditLayerBase* EditLayer = (LandscapeEdMode != nullptr) ? LandscapeEdMode->GetEditLayer(InLayerIndex) : nullptr;
+	return (EditLayer != nullptr) ? FSlateIconFinder::FindIconBrushForClass(EditLayer->GetClass()) : nullptr;
 }
 
 int32 FLandscapeEditorCustomNodeBuilder_Layers::SlotIndexToLayerIndex(int32 SlotIndex)
