Fix non-unity build by uninlining templated member funtion.

A public member function, inlined, can't reference private templated member functions which are not inlined or explicitly instantiated.

#rnx

diff --git a/Engine/Source/Editor/LandscapeEditor/Private/LandscapeImageFileCache.cpp b/Engine/Source/Editor/LandscapeEditor/Private/LandscapeImageFileCache.cpp
--- a/Engine/Source/Editor/LandscapeEditor/Private/LandscapeImageFileCache.cpp
+++ b/Engine/Source/Editor/LandscapeEditor/Private/LandscapeImageFileCache.cpp
@@ -23,6 +23,70 @@ FLandscapeImageFileCache::~FLandscapeImageFileCache()
 	}
 }
 
+template<typename T>
+FLandscapeFileInfo FLandscapeImageFileCache::FindImage(const TCHAR* InImageFilename, FLandscapeImageDataRef& OutImageData)
+{
+	FCacheEntry* CacheEntry = ChooseCache<T>().Find(InImageFilename);
+	FLandscapeFileInfo Result;
+
+	if (CacheEntry)
+	{
+		CacheEntry->UsageCount++;
+		OutImageData = CacheEntry->ImageData;
+		Result.PossibleResolutions.Add(FLandscapeFileResolution(OutImageData.Resolution.X, OutImageData.Resolution.Y));
+		Result.ResultCode = CacheEntry->ImageData.Result;
+		Result.ErrorMessage = CacheEntry->ImageData.ErrorMessage;
+		return Result;
+	}
+
+	FLandscapeImageDataRef NewImageData;
+	ILandscapeEditorModule& LandscapeEditorModule = FModuleManager::GetModuleChecked<ILandscapeEditorModule>("LandscapeEditor");
+	const ILandscapeFileFormat<T>* FileFormat = LandscapeEditorModule.GetFormatByExtension<T>(*FPaths::GetExtension(InImageFilename, true));
+
+	if (!FileFormat)
+	{
+		Result.ResultCode = ELandscapeImportResult::Error;
+		return Result;
+	}
+
+	const FLandscapeFileInfo FileInfo = FileFormat->Validate(InImageFilename);
+
+	if (FileInfo.ResultCode != ELandscapeImportResult::Error && FileInfo.PossibleResolutions.Num() > 0)
+	{
+		FLandscapeFileResolution ExpectedResolution = FileInfo.PossibleResolutions[0];
+		FLandscapeImportData<T> ImportData = FileFormat->Import(InImageFilename, ExpectedResolution);
+
+		const int32 BufferSize = ImportData.Data.Num() * sizeof(T);
+		NewImageData.Data = TSharedPtr<TArray<uint8>>(new TArray<uint8>());
+		NewImageData.Data->SetNumUninitialized(BufferSize);
+		FMemory::Memcpy(NewImageData.Data->GetData(), ImportData.Data.GetData(), BufferSize);
+		NewImageData.Resolution = FIntPoint(ExpectedResolution.Width, ExpectedResolution.Height);
+		NewImageData.Result = FileInfo.ResultCode;
+		NewImageData.ErrorMessage = FileInfo.ErrorMessage;
+		NewImageData.BytesPerPixel = BufferSize / (ExpectedResolution.Width * ExpectedResolution.Height);
+	}
+	else
+	{
+		return FileInfo;
+	}
+
+	Trim();
+	OutImageData = NewImageData;
+	Add<T>(FString(InImageFilename), OutImageData);
+
+	Result.PossibleResolutions.Add(FLandscapeFileResolution(OutImageData.Resolution.X, OutImageData.Resolution.Y));
+	Result.ResultCode = FileInfo.ResultCode;
+	Result.ErrorMessage = FileInfo.ErrorMessage;
+
+	return Result;
+}
+
+// FindImage is a public symbol, so requires explicit instantiation.  Or if its inlined, then everything it references must be inlined
+// or explicitly instantiated.
+template FLandscapeFileInfo FLandscapeImageFileCache::FindImage<uint8>(const TCHAR* InImageFilename, FLandscapeImageDataRef& OutImageData);
+template FLandscapeFileInfo FLandscapeImageFileCache::FindImage<uint16>(const TCHAR* InImageFilename, FLandscapeImageDataRef& OutImageData);
+
+
 void FLandscapeImageFileCache::MonitorCallback(const TArray<struct FFileChangeData>& Changes)
 {
 	for (const FFileChangeData& Change : Changes)
diff --git a/Engine/Source/Editor/LandscapeEditor/Private/LandscapeImageFileCache.h b/Engine/Source/Editor/LandscapeEditor/Private/LandscapeImageFileCache.h
--- a/Engine/Source/Editor/LandscapeEditor/Private/LandscapeImageFileCache.h
+++ b/Engine/Source/Editor/LandscapeEditor/Private/LandscapeImageFileCache.h
@@ -34,62 +34,7 @@ public:
 	~FLandscapeImageFileCache();
 
 	template<typename T>
-	FLandscapeFileInfo FindImage(const TCHAR* InImageFilename, FLandscapeImageDataRef& OutImageData)
-	{
-		FCacheEntry* CacheEntry = ChooseCache<T>().Find(InImageFilename);
-		FLandscapeFileInfo Result;
-
-		if (CacheEntry)
-		{
-			CacheEntry->UsageCount++;
-			OutImageData = CacheEntry->ImageData;
-			Result.PossibleResolutions.Add(FLandscapeFileResolution(OutImageData.Resolution.X, OutImageData.Resolution.Y));
-			Result.ResultCode = CacheEntry->ImageData.Result;
-			Result.ErrorMessage = CacheEntry->ImageData.ErrorMessage;
-			return Result;
-		}
-
-		FLandscapeImageDataRef NewImageData;
-		ILandscapeEditorModule& LandscapeEditorModule = FModuleManager::GetModuleChecked<ILandscapeEditorModule>("LandscapeEditor");
-		const ILandscapeFileFormat<T>* FileFormat = LandscapeEditorModule.GetFormatByExtension<T>(*FPaths::GetExtension(InImageFilename, true));
-
-		if (!FileFormat)
-		{
-			Result.ResultCode = ELandscapeImportResult::Error;
-			return Result;
-		}
-
-		const FLandscapeFileInfo FileInfo = FileFormat->Validate(InImageFilename);
-
-		if (FileInfo.ResultCode != ELandscapeImportResult::Error && FileInfo.PossibleResolutions.Num() > 0)
-		{
-			FLandscapeFileResolution ExpectedResolution = FileInfo.PossibleResolutions[0];
-			FLandscapeImportData<T> ImportData = FileFormat->Import(InImageFilename, ExpectedResolution);
-
-			const int32 BufferSize = ImportData.Data.Num() * sizeof(T);
-			NewImageData.Data = TSharedPtr<TArray<uint8>>(new TArray<uint8>());
-			NewImageData.Data->SetNumUninitialized(BufferSize);
-			FMemory::Memcpy(NewImageData.Data->GetData(), ImportData.Data.GetData(), BufferSize);
-			NewImageData.Resolution = FIntPoint(ExpectedResolution.Width, ExpectedResolution.Height);
-			NewImageData.Result = FileInfo.ResultCode;
-			NewImageData.ErrorMessage = FileInfo.ErrorMessage;
-			NewImageData.BytesPerPixel = BufferSize / (ExpectedResolution.Width * ExpectedResolution.Height);
-		}
-		else
-		{
-			return FileInfo;
-		}
-		
-		Trim();
-		OutImageData = NewImageData;
-		Add<T>(FString(InImageFilename), OutImageData);
-		
-		Result.PossibleResolutions.Add(FLandscapeFileResolution(OutImageData.Resolution.X, OutImageData.Resolution.Y));
-		Result.ResultCode = FileInfo.ResultCode;
-		Result.ErrorMessage = FileInfo.ErrorMessage;
-
-		return Result;
-	}
+	FLandscapeFileInfo FindImage(const TCHAR* InImageFilename, FLandscapeImageDataRef& OutImageData);
 
 	void SetMaxSize(uint64 InNewMaxSize);
 
@@ -126,7 +71,8 @@ private:
 
 	// Two cache maps.  Hold 8 bit and 16 bit data separately.  Otherwise importing the same file as weightmap and heightmap
 	// will use the cached 8bit version for the heightmap.
-	enum {
+	enum
+	{
 		Cache8 = 0,
 		Cache16 = 1
 	};
