Fixed target layers filter not appearing when no target layer was selected

diff --git a/Engine/Source/Editor/LandscapeEditor/Private/LandscapeEditorDetailCustomization_TargetLayers.cpp b/Engine/Source/Editor/LandscapeEditor/Private/LandscapeEditorDetailCustomization_TargetLayers.cpp
--- a/Engine/Source/Editor/LandscapeEditor/Private/LandscapeEditorDetailCustomization_TargetLayers.cpp
+++ b/Engine/Source/Editor/LandscapeEditor/Private/LandscapeEditorDetailCustomization_TargetLayers.cpp
@@ -571,9 +571,17 @@ EVisibility FLandscapeEditorCustomNodeBuilder_TargetLayers::GetLayersFilterVisib
 
 	if (LandscapeEdMode && LandscapeEdMode->CurrentToolMode)
 	{
-		if (LandscapeEdMode->CurrentToolTarget.TargetType == ELandscapeToolTargetType::Weightmap)
+		// ELandscapeToolTargetType::Invalid means "weightmap with no valid paint layer" so we still want to display that property if it has been marked to be displayed in Weightmap target type, to be consistent 
+		if ((LandscapeEdMode->CurrentToolTarget.TargetType == ELandscapeToolTargetType::Weightmap)
+			|| (LandscapeEdMode->CurrentToolTarget.TargetType == ELandscapeToolTargetType::Invalid))
 		{
-			return EVisibility::Visible;
+			const bool bContainsWeightmapLayers = LandscapeEdMode->GetTargetList().ContainsByPredicate([](TSharedRef<FLandscapeTargetListInfo> InInfo)
+				{
+					FName LayerName = InInfo->GetLayerName();
+					return (LayerName != NAME_None) && (LayerName != UMaterialExpressionLandscapeVisibilityMask::ParameterName);
+				});
+
+			return bContainsWeightmapLayers ? EVisibility::Visible : EVisibility::Collapsed;
 		}
 	}
 
