Fix editor crash from importing heightmap images with invalid tile coordinates.

diff --git a/Engine/Source/Editor/LandscapeEditor/Private/LandscapeTiledImage.cpp b/Engine/Source/Editor/LandscapeEditor/Private/LandscapeTiledImage.cpp
--- a/Engine/Source/Editor/LandscapeEditor/Private/LandscapeTiledImage.cpp
+++ b/Engine/Source/Editor/LandscapeEditor/Private/LandscapeTiledImage.cpp
@@ -148,6 +148,16 @@ FLandscapeFileInfo FLandscapeTiledImage::Load(const TCHAR* Filename)
 		return Result;
 	}
 
+	// Check for int overflows due to too large SizeInTiles values from the filenames
+	if (SizeInTiles.X > std::numeric_limits<int32>::max() / TileResolution.X ||
+		SizeInTiles.Y > std::numeric_limits<int32>::max() / TileResolution.Y ||
+		SizeInTiles.X <= 0 || SizeInTiles.Y <= 0)
+	{
+		Result.ResultCode = ELandscapeImportResult::Error;
+		Result.ErrorMessage = LOCTEXT("FileReadErrorTileCoordsInvalid", "Invalid tiled image coordinates");
+		return Result;
+	}
+
 	Result.PossibleResolutions.Add(FLandscapeFileResolution(GetResolution().X, GetResolution().Y));
 
 	return Result;
