* Fixed landscape target layers UX refresh issues
* Added info message to explain to the user that all the landscape's target layers are hidden (either by the "Hide  unused layers" filter of by the filter text box)
* Used a standard show/hide icon button to the "Hide unused layers" toggle button" in order to help the user understand the status of that filter
* Allowed drag and dropping target layers even when the "Hide unused layers" filter is active

diff --git a/Engine/Source/Editor/LandscapeEditor/Private/LandscapeEdMode.cpp b/Engine/Source/Editor/LandscapeEditor/Private/LandscapeEdMode.cpp
--- a/Engine/Source/Editor/LandscapeEditor/Private/LandscapeEdMode.cpp
+++ b/Engine/Source/Editor/LandscapeEditor/Private/LandscapeEdMode.cpp
@@ -129,7 +129,7 @@ FEdModeLandscape::FEdModeLandscape()
 	, ToolActiveViewport(nullptr)
 	, bIsPaintingInVR(false)
 	, InteractorPainting(nullptr)
-	, bNeedsUpdateShownLayerList(false)
+	, bNeedsUpdateLayerUsageInformation(false)
 	, bUpdatingLandscapeInfo(false)
 {
 	using namespace UE::Landscape::Editor::Tool;
@@ -343,9 +343,9 @@ void FEdModeLandscape::OnCanHaveLayersContentChanged()
 
 void FEdModeLandscape::PostUpdateLayerContent()
 {
-	if (bNeedsUpdateShownLayerList)
+	if (bNeedsUpdateLayerUsageInformation)
 	{
-		UpdateShownLayerList();
+		UpdateLayerUsageInformation();
 	}
 }
 
@@ -612,7 +612,6 @@ void FEdModeLandscape::Enter()
 
 	// It is cleared on exit so update here even if LandscapeInfo hasn't changed
 	UpdateTargetList();
-	UpdateShownLayerList();
 
 	FName ToolkitPalette = NAME_None;
 
@@ -2399,7 +2398,7 @@ int32 FEdModeLandscape::UpdateLandscapeList()
 
 			SetCurrentLayer(0);
 
-			UpdateShownLayerList();
+			UpdateLayerUsageInformation();
 						
 			if (!CurrentToolName.IsNone())
 			{
@@ -2475,7 +2474,7 @@ void FEdModeLandscape::SetTargetLandscape(const TWeakObjectPtr<ULandscapeInfo>&
 		}
 	}
 
-	UpdateShownLayerList();
+	UpdateLayerUsageInformation();
 }
 
 bool FEdModeLandscape::CanEditCurrentTarget(FText* Reason) const
@@ -2611,7 +2610,7 @@ void FEdModeLandscape::UpdateTargetList()
 
 	TargetsListUpdated.Broadcast();
 
-	UpdateShownLayerList();
+	UpdateLayerUsageInformation();
 }
 
 void FEdModeLandscape::UpdateTargetLayerDisplayOrder(ELandscapeLayerDisplayMode InTargetDisplayOrder)
@@ -2731,68 +2730,20 @@ void FEdModeLandscape::OnLandscapeMaterialChangedDelegate()
 	UpdateTargetList();
 }
 
-void FEdModeLandscape::RequestUpdateShownLayerList()
+void FEdModeLandscape::RequestUpdateLayerUsageInformation()
 {
-	bNeedsUpdateShownLayerList = true;
+	bNeedsUpdateLayerUsageInformation = true;
 
 	if (CurrentToolTarget.LandscapeInfo.IsValid() && !CurrentToolTarget.LandscapeInfo->CanHaveLayersContent())
 	{
-		UpdateShownLayerList(); // do it sync when not in lanscape mode.
-	}
-}
-
-void FEdModeLandscape::UpdateShownLayerList()
-{
-	bNeedsUpdateShownLayerList = false;
-
-	if (!CurrentToolTarget.LandscapeInfo.IsValid())
-	{
-		return;
-	}
-
-	// Make sure usage information is up to date
-	UpdateLayerUsageInformation();
-
-	bool DetailPanelRefreshRequired = false;
-
-	ShownTargetLayerList.Empty();
-
-	const TArray<FName>* DisplayOrderList = GetTargetDisplayOrderList();
-
-	if (DisplayOrderList == nullptr)
-	{
-		return;
-	}
-
-	for (const FName& LayerName : *DisplayOrderList)
-	{
-		for (const TSharedRef<FLandscapeTargetListInfo>& TargetInfo : GetTargetList())
-		{
-			if (TargetInfo->LayerName == LayerName)
-			{
-				// Keep a mapping of visible layer name to display order list so we can drag & drop proper items
-				if (ShouldShowLayer(TargetInfo))
-				{
-					ShownTargetLayerList.Add(TargetInfo->LayerName);
-					DetailPanelRefreshRequired = true;
-				}
-
-				break;
-			}
-		}
-	}	
-
-	if (DetailPanelRefreshRequired)
-	{
-		if (Toolkit.IsValid())
-		{
-			StaticCastSharedPtr<FLandscapeToolKit>(Toolkit)->RefreshDetailPanel();
-		}
+		UpdateLayerUsageInformation(); // do it synchronously when not in edit layers mode
 	}
 }
 
 void FEdModeLandscape::UpdateLayerUsageInformation(TWeakObjectPtr<ULandscapeLayerInfoObject>* LayerInfoObjectThatChanged)
 {
+	bNeedsUpdateLayerUsageInformation = false;
+
 	if (!CurrentToolTarget.LandscapeInfo.IsValid())
 	{
 		return;
@@ -2852,6 +2803,21 @@ void FEdModeLandscape::UpdateLayerUsageInformation(TWeakObjectPtr<ULandscapeLaye
 
 bool FEdModeLandscape::ShouldShowLayer(TSharedRef<FLandscapeTargetListInfo> Target) const
 {
+	if (CurrentToolTarget.TargetType == ELandscapeToolTargetType::Heightmap)
+	{
+		return (Target->TargetType == ELandscapeToolTargetType::Heightmap);
+	}
+	else if (CurrentToolTarget.TargetType == ELandscapeToolTargetType::Visibility)
+	{
+		return (Target->TargetType == ELandscapeToolTargetType::Visibility);
+	}
+
+	// Weightmap case : 
+	if (Target->TargetType != ELandscapeToolTargetType::Weightmap)
+	{
+		return false;
+	}
+
 	if (!UISettings->ShowUnusedLayers && (!Target->LayerInfoObj.IsValid() || !Target->LayerInfoObj.Get()->IsReferencedFromLoadedData))
 	{
 		return false;
@@ -2880,11 +2846,6 @@ bool FEdModeLandscape::ShouldShowLayer(TSharedRef<FLandscapeTargetListInfo> Targ
 	return true;
 }
 
-const TArray<FName>& FEdModeLandscape::GetTargetShownList() const
-{
-	return ShownTargetLayerList;
-}
-
 int32 FEdModeLandscape::GetTargetLayerStartingIndex() const
 {
 	return TargetLayerStartingIndex;
@@ -2928,8 +2889,7 @@ void FEdModeLandscape::MoveTargetLayerDisplayOrder(int32 IndexToMove, int32 Inde
 	LandscapeProxy->TargetDisplayOrder = ELandscapeLayerDisplayMode::UserSpecific;
 	UISettings->TargetDisplayOrder = ELandscapeLayerDisplayMode::UserSpecific;
 
-	// Everytime we move something from the display order we must rebuild the shown layer list
-	UpdateShownLayerList();
+	RefreshDetailPanel();
 }
 
 FEdModeLandscape::FTargetsListUpdated FEdModeLandscape::TargetsListUpdated;
@@ -2957,6 +2917,7 @@ void FEdModeLandscape::HandleLevelsChanged()
 
 void FEdModeLandscape::OnMaterialCompilationFinished(UMaterialInterface* MaterialInterface)
 {
+	// TODO [jonathan.bard] : we should remove this now that the target layer list is not automatically filled anymore
 	if (CurrentToolTarget.LandscapeInfo.IsValid() &&
 		CurrentToolTarget.LandscapeInfo->GetLandscapeProxy() != nullptr &&
 		CurrentToolTarget.LandscapeInfo->GetLandscapeProxy()->GetLandscapeMaterial() != nullptr &&
diff --git a/Engine/Source/Editor/LandscapeEditor/Private/LandscapeEdMode.h b/Engine/Source/Editor/LandscapeEditor/Private/LandscapeEdMode.h
--- a/Engine/Source/Editor/LandscapeEditor/Private/LandscapeEdMode.h
+++ b/Engine/Source/Editor/LandscapeEditor/Private/LandscapeEdMode.h
@@ -528,7 +528,6 @@ public:
 	const TArray<TSharedRef<FLandscapeTargetListInfo>>& GetTargetList() const;
 	UMaterialInterface* GetTargetLandscapeMaterial() const { return CachedLandscapeMaterial; }
 	const TArray<FName>* GetTargetDisplayOrderList() const;
-	const TArray<FName>& GetTargetShownList() const;
 	int32 GetTargetLayerStartingIndex() const;
 	const TArray<FLandscapeListInfo>& GetLandscapeList();
 
@@ -541,9 +540,7 @@ public:
 	void UpdateTargetLayerDisplayOrder(ELandscapeLayerDisplayMode InTargetDisplayOrder);
 	void MoveTargetLayerDisplayOrder(int32 IndexToMove, int32 IndexToDestination);
 
-	/** Update shown layer list */	
-	void RequestUpdateShownLayerList();
-	void UpdateShownLayerList();
+	void RequestUpdateLayerUsageInformation();
 	bool ShouldShowLayer(TSharedRef<FLandscapeTargetListInfo> Target) const;
 	void UpdateLayerUsageInformation(TWeakObjectPtr<ULandscapeLayerInfoObject>* LayerInfoObjectThatChanged = nullptr);
 	void OnLandscapeMaterialChangedDelegate();
@@ -650,8 +647,7 @@ private:
 	TArray<TSharedRef<FLandscapeTargetListInfo>> LandscapeTargetList;
 	TArray<FLandscapeListInfo> LandscapeList;
 	TArray<ALandscapeBlueprintBrushBase*> BrushList;
-	TArray<FName> ShownTargetLayerList;
-	
+
 	/** Represent the index offset of the target layer in LandscapeTargetList */
 	int32 TargetLayerStartingIndex;
 
@@ -674,6 +670,6 @@ private:
 	class UViewportInteractor* InteractorPainting;
 
 	/** Delayed refresh */
-	bool bNeedsUpdateShownLayerList;
+	bool bNeedsUpdateLayerUsageInformation;
 	bool bUpdatingLandscapeInfo;
 };
\ No newline at end of file
diff --git a/Engine/Source/Editor/LandscapeEditor/Private/LandscapeEditorDetailCustomization_Layers.cpp b/Engine/Source/Editor/LandscapeEditor/Private/LandscapeEditorDetailCustomization_Layers.cpp
--- a/Engine/Source/Editor/LandscapeEditor/Private/LandscapeEditorDetailCustomization_Layers.cpp
+++ b/Engine/Source/Editor/LandscapeEditor/Private/LandscapeEditorDetailCustomization_Layers.cpp
@@ -704,7 +704,7 @@ void FLandscapeEditorCustomNodeBuilder_Layers::ClearTargetLayerOnLayer(int32 InL
 			{
 				const FScopedTransaction Transaction(LOCTEXT("Landscape_Layers_TargetClear", "Clear Target Layer"));
 				Landscape->ClearPaintLayer(InLayerIndex, InLayerInfo);
-				LandscapeEdMode->RequestUpdateShownLayerList();
+				LandscapeEdMode->RequestUpdateLayerUsageInformation();
 			}
 		}
 	}
@@ -799,7 +799,7 @@ void FLandscapeEditorCustomNodeBuilder_Layers::ClearTargetLayersOnLayer(int32 In
 				OnLayerSelectionChanged(InLayerIndex);
 				if (InClearMode & ELandscapeClearMode::Clear_Weightmap)
 				{
-					LandscapeEdMode->RequestUpdateShownLayerList();
+					LandscapeEdMode->RequestUpdateLayerUsageInformation();
 				}
 			}
 		}
diff --git a/Engine/Source/Editor/LandscapeEditor/Private/LandscapeEditorDetailCustomization_TargetLayers.cpp b/Engine/Source/Editor/LandscapeEditor/Private/LandscapeEditorDetailCustomization_TargetLayers.cpp
--- a/Engine/Source/Editor/LandscapeEditor/Private/LandscapeEditorDetailCustomization_TargetLayers.cpp
+++ b/Engine/Source/Editor/LandscapeEditor/Private/LandscapeEditorDetailCustomization_TargetLayers.cpp
@@ -103,6 +103,17 @@ void FLandscapeEditorDetailCustomization_TargetLayers::CustomizeDetails(IDetailL
 				.Text(LOCTEXT("PopulateTargetLayers_Tip", "There are currently no target layers assigned to this landscape. Use the buttons above to add new ones or populate them from the material(s) currently assigned to the landscape"))
 				.AutoWrapText(true)
 		];
+
+	TargetsCategory.AddCustomRow(FText())
+		.Visibility(TAttribute<EVisibility>::Create(TAttribute<EVisibility>::FGetter::CreateStatic(&FLandscapeEditorDetailCustomization_TargetLayers::GetFilteredTargetLayersListInfoTipVisibility)))
+		[
+			SNew(SMultiLineEditableTextBox)
+				.IsReadOnly(true)
+				.Font(DetailBuilder.GetDetailFontBold())
+				.BackgroundColor(FAppStyle::GetColor("InfoReporting.BackgroundColor"))
+				.Text(LOCTEXT("FilteredTargetLayers_Tip", "All target layers assigned to this landscape are currently filtered. Use the buttons and/or the filter above to un-hide them."))
+				.AutoWrapText(true)
+		];
 }
 END_SLATE_FUNCTION_BUILD_OPTIMIZATION
 
@@ -187,6 +198,25 @@ EVisibility FLandscapeEditorDetailCustomization_TargetLayers::GetPopulateTargetL
 	return EVisibility::Collapsed;
 }
 
+EVisibility FLandscapeEditorDetailCustomization_TargetLayers::GetFilteredTargetLayersListInfoTipVisibility()
+{
+	FEdModeLandscape* LandscapeEdMode = GetEditorMode();
+	if (LandscapeEdMode && LandscapeEdMode->CurrentToolTarget.LandscapeInfo.IsValid())
+	{
+		if ((LandscapeEdMode->CurrentToolTarget.TargetType == ELandscapeToolTargetType::Weightmap)
+			|| (LandscapeEdMode->CurrentToolTarget.TargetType == ELandscapeToolTargetType::Invalid)) // ELandscapeToolTargetType::Invalid means "weightmap with no valid paint layer" 
+		{
+			const TArray<TSharedRef<FLandscapeTargetListInfo>>& TargetList = LandscapeEdMode->GetTargetList();
+			// The first target layers are for heightmap and visibility so only consider target layers above the starting index : 
+			const bool bHasTargetLayers = TargetList.Num() > LandscapeEdMode->GetTargetLayerStartingIndex();
+			const TArray<TSharedRef<FLandscapeTargetListInfo>> TargetDisplayList = FLandscapeEditorCustomNodeBuilder_TargetLayers::PrepareTargetLayerList(/*bInSort =*/ false, /*bInFilter = */true);
+			return (bHasTargetLayers && TargetDisplayList.IsEmpty()) ? EVisibility::Visible : EVisibility::Collapsed;
+		}
+	}
+
+	return EVisibility::Collapsed;
+}
+
 
 //////////////////////////////////////////////////////////////////////////
 
@@ -298,7 +328,7 @@ void FLandscapeEditorCustomNodeBuilder_TargetLayers::GenerateHeaderRowContent(FD
 						.HeightOverride(16.0f)
 						[
 							SNew(SImage)
-							.Image(FAppStyle::GetBrush("GenericViewButton"))
+								.Image(this, &FLandscapeEditorCustomNodeBuilder_TargetLayers::GetShowUnusedBrush)
 						]
 					]
 				]
@@ -483,6 +513,16 @@ TSharedRef<SWidget> FLandscapeEditorCustomNodeBuilder_TargetLayers::GetTargetLay
 	return MenuBuilder.MakeWidget();
 }
 
+const FSlateBrush* FLandscapeEditorCustomNodeBuilder_TargetLayers::GetShowUnusedBrush() const
+{
+	const FSlateBrush* Brush = FAppStyle::GetBrush("Level.VisibleIcon16x");
+	if (FEdModeLandscape* LandscapeEdMode = GetEditorMode(); (LandscapeEdMode != nullptr) && !LandscapeEdMode->UISettings->ShowUnusedLayers)
+	{
+		Brush = FAppStyle::GetBrush("Level.NotVisibleIcon16x");
+	}
+	return Brush;
+}
+
 void FLandscapeEditorCustomNodeBuilder_TargetLayers::ShowUnusedLayers(bool Result)
 {
 	TargetShowUnusedLayersPropertyHandle->SetValue(Result);
@@ -490,8 +530,7 @@ void FLandscapeEditorCustomNodeBuilder_TargetLayers::ShowUnusedLayers(bool Resul
 
 bool FLandscapeEditorCustomNodeBuilder_TargetLayers::ShouldShowUnusedLayers(bool Result) const
 {
-	FEdModeLandscape* LandscapeEdMode = GetEditorMode();
-	if (LandscapeEdMode != nullptr)
+	if (FEdModeLandscape* LandscapeEdMode = GetEditorMode())
 	{
 		return LandscapeEdMode->UISettings->ShowUnusedLayers == Result;
 	}
@@ -534,7 +573,8 @@ const FSlateBrush* FLandscapeEditorCustomNodeBuilder_TargetLayers::GetTargetLaye
 
 EVisibility FLandscapeEditorCustomNodeBuilder_TargetLayers::ShouldShowLayer(TSharedRef<FLandscapeTargetListInfo> Target) const
 {
-	if (Target->TargetType == ELandscapeToolTargetType::Weightmap)
+	if ((Target->TargetType == ELandscapeToolTargetType::Weightmap)
+		|| (Target->TargetType == ELandscapeToolTargetType::Invalid)) // Invalid means weightmap with no selected target layer
 	{
 		FEdModeLandscape* LandscapeEdMode = GetEditorMode();
 
@@ -598,6 +638,38 @@ FText FLandscapeEditorCustomNodeBuilder_TargetLayers::GetLayersFilterText() cons
 	return FText();
 }
 
+TArray<TSharedRef<FLandscapeTargetListInfo>> FLandscapeEditorCustomNodeBuilder_TargetLayers::PrepareTargetLayerList(bool bInSort, bool bInFilter)
+{
+	FEdModeLandscape* LandscapeEdMode = GetEditorMode();
+	if (LandscapeEdMode == nullptr)
+	{
+		return {};
+	}
+	const TArray<TSharedRef<FLandscapeTargetListInfo>>& TargetList = LandscapeEdMode->GetTargetList();
+	const TArray<FName>* TargetDisplayOrderList = LandscapeEdMode->GetTargetDisplayOrderList();
+	if (TargetDisplayOrderList == nullptr)
+	{
+		return {};
+	}
+
+	TArray<TSharedRef<FLandscapeTargetListInfo>> FinalList(TargetList);
+	if (bInFilter)
+	{
+		FinalList.RemoveAllSwap([LandscapeEdMode](TSharedRef<FLandscapeTargetListInfo> InTargetInfo)
+			{
+				return !LandscapeEdMode->ShouldShowLayer(InTargetInfo);
+			});
+	}
+
+	if (bInSort)
+	{
+		Algo::SortBy(FinalList, [TargetDisplayOrderList](TSharedRef<FLandscapeTargetListInfo> InTargetInfo)
+			{
+				return TargetDisplayOrderList->Find(InTargetInfo->GetLayerName());
+			});
+	}
+	return FinalList;
+}
 
 BEGIN_SLATE_FUNCTION_BUILD_OPTIMIZATION
 void FLandscapeEditorCustomNodeBuilder_TargetLayers::GenerateChildContent(IDetailChildrenBuilder& ChildrenBuilder)
@@ -605,15 +677,6 @@ void FLandscapeEditorCustomNodeBuilder_TargetLayers::GenerateChildContent(IDetai
 	FEdModeLandscape* LandscapeEdMode = GetEditorMode();
 	if (LandscapeEdMode != nullptr)
 	{
-		const TArray<TSharedRef<FLandscapeTargetListInfo>>& TargetList = LandscapeEdMode->GetTargetList();
-		const TArray<FName>* TargetDisplayOrderList = LandscapeEdMode->GetTargetDisplayOrderList();
-		const TArray<FName>& TargetShownLayerList = LandscapeEdMode->GetTargetShownList();
-
-		if (TargetDisplayOrderList == nullptr)
-		{
-			return;
-		}
-
 		TSharedPtr<SDragAndDropVerticalBox> TargetLayerList = SNew(SDragAndDropVerticalBox)
 			.OnCanAcceptDrop(this, &FLandscapeEditorCustomNodeBuilder_TargetLayers::HandleCanAcceptDrop)
 			.OnAcceptDrop(this, &FLandscapeEditorCustomNodeBuilder_TargetLayers::HandleAcceptDrop)
@@ -652,25 +715,18 @@ void FLandscapeEditorCustomNodeBuilder_TargetLayers::GenerateChildContent(IDetai
 					]
 			];
 
-		for (int32 i = 0; i < TargetDisplayOrderList->Num(); ++i)
+		// Generate a row for all target layers, including those that will be filtered and let the row's visibility lambda to compute their visibility dynamically. This allows 
+		//  filtering to work without refreshing the details panel (which causes the search box to lose focus) :
+		for (const TSharedRef<FLandscapeTargetListInfo>& TargetInfo : PrepareTargetLayerList(/*bInSort = */true, /*bInFilter = */false))
 		{
-			for (const TSharedRef<FLandscapeTargetListInfo>& TargetInfo : TargetList)
+			TSharedPtr<SWidget> GeneratedRowWidget = GenerateRow(TargetInfo);
+			if (GeneratedRowWidget.IsValid())
 			{
-				if (TargetInfo->LayerName == (*TargetDisplayOrderList)[i] && (TargetInfo->TargetType != ELandscapeToolTargetType::Weightmap || TargetShownLayerList.Find(TargetInfo->LayerName) != INDEX_NONE))
-				{
-					TSharedPtr<SWidget> GeneratedRowWidget = GenerateRow(TargetInfo);
-
-					if (GeneratedRowWidget.IsValid())
-					{
-						TargetLayerList->AddSlot()
-						.AutoHeight()						
-						[
-							GeneratedRowWidget.ToSharedRef()
-						];
-					}
-
-					break;
-				}
+				TargetLayerList->AddSlot()
+					.AutoHeight()
+					[
+						GeneratedRowWidget.ToSharedRef()
+					];
 			}
 		}
 	}
@@ -1034,33 +1090,16 @@ FReply FLandscapeEditorCustomNodeBuilder_TargetLayers::HandleDragDetected(const
 
 	if (LandscapeEdMode != nullptr)
 	{
-		const TArray<FName>& TargetShownList = LandscapeEdMode->GetTargetShownList();
-
-		if (TargetShownList.IsValidIndex(SlotIndex))
+		// The slot index corresponds to what is actually shown, so we need to both sort and filter the target layer list here :
+		TArray<TSharedRef<FLandscapeTargetListInfo>> TargetDisplayList = PrepareTargetLayerList(/*bInSort =*/ true, /*bInFilter = */true);
+		if (TargetDisplayList.IsValidIndex(SlotIndex))
 		{
-			const TArray<FName>* TargetDisplayOrderList = LandscapeEdMode->GetTargetDisplayOrderList();
-
-			if (TargetDisplayOrderList != nullptr)
+			if (const TArray<FName>* TargetDisplayOrderList = LandscapeEdMode->GetTargetDisplayOrderList())
 			{
-				FName ShownTargetName = LandscapeEdMode->UISettings->ShowUnusedLayers && TargetShownList.IsValidIndex(SlotIndex + LandscapeEdMode->GetTargetLayerStartingIndex()) ? TargetShownList[SlotIndex + LandscapeEdMode->GetTargetLayerStartingIndex()] : TargetShownList[SlotIndex];
-				int32 DisplayOrderLayerIndex = TargetDisplayOrderList->Find(ShownTargetName);
-
-				if (TargetDisplayOrderList->IsValidIndex(DisplayOrderLayerIndex))
+				TSharedPtr<SWidget> Row = GenerateRow(TargetDisplayList[SlotIndex]);
+				if (Row.IsValid())
 				{
-					const TArray<TSharedRef<FLandscapeTargetListInfo>>& TargetList = LandscapeEdMode->GetTargetList();
-
-					for (const TSharedRef<FLandscapeTargetListInfo>& TargetInfo : TargetList)
-					{
-						if (TargetInfo->LayerName == (*TargetDisplayOrderList)[DisplayOrderLayerIndex])
-						{
-							TSharedPtr<SWidget> Row = GenerateRow(TargetInfo);
-
-							if (Row.IsValid())
-							{
-								return FReply::Handled().BeginDragDrop(FTargetLayerDragDropOp::New(SlotIndex, Slot, Row));
-							}
-						}
-					}
+					return FReply::Handled().BeginDragDrop(FTargetLayerDragDropOp::New(SlotIndex, Slot, Row));
 				}
 			}
 		}
@@ -1091,21 +1130,20 @@ FReply FLandscapeEditorCustomNodeBuilder_TargetLayers::HandleAcceptDrop(FDragDro
 
 		if (LandscapeEdMode != nullptr)
 		{
-			const TArray<FName>& TargetShownList = LandscapeEdMode->GetTargetShownList();
+			// The slot index corresponds to what is actually shown, so we need to both sort and filter the target layer list here :
+			TArray<TSharedRef<FLandscapeTargetListInfo>> TargetDisplayList = PrepareTargetLayerList(/*bInSort =*/ true, /*bInFilter = */true);
 
-			if (TargetShownList.IsValidIndex(DragDropOperation->SlotIndexBeingDragged) && TargetShownList.IsValidIndex(SlotIndex))
+			if (TargetDisplayList.IsValidIndex(DragDropOperation->SlotIndexBeingDragged) && TargetDisplayList.IsValidIndex(SlotIndex))
 			{
-				const TArray<FName>* TargetDisplayOrderList = LandscapeEdMode->GetTargetDisplayOrderList();
-
-				if (TargetDisplayOrderList != nullptr && TargetShownList.IsValidIndex(DragDropOperation->SlotIndexBeingDragged + LandscapeEdMode->GetTargetLayerStartingIndex()) && TargetShownList.IsValidIndex(SlotIndex + LandscapeEdMode->GetTargetLayerStartingIndex()))
+				const FName TargetLayerNameBeingDragged = TargetDisplayList[DragDropOperation->SlotIndexBeingDragged]->GetLayerName();
+				const FName DestinationTargetLayerName = TargetDisplayList[SlotIndex]->GetLayerName();
+				if (const TArray<FName>* TargetDisplayOrderList = LandscapeEdMode->GetTargetDisplayOrderList())
 				{
-					int32 StartingLayerIndex = TargetDisplayOrderList->Find(LandscapeEdMode->UISettings->ShowUnusedLayers ? TargetShownList[DragDropOperation->SlotIndexBeingDragged + LandscapeEdMode->GetTargetLayerStartingIndex()] : TargetShownList[DragDropOperation->SlotIndexBeingDragged]);
-					int32 DestinationLayerIndex = TargetDisplayOrderList->Find(LandscapeEdMode->UISettings->ShowUnusedLayers ? TargetShownList[SlotIndex + LandscapeEdMode->GetTargetLayerStartingIndex()] : TargetShownList[SlotIndex]);
-
+					int32 StartingLayerIndex = TargetDisplayOrderList->Find(TargetLayerNameBeingDragged);
+					int32 DestinationLayerIndex = TargetDisplayOrderList->Find(DestinationTargetLayerName);
 					if (StartingLayerIndex != INDEX_NONE && DestinationLayerIndex != INDEX_NONE)
 					{
 						LandscapeEdMode->MoveTargetLayerDisplayOrder(StartingLayerIndex, DestinationLayerIndex);
-
 						return FReply::Handled();
 					}
 				}
@@ -1386,7 +1424,7 @@ void FLandscapeEditorCustomNodeBuilder_TargetLayers::OnClearLayer(const TSharedR
 			FScopedSetLandscapeEditingLayer Scope(LandscapeEdMode->GetLandscape(), LandscapeEdMode->GetCurrentLayerGuid(), [&] { LandscapeEdMode->RequestLayersContentUpdateForceAll(); });
 			FLandscapeEditDataInterface LandscapeEdit(Target->LandscapeInfo.Get());
 			LandscapeEdit.DeleteLayer(Target->LayerInfoObj.Get());
-			LandscapeEdMode->RequestUpdateShownLayerList();
+			LandscapeEdMode->RequestUpdateLayerUsageInformation();
 		}
 	}
 }
diff --git a/Engine/Source/Editor/LandscapeEditor/Private/LandscapeEditorDetailCustomization_TargetLayers.h b/Engine/Source/Editor/LandscapeEditor/Private/LandscapeEditorDetailCustomization_TargetLayers.h
--- a/Engine/Source/Editor/LandscapeEditor/Private/LandscapeEditorDetailCustomization_TargetLayers.h
+++ b/Engine/Source/Editor/LandscapeEditor/Private/LandscapeEditorDetailCustomization_TargetLayers.h
@@ -41,6 +41,7 @@ protected:
 	static EVisibility GetPaintingRestrictionVisibility();
 	static EVisibility GetVisibilityMaskTipVisibility();
 	static EVisibility GetPopulateTargetLayersInfoTipVisibility();
+	static EVisibility GetFilteredTargetLayersListInfoTipVisibility();
 };
 
 class FLandscapeEditorCustomNodeBuilder_TargetLayers : public IDetailCustomNodeBuilder, public TSharedFromThis<FLandscapeEditorCustomNodeBuilder_TargetLayers>
@@ -57,6 +58,8 @@ public:
 	virtual bool InitiallyCollapsed() const override { return false; }
 	virtual FName GetName() const override { return "TargetLayers"; }
 
+	static TArray<TSharedRef<FLandscapeTargetListInfo>> PrepareTargetLayerList(bool bInSort, bool bInFilter);
+
 protected:
 	TSharedRef<FAssetThumbnailPool> ThumbnailPool;
 	TSharedRef<IPropertyHandle> TargetDisplayOrderPropertyHandle;
@@ -113,6 +116,7 @@ protected:
 	bool IsSelectedDisplayOrder(ELandscapeLayerDisplayMode InDisplayOrder) const;
 
 	TSharedRef<SWidget> GetTargetLayerShowUnusedButtonMenuContent();
+	const FSlateBrush* GetShowUnusedBrush() const;
 	void ShowUnusedLayers(bool Result);
 	bool ShouldShowUnusedLayers(bool Result) const;
 	EVisibility ShouldShowLayer(TSharedRef<FLandscapeTargetListInfo> Target) const;
diff --git a/Engine/Source/Editor/LandscapeEditor/Private/LandscapeEditorObject.cpp b/Engine/Source/Editor/LandscapeEditor/Private/LandscapeEditorObject.cpp
--- a/Engine/Source/Editor/LandscapeEditor/Private/LandscapeEditorObject.cpp
+++ b/Engine/Source/Editor/LandscapeEditor/Private/LandscapeEditorObject.cpp
@@ -197,12 +197,6 @@ void ULandscapeEditorObject::PostEditChangeProperty(FPropertyChangedEvent& Prope
 	{
 		UpdateTargetLayerDisplayOrder();
 	}
-
-	if (PropertyChangedEvent.MemberProperty == nullptr ||
-		PropertyChangedEvent.MemberProperty->GetFName() == GET_MEMBER_NAME_CHECKED(ULandscapeEditorObject, ShowUnusedLayers))
-	{
-		UpdateShowUnusedLayers();
-	}
 }
 
 /** Load UI settings from ini file */
@@ -989,14 +983,6 @@ void ULandscapeEditorObject::UpdateTargetLayerDisplayOrder()
 	}
 }
 
-void ULandscapeEditorObject::UpdateShowUnusedLayers()
-{
-	if (ParentMode != nullptr)
-	{
-		ParentMode->UpdateShownLayerList();
-	}
-}
-
 float ULandscapeEditorObject::GetCurrentToolStrength() const
 {
 	if (IsWeightmapTarget())
diff --git a/Engine/Source/Editor/LandscapeEditor/Public/LandscapeEditorObject.h b/Engine/Source/Editor/LandscapeEditor/Public/LandscapeEditorObject.h
--- a/Engine/Source/Editor/LandscapeEditor/Public/LandscapeEditorObject.h
+++ b/Engine/Source/Editor/LandscapeEditor/Public/LandscapeEditorObject.h
@@ -811,7 +811,6 @@ public:
 	}
 
 	void UpdateTargetLayerDisplayOrder();
-	void UpdateShowUnusedLayers();
 
 	float GetCurrentToolStrength() const;
 	void SetCurrentToolStrength(float NewToolStrength);
