Fix incorrect check
#rnx

diff --git a/Engine/Source/Editor/LandscapeEditor/Private/LandscapeEditorDetailCustomization_NewLandscape.cpp b/Engine/Source/Editor/LandscapeEditor/Private/LandscapeEditorDetailCustomization_NewLandscape.cpp
--- a/Engine/Source/Editor/LandscapeEditor/Private/LandscapeEditorDetailCustomization_NewLandscape.cpp
+++ b/Engine/Source/Editor/LandscapeEditor/Private/LandscapeEditorDetailCustomization_NewLandscape.cpp
@@ -747,7 +747,7 @@ FText FLandscapeEditorDetailCustomization_NewLandscape::GetSectionsPerComponent(
 {
 	int32 SectionsPerComponent = 0;
 	FPropertyAccess::Result Result = PropertyHandle->GetValue(SectionsPerComponent);
-	check(Result != FPropertyAccess::Success);
+	check(Result != FPropertyAccess::Fail);
 
 	if (Result == FPropertyAccess::MultipleValues)
 	{
