Fix landscape heightmap import from tiled image set when file path contains brackets.

diff --git a/Engine/Source/Editor/LandscapeEditor/Private/LandscapeTiledImage.cpp b/Engine/Source/Editor/LandscapeEditor/Private/LandscapeTiledImage.cpp
--- a/Engine/Source/Editor/LandscapeEditor/Private/LandscapeTiledImage.cpp
+++ b/Engine/Source/Editor/LandscapeEditor/Private/LandscapeTiledImage.cpp
@@ -37,7 +37,8 @@ FLandscapeFileInfo FLandscapeTiledImage::Load(const TCHAR* Filename)
 	FindFiles(Filename, FoundFiles);
 
 	FString Path = FPaths::GetPath(Filename);
-	FString RegexFilename(Filename);
+	FString FilenameNoPath = FPaths::GetCleanFilename(Filename);
+	FString RegexFilename(FilenameNoPath);
 	TMap<int32, FString> Groups;
 
 	for (const TPair<FString, FString>& It : Tokens)
@@ -74,8 +75,8 @@ FLandscapeFileInfo FLandscapeTiledImage::Load(const TCHAR* Filename)
 		for (const FString& FoundFilname : FoundFiles)
 		{
 			FRegexPattern Pattern(RegexFilename);
-			FString FullPath = FPaths::Combine(Path, FoundFilname);
-			FRegexMatcher Matcher(Pattern, FullPath);
+			// Regex test the filename without the path so that conflicting symbols (e.g. brackets) in the path don't confuse it.
+			FRegexMatcher Matcher(Pattern, FoundFilname);
 
 			if (Matcher.FindNext())
 			{
@@ -106,6 +107,7 @@ FLandscapeFileInfo FLandscapeTiledImage::Load(const TCHAR* Filename)
 				if (X >= 0 && Y >= 0)
 				{
 					FIntPoint TileCoord(X, Y);
+					FString FullPath = FPaths::Combine(Path, FoundFilname);
 					TileFilenames.Add(TileCoord, FullPath);
 				}			
 			}
