Small fix for cl 37268466.  Fix non-tiled heightmap import.

#rnx

diff --git a/Engine/Source/Editor/LandscapeEditor/Private/LandscapeTiledImage.cpp b/Engine/Source/Editor/LandscapeEditor/Private/LandscapeTiledImage.cpp
--- a/Engine/Source/Editor/LandscapeEditor/Private/LandscapeTiledImage.cpp
+++ b/Engine/Source/Editor/LandscapeEditor/Private/LandscapeTiledImage.cpp
@@ -60,7 +60,7 @@ FLandscapeFileInfo FLandscapeTiledImage::Load(const TCHAR* Filename)
 		}
 	}
 
-	const bool bExactFilename = RegexFilename == Filename;
+	const bool bExactFilename = RegexFilename == FilenameNoPath;
 
 	SizeInTiles = FIntPoint::NoneValue;
 	TileResolution = FIntPoint::NoneValue;
