Refresh landscape list when EditorModeManager broacasts OnIsEditingDisallowedChanged (the use case is when the Level Instance restrictive mode is disabled/enabled).

diff --git a/Engine/Source/Editor/LandscapeEditor/Private/LandscapeEdMode.cpp b/Engine/Source/Editor/LandscapeEditor/Private/LandscapeEdMode.cpp
--- a/Engine/Source/Editor/LandscapeEditor/Private/LandscapeEdMode.cpp
+++ b/Engine/Source/Editor/LandscapeEditor/Private/LandscapeEdMode.cpp
@@ -469,6 +469,11 @@ void FEdModeLandscape::Enter()
 	OnLevelActorDeletedDelegateHandle = GEngine->OnLevelActorDeleted().AddSP(this, &FEdModeLandscape::OnLevelActorRemoved);
 	OnLevelActorAddedDelegateHandle = GEngine->OnLevelActorAdded().AddSP(this, &FEdModeLandscape::OnLevelActorAdded);
 	PreSaveWorldHandle = FEditorDelegates::PreSaveWorldWithContext.AddSP(this, &FEdModeLandscape::OnPreSaveWorld);
+	FLevelEditorModule& LevelEditorModule = FModuleManager::GetModuleChecked<FLevelEditorModule>("LevelEditor");
+	if (TSharedPtr<ILevelEditor> FirstLevelEditor = LevelEditorModule.GetFirstLevelEditor())
+	{
+		OnIsEditingDisallowedChangedHandle = FirstLevelEditor->GetEditorModeManager().OnIsEditingDisallowedChanged().AddSP(this, &FEdModeLandscape::HandleLevelsChanged);
+	}
 		
 	UpdateToolModes();
 
@@ -684,6 +689,11 @@ void FEdModeLandscape::Exit()
 		}
 	}
 
+	FLevelEditorModule& LevelEditorModule = FModuleManager::GetModuleChecked<FLevelEditorModule>("LevelEditor");
+	if (TSharedPtr<ILevelEditor> FirstLevelEditor = LevelEditorModule.GetFirstLevelEditor())
+	{
+		FirstLevelEditor->GetEditorModeManager().OnIsEditingDisallowedChanged().Remove(OnIsEditingDisallowedChangedHandle);
+	}
 	FEditorDelegates::PreSaveWorldWithContext.Remove(PreSaveWorldHandle);
 	GEngine->OnLevelActorDeleted().Remove(OnLevelActorDeletedDelegateHandle);
 	GEngine->OnLevelActorAdded().Remove(OnLevelActorAddedDelegateHandle);
diff --git a/Engine/Source/Editor/LandscapeEditor/Private/LandscapeEdMode.h b/Engine/Source/Editor/LandscapeEditor/Private/LandscapeEdMode.h
--- a/Engine/Source/Editor/LandscapeEditor/Private/LandscapeEdMode.h
+++ b/Engine/Source/Editor/LandscapeEditor/Private/LandscapeEdMode.h
@@ -662,6 +662,7 @@ private:
 	FDelegateHandle OnLevelActorDeletedDelegateHandle;
 	FDelegateHandle OnLevelActorAddedDelegateHandle;
 	FDelegateHandle PreSaveWorldHandle;
+	FDelegateHandle OnIsEditingDisallowedChangedHandle;
 	
 	/** Check if we are painting using the VREditor */
 	bool bIsPaintingInVR;
diff --git a/Engine/Source/Editor/LevelInstanceEditor/Private/LevelInstanceEditorMode.cpp b/Engine/Source/Editor/LevelInstanceEditor/Private/LevelInstanceEditorMode.cpp
--- a/Engine/Source/Editor/LevelInstanceEditor/Private/LevelInstanceEditorMode.cpp
+++ b/Engine/Source/Editor/LevelInstanceEditor/Private/LevelInstanceEditorMode.cpp
@@ -13,6 +13,7 @@
 #include "LevelInstance/LevelInstanceSubsystem.h"
 #include "LevelInstance/LevelInstanceInterface.h"
 #include "LevelInstance/ILevelInstanceEditorModule.h"
+#include "LevelEditor.h"
 #include "LevelEditorViewport.h"
 #include "LevelEditorActions.h"
 #include "EditorModeManager.h"
@@ -390,6 +391,12 @@ void ULevelInstanceEditorMode::ToggleContextRestrictionCommand()
 	bContextRestriction = !bContextRestriction;
 
 	UpdateEngineShowFlags();
+
+	FLevelEditorModule& LevelEditorModule = FModuleManager::GetModuleChecked<FLevelEditorModule>("LevelEditor");
+	if (TSharedPtr<ILevelEditor> FirstLevelEditor = LevelEditorModule.GetFirstLevelEditor())
+	{
+		FirstLevelEditor->GetEditorModeManager().BroadcastIsEditingDisallowedChanged();
+	}
 }
 
 bool ULevelInstanceEditorMode::IsContextRestrictionCommandEnabled() const
diff --git a/Engine/Source/Editor/UnrealEd/Public/EditorModeManager.h b/Engine/Source/Editor/UnrealEd/Public/EditorModeManager.h
--- a/Engine/Source/Editor/UnrealEd/Public/EditorModeManager.h
+++ b/Engine/Source/Editor/UnrealEd/Public/EditorModeManager.h
@@ -394,7 +394,14 @@ public:
 
 	/**	Broadcasts the CoordSystemChangedEvent event */
 	void BroadcastCoordSystemChanged(ECoordSystem InCoordSystem) { CoordSystemChangedEvent.Broadcast(InCoordSystem); }
-	
+
+	/** Delegate type for triggering when a mode Editing Disallowed changed */
+	DECLARE_EVENT(FEditorModeTools, FIsEditingDisallowedChangedEvent);
+	FIsEditingDisallowedChangedEvent& OnIsEditingDisallowedChanged() { return IsEditingDisallowedChangedEvent; }
+
+	/**	Broadcasts the IsEditingDisallowedChangedEvent event */
+	void BroadcastIsEditingDisallowedChanged() { IsEditingDisallowedChangedEvent.Broadcast(); }
+
 	/**
 	 * Returns the current CoordSystem
 	 * 
@@ -628,6 +635,9 @@ private:
 	/** Multicast delegate that is broadcast when the coordinate system is changed */
 	FCoordSystemChangedEvent CoordSystemChangedEvent;
 
+	/** Multicast delegate that is broadcast when a mode Editing Disallowed changed */
+	FIsEditingDisallowedChangedEvent IsEditingDisallowedChangedEvent;
+
 	/** Flag set between calls to StartTracking() and EndTracking() */
 	bool bIsTracking;
 
