Landscape mode : display the level instance prefix when editing a landscape that is inside the level instance being edited. This allows to distinguish between the world's landscape(s) and the edited level instance one(s) when they are named similarly

diff --git a/Engine/Source/Editor/LandscapeEditor/Private/LandscapeEditorDetails.cpp b/Engine/Source/Editor/LandscapeEditor/Private/LandscapeEditorDetails.cpp
--- a/Engine/Source/Editor/LandscapeEditor/Private/LandscapeEditorDetails.cpp
+++ b/Engine/Source/Editor/LandscapeEditor/Private/LandscapeEditorDetails.cpp
@@ -22,6 +22,7 @@
 #include "DetailCategoryBuilder.h"
 #include "DetailLayoutBuilder.h" 
 #include "VariablePrecisionNumericInterface.h"
+#include "LevelInstance/LevelInstanceSubsystem.h"
 
 #include "Templates/SharedPointer.h"
 
@@ -278,6 +279,23 @@ EVisibility FLandscapeEditorDetails::GetTargetLandscapeSelectorVisibility()
 	return EVisibility::Collapsed;
 }
 
+FText FLandscapeEditorDetails::GetLandscapeDisplayName(ALandscape* InLandscape)
+{
+	FString DisplayString = InLandscape->GetActorLabel();
+	if (UWorld* OwningWorld = InLandscape->GetWorld())
+	{
+		ULevel* Level = InLandscape->GetLevel();
+		ULevelInstanceSubsystem* LevelInstanceSubsystem = UWorld::GetSubsystem<ULevelInstanceSubsystem>(OwningWorld);
+		if ((LevelInstanceSubsystem != nullptr)
+			&& (Level != nullptr)
+			&& (Level != OwningWorld->PersistentLevel))
+		{
+			DisplayString = LevelInstanceSubsystem->PrefixWithParentLevelInstanceActorLabels(DisplayString, Level);
+		}
+	}
+	return FText::FromString(DisplayString);
+}
+
 FText FLandscapeEditorDetails::GetTargetLandscapeName()
 {
 	FEdModeLandscape* LandscapeEdMode = GetEditorMode();
@@ -286,14 +304,12 @@ FText FLandscapeEditorDetails::GetTargetLandscapeName()
 		ULandscapeInfo* Info = LandscapeEdMode->CurrentToolTarget.LandscapeInfo.Get();
 		if (Info)
 		{
-			ALandscapeProxy* Proxy = Info->GetLandscapeProxy();
-			if (Proxy)
+			if (ALandscape* LandscapeActor = Info->LandscapeActor.Get())
 			{
-				return FText::FromString(Proxy->GetActorLabel());
+				return GetLandscapeDisplayName(LandscapeActor);
 			}
 		}
 	}
-
 	return FText();
 }
 
@@ -307,8 +323,11 @@ TSharedRef<SWidget> FLandscapeEditorDetails::GetTargetLandscapeMenu()
 		const TArray<FLandscapeListInfo>& LandscapeList = LandscapeEdMode->GetLandscapeList();
 		for (auto It = LandscapeList.CreateConstIterator(); It; It++)
 		{
-			FUIAction Action = FUIAction(FExecuteAction::CreateStatic(&FLandscapeEditorDetails::OnChangeTargetLandscape, MakeWeakObjectPtr(It->Info)));
-			MenuBuilder.AddMenuEntry(FText::FromString(It->Info->GetLandscapeProxy()->GetActorLabel()), FText(), FSlateIcon(), Action);
+			if (ALandscape* LandscapeActor = It->Info->LandscapeActor.Get())
+			{
+				FUIAction Action = FUIAction(FExecuteAction::CreateStatic(&FLandscapeEditorDetails::OnChangeTargetLandscape, MakeWeakObjectPtr(It->Info)));
+				MenuBuilder.AddMenuEntry(GetLandscapeDisplayName(LandscapeActor), FText(), FSlateIcon(), Action);
+			}
 		}
 
 		return MenuBuilder.MakeWidget();
diff --git a/Engine/Source/Editor/LandscapeEditor/Private/LandscapeEditorDetails.h b/Engine/Source/Editor/LandscapeEditor/Private/LandscapeEditorDetails.h
--- a/Engine/Source/Editor/LandscapeEditor/Private/LandscapeEditorDetails.h
+++ b/Engine/Source/Editor/LandscapeEditor/Private/LandscapeEditorDetails.h
@@ -48,6 +48,7 @@ protected:
 	static FText GetLocalizedName(FString Name);
 
 	static EVisibility GetTargetLandscapeSelectorVisibility();
+	static FText GetLandscapeDisplayName(ALandscape* InLandscape);
 	static FText GetTargetLandscapeName();
 	static TSharedRef<SWidget> GetTargetLandscapeMenu();
 	static void OnChangeTargetLandscape(TWeakObjectPtr<ULandscapeInfo> LandscapeInfo);
