Fix new Landscape offset not matching preview Landscape
(PR #12582)

diff --git a/Engine/Source/Editor/LandscapeEditor/Private/LandscapeEditorDetailCustomization_NewLandscape.cpp b/Engine/Source/Editor/LandscapeEditor/Private/LandscapeEditorDetailCustomization_NewLandscape.cpp
--- a/Engine/Source/Editor/LandscapeEditor/Private/LandscapeEditorDetailCustomization_NewLandscape.cpp
+++ b/Engine/Source/Editor/LandscapeEditor/Private/LandscapeEditorDetailCustomization_NewLandscape.cpp
@@ -1098,7 +1098,7 @@ FReply FLandscapeEditorDetailCustomization_NewLandscape::OnCreateButtonClicked()
 		Transaction = MakeUnique<FScopedTransaction>(bIsNewLandscape ? LOCTEXT("Undo_CreateLandscape", "Creating New Landscape") : LOCTEXT("Undo_CreateAndImportLandscape", "Creating and Importing New Landscape"));
 	}
 
-	const FVector Offset = FTransform(UISettings->NewLandscape_Rotation, FVector::ZeroVector, UISettings->NewLandscape_Scale).TransformVector(FVector(-UISettings->NewLandscape_ComponentCount.X * QuadsPerComponent / 2, -UISettings->NewLandscape_ComponentCount.Y * QuadsPerComponent / 2, 0));
+	const FVector Offset = FTransform(UISettings->NewLandscape_Rotation, FVector::ZeroVector, UISettings->NewLandscape_Scale).TransformVector(FVector(-UISettings->NewLandscape_ComponentCount.X * QuadsPerComponent / 2.0, -UISettings->NewLandscape_ComponentCount.Y * QuadsPerComponent / 2.0, 0.0));
 	
 	ALandscape* Landscape = LandscapeEdMode->GetWorld()->SpawnActor<ALandscape>(UISettings->NewLandscape_Location + Offset, UISettings->NewLandscape_Rotation);
 	Landscape->bCanHaveLayersContent = UISettings->bCanHaveLayersContent;
