Landscape Mode while in Level Instance Edit will allow modifying Persistent Level's Landscape(s) if none are found in the current editing level.

diff --git a/Engine/Source/Editor/LandscapeEditor/Private/LandscapeEdMode.cpp b/Engine/Source/Editor/LandscapeEditor/Private/LandscapeEdMode.cpp
--- a/Engine/Source/Editor/LandscapeEditor/Private/LandscapeEdMode.cpp
+++ b/Engine/Source/Editor/LandscapeEditor/Private/LandscapeEdMode.cpp
@@ -2368,44 +2368,69 @@ int32 FEdModeLandscape::UpdateLandscapeList()
 		int32 Index = 0;
 		auto& LandscapeInfoMap = ULandscapeInfoMap::GetLandscapeInfoMap(World);
 
-		for (auto It = LandscapeInfoMap.Map.CreateIterator(); It; ++It)
+		auto AddLandscape = [this, &CurrentIndex, &Index](ALandscape* Landscape, ULandscapeInfo* LandscapeInfo)
 		{
-			ULandscapeInfo* LandscapeInfo = It.Value();
-			if (IsValid(LandscapeInfo) && 
-				LandscapeInfo->SupportsLandscapeEditing() &&
-				!IsEditingDisallowedByEdModes(LandscapeInfo->LandscapeActor.Get()))
+			if (Landscape)
+			{
+				Landscape->RegisterLandscapeEdMode(this);
+			}
+
+			ALandscapeProxy* LandscapeProxy = LandscapeInfo->GetLandscapeProxy();
+			if (LandscapeProxy)
 			{
-				if (ALandscape* Landscape = LandscapeInfo->LandscapeActor.Get())
+				if (CurrentToolTarget.LandscapeInfo == LandscapeInfo)
 				{
-					Landscape->RegisterLandscapeEdMode(this);
+					CurrentIndex = Index;
+
+					// Update GizmoActor Landscape Target (is this necessary?)
+					if (CurrentGizmoActor.IsValid())
+					{
+						CurrentGizmoActor->SetTargetLandscape(LandscapeInfo);
+					}
 				}
 
-				ALandscapeProxy* LandscapeProxy = LandscapeInfo->GetLandscapeProxy();
-				if (LandscapeProxy)
+				int32 MinX, MinY, MaxX, MaxY;
+				int32 Width = 0, Height = 0;
+				if (LandscapeInfo->GetLandscapeExtent(MinX, MinY, MaxX, MaxY))
 				{
-					if (CurrentToolTarget.LandscapeInfo == LandscapeInfo)
-					{
-						CurrentIndex = Index;
+					Width = MaxX - MinX + 1;
+					Height = MaxY - MinY + 1;
+				}
 
-						// Update GizmoActor Landscape Target (is this necessary?)
-						if (CurrentGizmoActor.IsValid())
-						{
-							CurrentGizmoActor->SetTargetLandscape(LandscapeInfo);
-						}
-					}
+				LandscapeList.Add(FLandscapeListInfo(*LandscapeProxy->GetName(), LandscapeInfo,
+					LandscapeInfo->ComponentSizeQuads, LandscapeInfo->ComponentNumSubsections, Width, Height));
+				Index++;
+			}
+		};
 
-					int32 MinX, MinY, MaxX, MaxY;
-					int32 Width = 0, Height = 0;
-					if (LandscapeInfo->GetLandscapeExtent(MinX, MinY, MaxX, MaxY))
+		TArray<TPair<ALandscape*, ULandscapeInfo*>> PersistentLevelLandscapes;
+		for (auto It = LandscapeInfoMap.Map.CreateIterator(); It; ++It)
+		{
+			ULandscapeInfo* LandscapeInfo = It.Value();
+			if (IsValid(LandscapeInfo) && LandscapeInfo->SupportsLandscapeEditing())
+			{
+				ALandscape* Landscape = LandscapeInfo->LandscapeActor.Get();
+				if (IsEditingDisallowedByEdModes(Landscape))
+				{
+					if (Landscape && Landscape->GetLevel()->IsPersistentLevel())
 					{
-						Width = MaxX - MinX + 1;
-						Height = MaxY - MinY + 1;
+						PersistentLevelLandscapes.Emplace(Landscape, LandscapeInfo);
 					}
-
-					LandscapeList.Add(FLandscapeListInfo(*LandscapeProxy->GetName(), LandscapeInfo,
-						LandscapeInfo->ComponentSizeQuads, LandscapeInfo->ComponentNumSubsections, Width, Height));
-					Index++;
 				}
+				else
+				{
+					AddLandscape(Landscape, LandscapeInfo);
+				}
+			}
+		}
+
+		// Landscape Mode will fallback on Persistent Level's Landscape(s) if none are found after
+		// applying the editor modes IsEditingDisallowed.
+		if (LandscapeList.IsEmpty() && !PersistentLevelLandscapes.IsEmpty())
+		{
+			for (const TPair<ALandscape*, ULandscapeInfo*>& Landscape : PersistentLevelLandscapes)
+			{
+				AddLandscape(Landscape.Key, Landscape.Value);
 			}
 		}
 	}
