Landscape fixes :
* Use defensive attitude towards unloaded landscape info (weak object ptr returning null, likely on unload), across minor landscape tools, like it was already done on most landscape tool brushes
* Speculative crash fix in landscape ramp tool when the sprite texture resource isn't ready while drawing (let the AddSprite method handle the texture size, when the texture resource is invalid)

diff --git a/Engine/Source/Editor/LandscapeEditor/Private/LandscapeEdModeBrushes.cpp b/Engine/Source/Editor/LandscapeEditor/Private/LandscapeEdModeBrushes.cpp
--- a/Engine/Source/Editor/LandscapeEditor/Private/LandscapeEdModeBrushes.cpp
+++ b/Engine/Source/Editor/LandscapeEditor/Private/LandscapeEdModeBrushes.cpp
@@ -1257,8 +1257,7 @@ public:
 	{
 		Super::Tick(ViewportClient, DeltaTime);
 
-		ALandscapeProxy* Proxy = EdMode->CurrentToolTarget.LandscapeInfo.IsValid() ? EdMode->CurrentToolTarget.LandscapeInfo->GetLandscapeProxy() : nullptr;
-		if (Proxy)
+		if (const ALandscapeProxy* Proxy = EdMode->CurrentToolTarget.LandscapeInfo.IsValid() ? EdMode->CurrentToolTarget.LandscapeInfo->GetLandscapeProxy() : nullptr)
 		{
 			const float ScaleXY = static_cast<float>(FMath::Abs(EdMode->CurrentToolTarget.LandscapeInfo->DrawScale.X));
 
diff --git a/Engine/Source/Editor/LandscapeEditor/Private/LandscapeEdModeMirrorTool.cpp b/Engine/Source/Editor/LandscapeEditor/Private/LandscapeEdModeMirrorTool.cpp
--- a/Engine/Source/Editor/LandscapeEditor/Private/LandscapeEdModeMirrorTool.cpp
+++ b/Engine/Source/Editor/LandscapeEditor/Private/LandscapeEdModeMirrorTool.cpp
@@ -99,9 +99,9 @@ public:
 
 	virtual bool InputDelta(FEditorViewportClient* InViewportClient, FViewport* InViewport, FVector& InDrag, FRotator& InRot, FVector& InScale) override
 	{
-		if (InViewportClient->GetCurrentWidgetAxis() != EAxisList::None)
+		const ALandscapeProxy* LandscapeProxy = EdMode->CurrentToolTarget.LandscapeInfo.IsValid() ? EdMode->CurrentToolTarget.LandscapeInfo->GetLandscapeProxy() : nullptr;
+		if ((LandscapeProxy != nullptr) && (InViewportClient->GetCurrentWidgetAxis() != EAxisList::None))
 		{
-			const ALandscapeProxy* LandscapeProxy = EdMode->CurrentToolTarget.LandscapeInfo->GetLandscapeProxy();
 			const FTransform LandscapeToWorld = LandscapeProxy->LandscapeActorToWorld();
 
 			EdMode->UISettings->MirrorPoint += FVector2D(LandscapeToWorld.InverseTransformVector(InDrag));
@@ -266,8 +266,7 @@ public:
 
 	virtual FMatrix GetWidgetRotation() const override
 	{
-		const ALandscapeProxy* const LandscapeProxy = EdMode->CurrentToolTarget.LandscapeInfo->GetLandscapeProxy();
-		if (LandscapeProxy)
+		if (const ALandscapeProxy* LandscapeProxy = EdMode->CurrentToolTarget.LandscapeInfo.IsValid() ? EdMode->CurrentToolTarget.LandscapeInfo->GetLandscapeProxy() : nullptr)
 		{
 			const FTransform LandscapeToWorld = LandscapeProxy->LandscapeActorToWorld();
 
diff --git a/Engine/Source/Editor/LandscapeEditor/Private/LandscapeEdModePaintTools.cpp b/Engine/Source/Editor/LandscapeEditor/Private/LandscapeEdModePaintTools.cpp
--- a/Engine/Source/Editor/LandscapeEditor/Private/LandscapeEdModePaintTools.cpp
+++ b/Engine/Source/Editor/LandscapeEditor/Private/LandscapeEdModePaintTools.cpp
@@ -1189,15 +1189,16 @@ public:
 
 		Super::Tick(ViewportClient, DeltaTime);
 
-		if (HeightmapFlattenPreviewComponent != nullptr)
+		const ALandscapeProxy* LandscapeProxy = this->EdMode->CurrentToolTarget.LandscapeInfo.IsValid() ? this->EdMode->CurrentToolTarget.LandscapeInfo->GetLandscapeProxy() : nullptr;
+		if ((LandscapeProxy != nullptr) && (HeightmapFlattenPreviewComponent != nullptr))
 		{
 			bool bShowGrid = this->EdMode->UISettings->bUseFlattenTarget && this->EdMode->UISettings->bShowFlattenTargetPreview;
 			HeightmapFlattenPreviewComponent->SetVisibility(bShowGrid);
 
 			// Always adjust the preview component's Z so that we react to the FlattenTarget changes in the details panel, which greatly helps eyeballing the adjustment : 
-			const FTransform LocalToWorld = this->EdMode->CurrentToolTarget.LandscapeInfo->GetLandscapeProxy()->ActorToWorld();
+			const FTransform LocalToWorld = LandscapeProxy->ActorToWorld();
 			FVector Origin = HeightmapFlattenPreviewComponent->GetRelativeLocation();
-			Origin.Z = ConvertFlattenTargetToLocalZ(/*bInReturnPreviewValueIfActive = */true);
+			Origin.Z = ConvertFlattenTargetToLocalZ(LandscapeProxy, /*bInReturnPreviewValueIfActive = */true);
 			HeightmapFlattenPreviewComponent->SetRelativeLocation(Origin, /*bSweep = */false);
 		}
 	}
@@ -1226,7 +1227,8 @@ public:
 
 		if (ViewportClient->IsLevelEditorClient())
 		{
-			if (HeightmapFlattenPreviewComponent != nullptr)
+			const ALandscapeProxy* LandscapeProxy = this->EdMode->CurrentToolTarget.LandscapeInfo.IsValid() ? this->EdMode->CurrentToolTarget.LandscapeInfo->GetLandscapeProxy() : nullptr;
+			if ((LandscapeProxy != nullptr) && (HeightmapFlattenPreviewComponent != nullptr))
 			{
 				FVector MousePosition;
 				if (this->EdMode->LandscapeMouseTrace((FEditorViewportClient*)ViewportClient, x, y, MousePosition))
@@ -1235,7 +1237,7 @@ public:
 					FVector Origin;
 					Origin.X = FMath::RoundToDouble(MousePosition.X);
 					Origin.Y = FMath::RoundToDouble(MousePosition.Y);
-					Origin.Z = ConvertFlattenTargetToLocalZ(/*bInReturnPreviewValueIfActive = */true);
+					Origin.Z = ConvertFlattenTargetToLocalZ(LandscapeProxy, /*bInReturnPreviewValueIfActive = */true);
 					HeightmapFlattenPreviewComponent->SetRelativeLocation(Origin, false);
 
 					// Clamp the value to the height map
@@ -1261,9 +1263,9 @@ public:
 
 		this->EdMode->UISettings->bFlattenEyeDropperModeMousingOverViewport = false;
 
-		if (ToolTarget::TargetType == ELandscapeToolTargetType::Heightmap)
+		ALandscapeProxy* LandscapeProxy = this->EdMode->CurrentToolTarget.LandscapeInfo.IsValid() ? this->EdMode->CurrentToolTarget.LandscapeInfo->GetLandscapeProxy() : nullptr;
+		if ((LandscapeProxy != nullptr) && (ToolTarget::TargetType == ELandscapeToolTargetType::Heightmap))
 		{
-			ALandscapeProxy* LandscapeProxy = this->EdMode->CurrentToolTarget.LandscapeInfo->GetLandscapeProxy();
 			HeightmapFlattenPreviewComponent = NewObject<UStaticMeshComponent>(LandscapeProxy, NAME_None, RF_Transient);
 			HeightmapFlattenPreviewComponent->SetStaticMesh(HeightmapFlattenPlaneMesh);
 			HeightmapFlattenPreviewComponent->SetCanEverAffectNavigation(false);
@@ -1274,9 +1276,9 @@ public:
 			HeightmapFlattenPreviewComponent->SetVisibility(bShowGrid);
 
 			// Try to set a sane initial location for the preview grid
-			const FTransform LocalToWorld = this->EdMode->CurrentToolTarget.LandscapeInfo->GetLandscapeProxy()->GetRootComponent()->GetComponentToWorld();
+			const FTransform LocalToWorld = LandscapeProxy->GetRootComponent()->GetComponentToWorld();
 			FVector Origin = FVector::ZeroVector;
-			Origin.Z = ConvertFlattenTargetToLocalZ(/*bInReturnPreviewValueIfActive = */false);
+			Origin.Z = ConvertFlattenTargetToLocalZ(LandscapeProxy, /*bInReturnPreviewValueIfActive = */false);
 			HeightmapFlattenPreviewComponent->SetRelativeLocation(Origin, false);
 		}
 	}
@@ -1298,9 +1300,9 @@ public:
 	}
 
 private:
-	double ConvertFlattenTargetToLocalZ(bool bInReturnPreviewValueIfActive) const
+	double ConvertFlattenTargetToLocalZ(const ALandscapeProxy* InLandscapeProxy, bool bInReturnPreviewValueIfActive) const
 	{
-		const FTransform LocalToWorld = this->EdMode->CurrentToolTarget.LandscapeInfo->GetLandscapeProxy()->ActorToWorld();
+		const FTransform LocalToWorld = InLandscapeProxy->ActorToWorld();
 		return (FMath::RoundToDouble((this->EdMode->UISettings->GetFlattenTarget(bInReturnPreviewValueIfActive) - LocalToWorld.GetTranslation().Z) / LocalToWorld.GetScale3D().Z * LANDSCAPE_INV_ZSCALE) - 0.1f) * LANDSCAPE_ZSCALE;
 	}
 };
diff --git a/Engine/Source/Editor/LandscapeEditor/Private/LandscapeEdModeRampTool.cpp b/Engine/Source/Editor/LandscapeEditor/Private/LandscapeEdModeRampTool.cpp
--- a/Engine/Source/Editor/LandscapeEditor/Private/LandscapeEdModeRampTool.cpp
+++ b/Engine/Source/Editor/LandscapeEditor/Private/LandscapeEdModeRampTool.cpp
@@ -312,10 +312,12 @@ public:
 	{
 		if (SelectedPoint != INDEX_NONE && InViewportClient->GetCurrentWidgetAxis() != EAxisList::None)
 		{
-			const ALandscapeProxy* LandscapeProxy = EdMode->CurrentToolTarget.LandscapeInfo->GetLandscapeProxy();
-			const FTransform LandscapeToWorld = LandscapeProxy->LandscapeActorToWorld();
+			if (const ALandscapeProxy* LandscapeProxy = EdMode->CurrentToolTarget.LandscapeInfo.IsValid() ? EdMode->CurrentToolTarget.LandscapeInfo->GetLandscapeProxy() : nullptr)
+			{
+				const FTransform LandscapeToWorld = LandscapeProxy->LandscapeActorToWorld();
 
-			Points[SelectedPoint] += LandscapeToWorld.InverseTransformVector(InDrag);
+				Points[SelectedPoint] += LandscapeToWorld.InverseTransformVector(InDrag);
+			}
 
 			return true;
 		}
@@ -326,8 +328,8 @@ public:
 	virtual void Render(const FSceneView* View, FViewport* Viewport, FPrimitiveDrawInterface* PDI) override
 	{
 		// The editor can try to render the tool before the UpdateLandscapeEditorData command runs and the landscape editor realizes that the landscape has been hidden/deleted
-		const ALandscapeProxy* LandscapeProxy = EdMode->CurrentToolTarget.LandscapeInfo->GetLandscapeProxy();
-		if (LandscapeProxy && NumPoints > 0)
+		const ALandscapeProxy* LandscapeProxy = EdMode->CurrentToolTarget.LandscapeInfo.IsValid() ? EdMode->CurrentToolTarget.LandscapeInfo->GetLandscapeProxy() : nullptr;
+		if ((LandscapeProxy != nullptr) && (NumPoints > 0))
 		{
 			const FTransform LandscapeToWorld = LandscapeProxy->LandscapeActorToWorld();
 
@@ -357,8 +359,8 @@ public:
 					SpriteTexture->GetResource(),
 					SpriteColor,
 					SDPG_Foreground,
-					0, static_cast<float>(SpriteTexture->GetResource()->GetSizeX()),
-					0, static_cast<float>(SpriteTexture->GetResource()->GetSizeY()),
+					/*U = */0.f, /*UL = */0.f, 
+					/*V = */0.f, /*VL = */0.f,
 					SE_BLEND_Masked);
 			}
 			PDI->SetHitProxy(NULL);
@@ -423,11 +425,7 @@ public:
 		if (SelectedPoint != INDEX_NONE)
 		{
 			// The editor can try to render the transform widget before the landscape editor ticks and realizes that the landscape has been hidden/deleted
-			const ALandscapeProxy* LandscapeProxy = EdMode->CurrentToolTarget.LandscapeInfo->GetLandscapeProxy();
-			if (LandscapeProxy)
-			{
-				return true;
-			}
+			return EdMode->CurrentToolTarget.LandscapeInfo.IsValid() && (EdMode->CurrentToolTarget.LandscapeInfo->GetLandscapeProxy() != nullptr);
 		}
 
 		return false;
diff --git a/Engine/Source/Editor/LandscapeEditor/Private/LandscapeEdModeSplineTools.cpp b/Engine/Source/Editor/LandscapeEditor/Private/LandscapeEdModeSplineTools.cpp
--- a/Engine/Source/Editor/LandscapeEditor/Private/LandscapeEdModeSplineTools.cpp
+++ b/Engine/Source/Editor/LandscapeEditor/Private/LandscapeEdModeSplineTools.cpp
@@ -1890,8 +1890,7 @@ public:
 	virtual void Render(const FSceneView* View, FViewport* Viewport, FPrimitiveDrawInterface* PDI) override
 	{
 		// The editor can try to render the tool before the UpdateLandscapeEditorData command runs and the landscape editor realizes that the landscape has been hidden/deleted
-		const ALandscapeProxy* LandscapeProxy = EdMode->CurrentToolTarget.LandscapeInfo->GetLandscapeProxy();
-		if (LandscapeProxy)
+		if (const ALandscapeProxy* LandscapeProxy = EdMode->CurrentToolTarget.LandscapeInfo.IsValid() ? EdMode->CurrentToolTarget.LandscapeInfo->GetLandscapeProxy() : nullptr)
 		{
 			for (ULandscapeSplineControlPoint* ControlPoint : SelectedSplineControlPoints)
 			{
@@ -1967,11 +1966,7 @@ public:
 		if (SelectedSplineControlPoints.Num() > 0 || DraggingTangent_Segment)
 		{
 			// The editor can try to render the transform widget before the landscape editor ticks and realizes that the landscape has been hidden/deleted
-			const ALandscapeProxy* LandscapeProxy = EdMode->CurrentToolTarget.LandscapeInfo->GetLandscapeProxy();
-			if (LandscapeProxy)
-			{
-				return true;
-			}
+			return EdMode->CurrentToolTarget.LandscapeInfo.IsValid() && (EdMode->CurrentToolTarget.LandscapeInfo->GetLandscapeProxy() != nullptr);
 		}
 
 		return false;
@@ -2002,8 +1997,7 @@ public:
 
 	virtual FVector GetWidgetLocation() const override
 	{
-		const ALandscapeProxy* LandscapeProxy = EdMode->CurrentToolTarget.LandscapeInfo->GetLandscapeProxy();
-		if (LandscapeProxy)
+		if (const ALandscapeProxy* LandscapeProxy = EdMode->CurrentToolTarget.LandscapeInfo.IsValid() ? EdMode->CurrentToolTarget.LandscapeInfo->GetLandscapeProxy() : nullptr)
 		{
 			if (DraggingTangent_Segment)
 			{
@@ -2030,7 +2024,7 @@ public:
 
 	virtual FMatrix GetWidgetRotation() const override
 	{
-		const ALandscapeProxy* LandscapeProxy = EdMode->CurrentToolTarget.LandscapeInfo->GetLandscapeProxy();
+		const ALandscapeProxy* LandscapeProxy = EdMode->CurrentToolTarget.LandscapeInfo.IsValid() ? EdMode->CurrentToolTarget.LandscapeInfo->GetLandscapeProxy() : nullptr;
 		if (LandscapeProxy)
 		{
 			if (DraggingTangent_Segment)
