Create New Landscape - Fix to revert both Scale X and Y values with a single Undo Command

diff --git a/Engine/Source/Editor/LandscapeEditor/Private/LandscapeEditorDetailCustomization_NewLandscape.cpp b/Engine/Source/Editor/LandscapeEditor/Private/LandscapeEditorDetailCustomization_NewLandscape.cpp
--- a/Engine/Source/Editor/LandscapeEditor/Private/LandscapeEditorDetailCustomization_NewLandscape.cpp
+++ b/Engine/Source/Editor/LandscapeEditor/Private/LandscapeEditorDetailCustomization_NewLandscape.cpp
@@ -659,6 +659,7 @@ FText FLandscapeEditorDetailCustomization_NewLandscape::GetOverallResolutionTool
 
 void FLandscapeEditorDetailCustomization_NewLandscape::SetScale(FVector::FReal NewValue, ETextCommit::Type, TSharedRef<IPropertyHandle> PropertyHandle)
 {
+	FScopedTransaction Transaction(LOCTEXT("ChangeScale_Transaction", "Change Landscape Scale"));
 	FVector::FReal OldValue = 0;
 	PropertyHandle->GetValue(OldValue);
 
