Fix landscape weightmap import comparison vs heightmap resolution.

diff --git a/Engine/Source/Editor/LandscapeEditor/Private/LandscapeTiledImage.cpp b/Engine/Source/Editor/LandscapeEditor/Private/LandscapeTiledImage.cpp
--- a/Engine/Source/Editor/LandscapeEditor/Private/LandscapeTiledImage.cpp
+++ b/Engine/Source/Editor/LandscapeEditor/Private/LandscapeTiledImage.cpp
@@ -118,10 +118,15 @@ FLandscapeFileInfo FLandscapeTiledImage::Load(const TCHAR* Filename)
 	{
 		FLandscapeImageDataRef ImageData;
 		FLandscapeImageFileCache& LandscapeImageFileCache = FModuleManager::GetModuleChecked<ILandscapeEditorModule>("LandscapeEditor").GetImageFileCache();
-		Result = LandscapeImageFileCache.FindImage<T>(*Tile.Value, ImageData);
-		if (Result.ResultCode == ELandscapeImportResult::Error)
+		FLandscapeFileInfo TileResult = LandscapeImageFileCache.FindImage<T>(*Tile.Value, ImageData);
+		if (TileResult.ResultCode == ELandscapeImportResult::Error)
 		{
-			return Result;
+			return TileResult;
+		}
+		else if (TileResult.ResultCode == ELandscapeImportResult::Warning)
+		{
+			Result.ResultCode = TileResult.ResultCode;
+			Result.ErrorMessage = TileResult.ErrorMessage;
 		}
 
 		SizeInTiles.X = FMath::Max(SizeInTiles.X, Tile.Key.X + 1);
