Added Landscape Spline linear path navigation and selection controls + New button to toggle spline selection mode between control points and segments

diff --git a/Engine/Source/Editor/EditorStyle/Private/StarshipStyle.cpp b/Engine/Source/Editor/EditorStyle/Private/StarshipStyle.cpp
--- a/Engine/Source/Editor/EditorStyle/Private/StarshipStyle.cpp
+++ b/Engine/Source/Editor/EditorStyle/Private/StarshipStyle.cpp
@@ -2004,44 +2004,57 @@ void FStarshipEditorStyle::FStyle::SetupLevelGeneralStyles()
 			.SetNormal(IMAGE_BRUSH_SVG("Starship/Splines/Spline_SelectFirst", Icon20x20))
 			.SetHovered(IMAGE_BRUSH_SVG("Starship/Splines/Spline_SelectFirst", Icon20x20, SelectionColor))
 			.SetPressed(IMAGE_BRUSH_SVG("Starship/Splines/Spline_SelectFirst", Icon20x20, SelectionColor_Pressed))
+			.SetDisabled(IMAGE_BRUSH_SVG("Starship/Splines/Spline_SelectFirst", Icon20x20, SelectionColor_Inactive))
 		);
 
 		Set("SplineComponentDetails.AddPrev", FButtonStyle(Button)
 			.SetNormal(IMAGE_BRUSH_SVG("Starship/Splines/Spline_AddPrevious", Icon20x20))
 			.SetHovered(IMAGE_BRUSH_SVG("Starship/Splines/Spline_AddPrevious", Icon20x20, SelectionColor))
 			.SetPressed(IMAGE_BRUSH_SVG("Starship/Splines/Spline_AddPrevious", Icon20x20, SelectionColor_Pressed))
+			.SetDisabled(IMAGE_BRUSH_SVG("Starship/Splines/Spline_AddPrevious", Icon20x20, SelectionColor_Inactive))
 		);
 
 		Set("SplineComponentDetails.SelectPrev", FButtonStyle(Button)
 			.SetNormal(IMAGE_BRUSH_SVG("Starship/Splines/Spline_SelectPrevious", Icon20x20))
 			.SetHovered(IMAGE_BRUSH_SVG("Starship/Splines/Spline_SelectPrevious", Icon20x20, SelectionColor))
 			.SetPressed(IMAGE_BRUSH_SVG("Starship/Splines/Spline_SelectPrevious", Icon20x20, SelectionColor_Pressed))
+			.SetDisabled(IMAGE_BRUSH_SVG("Starship/Splines/Spline_SelectPrevious", Icon20x20, SelectionColor_Inactive))
 		);
 
 		Set("SplineComponentDetails.SelectAll", FButtonStyle(Button)
 			.SetNormal(IMAGE_BRUSH_SVG("Starship/Splines/Spline_SelectAll", Icon20x20))
 			.SetHovered(IMAGE_BRUSH_SVG("Starship/Splines/Spline_SelectAll", Icon20x20, SelectionColor))
 			.SetPressed(IMAGE_BRUSH_SVG("Starship/Splines/Spline_SelectAll", Icon20x20, SelectionColor_Pressed))
+			.SetDisabled(IMAGE_BRUSH_SVG("Starship/Splines/Spline_SelectAll", Icon20x20, SelectionColor_Inactive))
 		);
 
 		Set("SplineComponentDetails.SelectNext", FButtonStyle(Button)
 			.SetNormal(IMAGE_BRUSH_SVG("Starship/Splines/Spline_SelectNext", Icon20x20))
 			.SetHovered(IMAGE_BRUSH_SVG("Starship/Splines/Spline_SelectNext", Icon20x20, SelectionColor))
 			.SetPressed(IMAGE_BRUSH_SVG("Starship/Splines/Spline_SelectNext", Icon20x20, SelectionColor_Pressed))
+			.SetDisabled(IMAGE_BRUSH_SVG("Starship/Splines/Spline_SelectNext", Icon20x20, SelectionColor_Inactive))
 		);
 
 		Set("SplineComponentDetails.AddNext", FButtonStyle(Button)
 			.SetNormal(IMAGE_BRUSH_SVG("Starship/Splines/Spline_AddNext", Icon20x20))
 			.SetHovered(IMAGE_BRUSH_SVG("Starship/Splines/Spline_AddNext", Icon20x20, SelectionColor))
 			.SetPressed(IMAGE_BRUSH_SVG("Starship/Splines/Spline_AddNext", Icon20x20, SelectionColor_Pressed))
+			.SetDisabled(IMAGE_BRUSH_SVG("Starship/Splines/Spline_AddNext", Icon20x20, SelectionColor_Inactive))
 		);
 
 		Set("SplineComponentDetails.SelectLast", FButtonStyle(Button)
 			.SetNormal(IMAGE_BRUSH_SVG("Starship/Splines/Spline_SelectLast", Icon20x20))
 			.SetHovered(IMAGE_BRUSH_SVG("Starship/Splines/Spline_SelectLast", Icon20x20, SelectionColor))
 			.SetPressed(IMAGE_BRUSH_SVG("Starship/Splines/Spline_SelectLast", Icon20x20, SelectionColor_Pressed))
+			.SetDisabled(IMAGE_BRUSH_SVG("Starship/Splines/Spline_SelectLast", Icon20x20, SelectionColor_Inactive))
 		);
 
+		Set("SplineComponentDetails.Convert", FButtonStyle(Button)
+			.SetNormal(IMAGE_BRUSH_SVG("Starship/Splines/Spline_ToggleSelectMode", Icon20x20))
+			.SetHovered(IMAGE_BRUSH_SVG("Starship/Splines/Spline_ToggleSelectMode", Icon20x20, SelectionColor))
+			.SetPressed(IMAGE_BRUSH_SVG("Starship/Splines/Spline_ToggleSelectMode", Icon20x20, SelectionColor_Pressed))
+			.SetDisabled(IMAGE_BRUSH_SVG("Starship/Splines/Spline_ToggleSelectMode", Icon20x20, SelectionColor_Inactive))
+		);
 	}
 }
 
diff --git a/Engine/Source/Editor/LandscapeEditor/Private/LandscapeEdMode.h b/Engine/Source/Editor/LandscapeEditor/Private/LandscapeEdMode.h
--- a/Engine/Source/Editor/LandscapeEditor/Private/LandscapeEdMode.h
+++ b/Engine/Source/Editor/LandscapeEditor/Private/LandscapeEdMode.h
@@ -275,6 +275,22 @@ enum class ELandscapeEditingState : uint8
 	NoLandscape,
 };
 
+enum class ESplineNavigationFlags : uint8
+{
+	None = 0,
+	DirectionForward = 1 << 0,
+	DirectionBackward = 1 << 1,
+
+	SegmentSelectModeEnabled = 1 << 2,
+	ControlPointSelectModeEnabled = 1 << 3,
+
+	AddToSelection = 1 << 4,
+
+	SelectModeMask = SegmentSelectModeEnabled | ControlPointSelectModeEnabled,
+	DirectionMask = DirectionForward | DirectionBackward
+};
+ENUM_CLASS_FLAGS(ESplineNavigationFlags);
+
 /**
  * Landscape editor mode
  */
@@ -313,12 +329,18 @@ public:
 	FLandscapeToolSplines* SplinesTool;
 	void ShowSplineProperties();
 	bool HasSelectedSplineSegments() const;
+	bool HasAdjacentLinearSplineConnection(ESplineNavigationFlags Flags) const;
+	bool HasEndLinearSplineConnection(ESplineNavigationFlags Flags) const;
 	void FlipSelectedSplineSegments();
 	void GetSelectedSplineOwners(TSet<AActor*>& SelectedSplineOwners) const;
 	virtual void SelectAllSplineControlPoints();
 	virtual void SelectAllSplineSegments();
 	virtual void SelectAllConnectedSplineControlPoints();
 	virtual void SelectAllConnectedSplineSegments();
+	virtual void SelectAdjacentLinearSplineElement(ESplineNavigationFlags Flags) const;
+	virtual void SelectEndLinearSplineElement(ESplineNavigationFlags Flags) const;
+	virtual void SelectSplineControlPointsFromCurrentSegmentSelection() const;
+	virtual void SelectSplineSegmentsFromCurrentControlPointSelection() const;
 	virtual void SplineMoveToCurrentLevel();
 	virtual bool CanMoveSplineToCurrentLevel() const;
 	virtual void UpdateSplineMeshLevels();
diff --git a/Engine/Source/Editor/LandscapeEditor/Private/LandscapeEdModeSplineTools.cpp b/Engine/Source/Editor/LandscapeEditor/Private/LandscapeEdModeSplineTools.cpp
--- a/Engine/Source/Editor/LandscapeEditor/Private/LandscapeEdModeSplineTools.cpp
+++ b/Engine/Source/Editor/LandscapeEditor/Private/LandscapeEdModeSplineTools.cpp
@@ -51,6 +51,11 @@ public:
 		, LandscapeInfo(nullptr)
 		, SelectedSplineControlPoints()
 		, SelectedSplineSegments()
+		, LastSelectedControlPoint(nullptr)
+		, LastSelectedSegment(nullptr)
+		, LinearControlPoints()
+		, LinearSegments()
+		, ControlPointWidgetMode()
 		, DraggingTangent_Segment(nullptr)
 		, DraggingTangent_Length(0.0f)
 		, DraggingTangent_CacheCoordSpace(ECoordSystem::COORD_None)
@@ -161,7 +166,15 @@ public:
 
 	void SelectControlPoint(ULandscapeSplineControlPoint* ControlPoint)
 	{
+		SelectNavigationControlPoint(ControlPoint);
+
 		checkSlow(!ControlPoint->IsSplineSelected());
+
+		// When switching from segment mode, restore the last used control point mode
+		if (SelectedSplineSegments.Num() > 0 && ControlPointWidgetMode.IsSet())
+		{
+			GLevelEditorModeTools().SetWidgetMode(ControlPointWidgetMode.GetValue());
+		}
 		SelectedSplineControlPoints.Add(ControlPoint);
 		ControlPoint->Modify(false);
 		ControlPoint->SetSplineSelected(true);
@@ -169,7 +182,15 @@ public:
 
 	void SelectSegment(ULandscapeSplineSegment* Segment)
 	{
+		SelectNavigationSegment(Segment);
+
 		checkSlow(!Segment->IsSplineSelected());
+
+		// Before entering segment mode, save the control point mode for persistence
+		if (SelectedSplineSegments.Num() == 0)
+		{
+			ControlPointWidgetMode = GLevelEditorModeTools().GetWidgetMode();
+		}
 		SelectedSplineSegments.Add(Segment);
 		Segment->Modify(false);
 		Segment->SetSplineSelected(true);
@@ -280,6 +301,240 @@ public:
 		}
 	}
 
+	ULandscapeSplineControlPoint* GetLinearEndControlPoint(ULandscapeSplineControlPoint* SelectedControlPoint)
+	{
+		TSet<ULandscapeSplineControlPoint*> ControlPointsVisited;
+		TArray<ULandscapeSplineControlPoint*> ControlPointsToProcess;
+		ControlPointsToProcess.Add(SelectedControlPoint);
+
+		// Reset the linear path before it gets rebuilt below
+		LinearControlPoints.Empty();
+
+		// Control points with zero/multi node connections are not part of a valid linear path
+		if (SelectedControlPoint->ConnectedSegments.IsEmpty() || SelectedControlPoint->ConnectedSegments.Num() > 2)
+		{
+			LinearControlPoints.Add(SelectedControlPoint);
+			return SelectedControlPoint;
+		}
+
+		while (ControlPointsToProcess.Num() > 0)
+		{
+			ULandscapeSplineControlPoint* ControlPoint = ControlPointsToProcess.Pop();
+			ControlPointsVisited.Add(ControlPoint);
+
+			// Store nodes in order of visiting to create the linear path
+			if (!LinearControlPoints.Contains(ControlPoint))
+			{
+				LinearControlPoints.Add(ControlPoint);
+			}
+
+			// Valid end points only have one connection
+			if (ControlPoint->ConnectedSegments.Num() == 1 && ControlPoint != SelectedControlPoint)
+			{
+				return ControlPoint;
+			}
+
+			for (const FLandscapeSplineConnection& Connection : ControlPoint->ConnectedSegments)
+			{
+				ULandscapeSplineControlPoint* OtherEnd = Connection.GetFarConnection().ControlPoint;
+				check(OtherEnd != nullptr);
+
+				if (OtherEnd->ConnectedSegments.IsEmpty() || OtherEnd->ConnectedSegments.Num() > 2)
+				{
+					continue;
+				}
+
+				if (!ControlPointsVisited.Contains(OtherEnd))
+				{
+					ControlPointsToProcess.Add(OtherEnd);
+				}
+			}
+		}
+
+		return SelectedControlPoint;
+	}
+
+	void ResetNavigationPath()
+	{
+		// Clear cached path
+		LinearControlPoints.Empty();
+		LinearSegments.Empty();
+
+		LastSelectedControlPoint = nullptr;
+		LastSelectedSegment = nullptr;
+	}
+
+	void BuildLinearPathFromLastSelectedPoint()
+	{
+		LinearControlPoints.Empty();
+		LinearSegments.Empty();
+
+		ULandscapeSplineControlPoint* StartControlPoint = GetLinearEndControlPoint(LastSelectedControlPoint);
+		ULandscapeSplineControlPoint* EndControlPoint = GetLinearEndControlPoint(StartControlPoint);
+
+		for (const ULandscapeSplineControlPoint* Point : LinearControlPoints)
+		{
+			for (const FLandscapeSplineConnection& Connection : Point->ConnectedSegments)
+			{
+				if (!LinearSegments.Contains(Connection.Segment))
+				{
+					LinearSegments.Add(Connection.Segment);
+				}
+			}
+		}
+	}
+
+	void SelectNavigationSegment(ULandscapeSplineSegment* Segment)
+	{
+		LastSelectedSegment = Segment;
+		LastSelectedControlPoint = Segment->Connections[0].ControlPoint;
+
+		if (!LinearSegments.Contains(Segment))
+		{
+			BuildLinearPathFromLastSelectedPoint();
+		}
+	}
+
+	void SelectNavigationControlPoint(ULandscapeSplineControlPoint* ControlPoint)
+	{
+		LastSelectedControlPoint = ControlPoint;
+
+		if (!LinearControlPoints.Contains(ControlPoint))
+		{
+			BuildLinearPathFromLastSelectedPoint();
+		}
+	}
+
+	bool IsSelectionValidForNavigation() const
+	{
+		// If any elements in selection are part of different linear paths, disable navigation
+		for (ULandscapeSplineControlPoint* ControlPoint : SelectedSplineControlPoints)
+		{
+			if (!LinearControlPoints.Contains(ControlPoint))
+			{
+				return false;
+			}
+		}
+
+		for (ULandscapeSplineSegment* Segment : SelectedSplineSegments)
+		{
+			if (!LinearSegments.Contains(Segment))
+			{
+				return false;
+			}
+		}
+
+		return true;
+	}
+
+	ULandscapeSplineSegment* GetEndSegmentInLinearPath(ESplineNavigationFlags Flags) const
+	{
+		check(EnumHasAllFlags(Flags, ESplineNavigationFlags::SegmentSelectModeEnabled));
+		check(FMath::CountBits(static_cast<uint8>(Flags & ESplineNavigationFlags::DirectionMask)) == 1);
+
+		if (LinearSegments.IsEmpty())
+		{
+			return nullptr;
+		}
+
+		if (EnumHasAllFlags(Flags, ESplineNavigationFlags::DirectionForward))
+		{
+			return LinearSegments.Top();
+		}
+
+		return LinearSegments[0];
+	}
+
+	ULandscapeSplineControlPoint* GetEndControlPointInLinearPath(ESplineNavigationFlags Flags) const
+	{
+		check(EnumHasAllFlags(Flags, ESplineNavigationFlags::ControlPointSelectModeEnabled));
+		check(FMath::CountBits(static_cast<uint8>(Flags & ESplineNavigationFlags::SelectModeMask)) == 1);
+
+		if (LinearControlPoints.IsEmpty())
+		{
+			return nullptr;
+		}
+
+		if (EnumHasAllFlags(Flags, ESplineNavigationFlags::DirectionForward))
+		{
+			return LinearControlPoints.Top();
+		}
+
+		return LinearControlPoints[0];
+	}
+
+	bool HasAdjacentSegmentInLinearPath(ESplineNavigationFlags Flags) const
+	{
+		check(EnumHasAllFlags(Flags, ESplineNavigationFlags::SegmentSelectModeEnabled));
+		check(FMath::CountBits(static_cast<uint8>(Flags & ESplineNavigationFlags::SelectModeMask)) == 1);
+
+		const ULandscapeSplineSegment* AdjacentSegment = GetAdjacentSegmentInLinearPath(Flags);
+
+		if (!AdjacentSegment || !IsSelectionValidForNavigation())
+		{
+			return false;
+		}
+
+		if (EnumHasAllFlags(Flags, ESplineNavigationFlags::AddToSelection))
+		{
+			return !AdjacentSegment->IsSplineSelected();
+		}
+
+		return true;
+	}
+
+	bool HasAdjacentControlPointInLinearPath(ESplineNavigationFlags Flags) const
+	{
+		check(EnumHasAllFlags(Flags, ESplineNavigationFlags::ControlPointSelectModeEnabled));
+		check(FMath::CountBits(static_cast<uint8>(Flags & ESplineNavigationFlags::SelectModeMask)) == 1);
+
+		const ULandscapeSplineControlPoint* AdjacentPoint = GetAdjacentControlPointInPath(Flags);
+
+		if (!AdjacentPoint || !IsSelectionValidForNavigation())
+		{
+			return false;
+		}
+
+		if (EnumHasAllFlags(Flags, ESplineNavigationFlags::AddToSelection))
+		{
+			return !AdjacentPoint->IsSplineSelected();
+		}
+
+		return true;
+	}
+
+	ULandscapeSplineSegment* GetAdjacentSegmentInLinearPath(ESplineNavigationFlags Flags) const
+	{
+		check(EnumHasAllFlags(Flags, ESplineNavigationFlags::SegmentSelectModeEnabled));
+		check(FMath::CountBits(static_cast<uint8>(Flags & ESplineNavigationFlags::SelectModeMask)) == 1);
+
+		const int32 LastSelectedSegmentIndex = LinearSegments.Find(LastSelectedSegment);
+		const int32 NextIndexOffset = EnumHasAllFlags(Flags, ESplineNavigationFlags::DirectionForward) ? 1 : -1;
+		check(LastSelectedSegmentIndex != INDEX_NONE); // last selected should always be part of the linear path
+
+		if (LinearSegments.IsValidIndex(LastSelectedSegmentIndex + NextIndexOffset))
+		{
+			return LinearSegments[LastSelectedSegmentIndex + NextIndexOffset];
+		}
+		return nullptr;
+	}
+
+	ULandscapeSplineControlPoint* GetAdjacentControlPointInPath(ESplineNavigationFlags Flags) const
+	{
+		check(EnumHasAllFlags(Flags, ESplineNavigationFlags::ControlPointSelectModeEnabled));
+		check(FMath::CountBits(static_cast<uint8>(Flags & ESplineNavigationFlags::SelectModeMask)) == 1);
+
+		const int32 LastSelectedControlPointIndex = LinearControlPoints.Find(LastSelectedControlPoint);
+		const int32 NextIndexOffset = EnumHasAllFlags(Flags, ESplineNavigationFlags::DirectionForward) ? 1 : -1;
+		check(LastSelectedControlPointIndex != INDEX_NONE); // last selected should always be part of the linear path
+
+		if (LinearControlPoints.IsValidIndex(LastSelectedControlPointIndex + NextIndexOffset))
+		{
+			return LinearControlPoints[LastSelectedControlPointIndex + NextIndexOffset];
+		}
+		return nullptr;
+	}
+
 	void AddSegment(ULandscapeSplineControlPoint* Start, ULandscapeSplineControlPoint* End, bool bAutoRotateStart, bool bAutoRotateEnd)
 	{
 		FScopedTransaction Transaction(LOCTEXT("LandscapeSpline_AddSegment", "Add Landscape Spline Segment"));
@@ -399,6 +654,10 @@ public:
 		{
 			NewSegment->UpdateSplinePoints();
 		}
+
+		// Adding a segment will change the linear navigation path, reset it
+		ResetNavigationPath();
+		SelectNavigationControlPoint(Start);
 	}
 
 	void FlipSelectedSplineSegments()
@@ -476,6 +735,10 @@ public:
 		ULandscapeSplinesComponent* SplinesComponent = ControlPoint->GetOuterULandscapeSplinesComponent();
 		SplinesComponent->MarkRenderStateDirty();
 
+		// Adding a segment will change the linear navigation path, reset it
+		ResetNavigationPath();
+		SelectNavigationControlPoint(Start);
+
 		return true;
 	}
 
@@ -540,6 +803,7 @@ public:
 			NewControlPoint->UpdateSplinePoints();
 		}
 
+		ResetNavigationPath();
 		ClearSelection();
 		SelectControlPoint(NewControlPoint);
 		UpdatePropertiesWindows();
@@ -584,6 +848,7 @@ public:
 		}
 
 		SplinesComponent->Segments.Remove(ToDelete);
+		ResetNavigationPath();
 
 		ToDelete->Connections[0].ControlPoint->UpdateSplinePoints();
 		ToDelete->Connections[1].ControlPoint->UpdateSplinePoints();
@@ -601,6 +866,7 @@ public:
 
 		ToDelete->Modify();
 		ToDelete->DeleteSplinePoints();
+		ResetNavigationPath();
 
 		if (ToDelete->ConnectedSegments.Num() == 2
 			&& ToDelete->ConnectedSegments[0].Segment != ToDelete->ConnectedSegments[1].Segment)
@@ -2313,6 +2579,17 @@ protected:
 	TSet<TObjectPtr<ULandscapeSplineControlPoint>> SelectedSplineControlPoints;
 	TSet<TObjectPtr<ULandscapeSplineSegment>> SelectedSplineSegments;
 
+	// Last Selected
+	ULandscapeSplineControlPoint* LastSelectedControlPoint;
+	ULandscapeSplineSegment* LastSelectedSegment;
+
+	// Control Points and Segments stored in sequential order
+	TArray<ULandscapeSplineControlPoint*> LinearControlPoints;
+	TArray<ULandscapeSplineSegment*> LinearSegments;
+
+	// Persistent control point mode
+	TOptional<UE::Widget::EWidgetMode> ControlPointWidgetMode;
+
 	TObjectPtr<ULandscapeSplineSegment> DraggingTangent_Segment;
 	float DraggingTangent_Length;
 	ECoordSystem DraggingTangent_CacheCoordSpace;
@@ -2365,6 +2642,47 @@ bool FEdModeLandscape::HasSelectedSplineSegments() const
 	return SplinesTool && (SplinesTool->SelectedSplineSegments.Num() > 0);
 }
 
+bool FEdModeLandscape::HasAdjacentLinearSplineConnection(ESplineNavigationFlags Flags) const
+{
+	if (!SplinesTool)
+	{
+		return false;
+	}
+
+	// check only one selection mode is set
+	check(FMath::CountBits(static_cast<uint8>(Flags & ESplineNavigationFlags::SelectModeMask)) == 1);
+
+	if (EnumHasAllFlags(Flags, ESplineNavigationFlags::SegmentSelectModeEnabled))
+	{
+		return SplinesTool->HasAdjacentSegmentInLinearPath(Flags);
+	}
+	else
+	{
+		return SplinesTool->HasAdjacentControlPointInLinearPath(Flags);
+	}
+}
+
+bool FEdModeLandscape::HasEndLinearSplineConnection(ESplineNavigationFlags Flags) const
+{
+	if (!SplinesTool)
+	{
+		return false;
+	}
+	// check only one selection mode is set
+	check(FMath::CountBits(static_cast<uint8>(Flags & ESplineNavigationFlags::SelectModeMask)) == 1);
+
+	if (EnumHasAllFlags(Flags, ESplineNavigationFlags::SegmentSelectModeEnabled))
+	{
+		const ULandscapeSplineSegment* EndSegment = SplinesTool->GetEndSegmentInLinearPath(Flags);
+		return EndSegment != nullptr && !EndSegment->IsSplineSelected() && SplinesTool->IsSelectionValidForNavigation();
+	}
+	else
+	{
+		const ULandscapeSplineControlPoint* EndControlPoint = SplinesTool->GetEndControlPointInLinearPath(Flags);
+		return EndControlPoint != nullptr && !EndControlPoint->IsSplineSelected() && SplinesTool->IsSelectionValidForNavigation();
+	}
+}
+
 void FEdModeLandscape::FlipSelectedSplineSegments()
 {
 	if (!SplinesTool)
@@ -2449,6 +2767,115 @@ void FEdModeLandscape::SelectAllConnectedSplineSegments()
 	}
 }
 
+void FEdModeLandscape::SelectAdjacentLinearSplineElement(ESplineNavigationFlags Flags) const
+{
+	if (SplinesTool /*&& SplinesTool == CurrentTool*/)
+	{
+		// check only one selection mode is set
+		check(FMath::CountBits(static_cast<uint8>(Flags & ESplineNavigationFlags::SelectModeMask)) == 1);
+
+		if (EnumHasAllFlags(Flags, ESplineNavigationFlags::SegmentSelectModeEnabled))
+		{
+			ULandscapeSplineSegment* AdjacentSegment = SplinesTool->GetAdjacentSegmentInLinearPath(Flags);
+
+			// Clear all other selections if the user only wants the next segment selected
+			if (!EnumHasAllFlags(Flags, ESplineNavigationFlags::AddToSelection))
+			{
+				SplinesTool->ClearSelectedSegments();
+			}
+
+			if (!AdjacentSegment->IsSelected())
+			{
+				SplinesTool->SelectSegment(AdjacentSegment);
+			}
+
+			SplinesTool->ClearSelectedControlPoints();
+		}
+		else 
+		{
+			ULandscapeSplineControlPoint* AdjacentPoint = SplinesTool->GetAdjacentControlPointInPath(Flags);
+
+			// Clear all other selections if the user only wants the next point selected
+			if (!EnumHasAllFlags(Flags, ESplineNavigationFlags::AddToSelection))
+			{
+				SplinesTool->ClearSelectedControlPoints();
+			}
+
+			if (!AdjacentPoint->IsSplineSelected())
+			{
+				SplinesTool->SelectControlPoint(AdjacentPoint);
+			}
+
+			SplinesTool->ClearSelectedSegments();
+		}
+		
+		SplinesTool->UpdatePropertiesWindows();
+		GUnrealEd->RedrawLevelEditingViewports();
+	}
+}
+
+void FEdModeLandscape::SelectEndLinearSplineElement(ESplineNavigationFlags Flags) const
+{
+	if (SplinesTool /*&& SplinesTool == CurrentTool*/)
+	{
+		// check only one selection mode is set
+		check(FMath::CountBits(static_cast<uint8>(Flags & ESplineNavigationFlags::SelectModeMask)) == 1);
+
+		if (EnumHasAllFlags(Flags, ESplineNavigationFlags::SegmentSelectModeEnabled))
+		{
+			ULandscapeSplineSegment* EndSegment = SplinesTool->GetEndSegmentInLinearPath(Flags);
+
+			if (EndSegment != nullptr)
+			{
+				SplinesTool->ClearSelectedSegments();
+				if (!EndSegment->IsSplineSelected())
+				{
+					SplinesTool->SelectSegment(EndSegment);
+				}
+				SplinesTool->ClearSelectedControlPoints();
+			}
+		}
+		else
+		{
+			ULandscapeSplineControlPoint* EndPoint = SplinesTool->GetEndControlPointInLinearPath(Flags);
+			if (EndPoint != nullptr)
+			{
+				SplinesTool->ClearSelectedControlPoints();
+				if (!EndPoint->IsSplineSelected())
+				{
+					SplinesTool->SelectControlPoint(EndPoint);
+				}
+				SplinesTool->ClearSelectedSegments();
+			}
+		}
+
+		SplinesTool->UpdatePropertiesWindows();
+		GUnrealEd->RedrawLevelEditingViewports();
+	}
+}
+
+void FEdModeLandscape::SelectSplineControlPointsFromCurrentSegmentSelection() const
+{
+	if (SplinesTool /*&& SplinesTool == CurrentTool*/)
+	{
+		SplinesTool->SelectAdjacentControlPoints();
+		SplinesTool->ClearSelectedSegments();
+		SplinesTool->UpdatePropertiesWindows();
+		GUnrealEd->RedrawLevelEditingViewports();
+	}
+}
+
+void FEdModeLandscape::SelectSplineSegmentsFromCurrentControlPointSelection() const
+{
+	if (SplinesTool /*&& SplinesTool == CurrentTool*/)
+	{
+		SplinesTool->SelectAdjacentSegments();
+		SplinesTool->ClearSelectedControlPoints();
+		SplinesTool->UpdatePropertiesWindows();
+		GUnrealEd->RedrawLevelEditingViewports();
+	}
+}
+
 void FEdModeLandscape::SplineMoveToCurrentLevel()
 {
 	FScopedTransaction Transaction(LOCTEXT("LandscapeSpline_MoveToCurrentLevel", "Move Landscape Spline to current level"));
diff --git a/Engine/Source/Editor/LandscapeEditor/Private/LandscapeEditorModule.cpp b/Engine/Source/Editor/LandscapeEditor/Private/LandscapeEditorModule.cpp
--- a/Engine/Source/Editor/LandscapeEditor/Private/LandscapeEditorModule.cpp
+++ b/Engine/Source/Editor/LandscapeEditor/Private/LandscapeEditorModule.cpp
@@ -92,6 +92,9 @@ public:
 		PropertyModule.RegisterCustomClassLayout("LandscapeSplineControlPoint", FOnGetDetailCustomizationInstance::CreateStatic(&FLandscapeSplineDetails::MakeInstance));
 		PropertyModule.RegisterCustomClassLayout("LandscapeSplineSegment", FOnGetDetailCustomizationInstance::CreateStatic(&FLandscapeSplineDetails::MakeInstance));
 
+		// register property sections
+		RegisterPropertySectionMappings();
+		
 		GlobalUICommandList = MakeShareable(new FUICommandList);
 		const FLandscapeEditorCommands& LandscapeActions = FLandscapeEditorCommands::Get();
 		GlobalUICommandList->MapAction(LandscapeActions.ViewModeNormal, FExecuteAction::CreateStatic(&ChangeLandscapeViewMode, ELandscapeViewMode::Normal), FCanExecuteAction(), FIsActionChecked::CreateStatic(&IsLandscapeViewModeSelected, ELandscapeViewMode::Normal));
@@ -174,6 +177,9 @@ public:
 		PropertyModule.UnregisterCustomClassLayout("LandscapeSplineControlPoint");
 		PropertyModule.UnregisterCustomClassLayout("LandscapeSplineSegment");
 
+		// unregister property sections
+		UnregisterPropertySectionMappings();
+
 		// remove menu extension
 		FLevelEditorModule& LevelEditorModule = FModuleManager::LoadModuleChecked<FLevelEditorModule>("LevelEditor");
 		GlobalUICommandList = nullptr;
@@ -295,6 +301,11 @@ public:
 	virtual int32 GetOrCreateEditLayer(FName InEditLayerName, ALandscape* InTargetLandscape, const TSubclassOf<ULandscapeEditLayerBase>& InEditLayerClass) override;
 	virtual void RefreshDetailPanel() override;
 
+private:
+	TSharedRef<FPropertySection> RegisterPropertySection(FPropertyEditorModule& PropertyModule, FName ClassName, FName SectionName, FText DisplayName);
+	void RegisterPropertySectionMappings();
+	void UnregisterPropertySectionMappings();
+
 protected:
 	TSharedPtr<FUICommandList> GlobalUICommandList;
 	TArray<FRegisteredLandscapeHeightmapFileFormat> HeightmapFormats;
@@ -304,6 +315,9 @@ protected:
 	mutable FString HeightmapExportDialogTypeString;
 	mutable FString WeightmapExportDialogTypeString;
 	TUniquePtr<FLandscapeImageFileCache> LandscapeImageFileCache;
+
+	TMultiMap<FName, FName> RegisteredPropertySections;
+	TArray<FName> PropertyTypesToUnregisterOnShutdown;
 };
 
 IMPLEMENT_MODULE(FLandscapeEditorModule, LandscapeEditor);
@@ -342,6 +356,71 @@ FRegisteredLandscapeWeightmapFileFormat::FRegisteredLandscapeWeightmapFileFormat
 	}
 }
 
+TSharedRef<FPropertySection> FLandscapeEditorModule::RegisterPropertySection(FPropertyEditorModule& PropertyModule, FName ClassName, FName SectionName, FText DisplayName)
+{
+	TSharedRef<FPropertySection> PropertySection = PropertyModule.FindOrCreateSection(ClassName, SectionName, DisplayName);
+	RegisteredPropertySections.Add(ClassName, SectionName);
+
+	return PropertySection;
+}
+
+void FLandscapeEditorModule::RegisterPropertySectionMappings()
+{
+	const FName PropertyEditorModuleName("PropertyEditor");
+	FPropertyEditorModule& PropertyModule = FModuleManager::GetModuleChecked<FPropertyEditorModule>(PropertyEditorModuleName);
+
+	// Segment Physics
+	{
+		TSharedRef<FPropertySection> Section = RegisterPropertySection(PropertyModule, "LandscapeSplineSegment", "Physics", LOCTEXT("Physics", "Physics"));
+		Section->AddCategory("Collision");
+		Section->AddCategory("Physics");
+	}
+
+	// Segment Rendering
+	{
+		TSharedRef<FPropertySection> Section = RegisterPropertySection(PropertyModule, "LandscapeSplineSegment", "Rendering", LOCTEXT("Rendering", "Rendering"));
+		Section->AddCategory("VirtualTexture");
+		Section->AddCategory("Mesh");
+		Section->AddCategory("LandscapeSplineMeshes");
+		Section->AddCategory("LandscapeSplineMeshEntry");
+		Section->AddCategory("Rendering");
+	}
+
+	// Control Point Physics
+	{
+		TSharedRef<FPropertySection> Section = RegisterPropertySection(PropertyModule, "LandscapeSplineControlPoint", "Physics", LOCTEXT("Physics", "Physics"));
+		Section->AddCategory("Collision");
+		Section->AddCategory("Physics");
+	}
+
+	// Control Point Rendering
+	{
+		TSharedRef<FPropertySection> Section = RegisterPropertySection(PropertyModule, "LandscapeSplineControlPoint", "Rendering", LOCTEXT("Rendering", "Rendering"));
+		Section->AddCategory("VirtualTexture");
+		Section->AddCategory("Mesh");
+		Section->AddCategory("Rendering");
+	}
+}
+
+void FLandscapeEditorModule::UnregisterPropertySectionMappings()
+{
+	const FName PropertyEditorModuleName("PropertyEditor");
+	FPropertyEditorModule* PropertyModule = FModuleManager::GetModulePtr<FPropertyEditorModule>(PropertyEditorModuleName);
+
+	if (!PropertyModule)
+	{
+		return;
+	}
+
+	for (TMultiMap<FName, FName>::TIterator PropertySectionIterator = RegisteredPropertySections.CreateIterator(); PropertySectionIterator; ++PropertySectionIterator)
+	{
+		PropertyModule->RemoveSection(PropertySectionIterator->Key, PropertySectionIterator->Value);
+		PropertySectionIterator.RemoveCurrent();
+	}
+
+	check(RegisteredPropertySections.IsEmpty());
+}
+
 const TCHAR* FLandscapeEditorModule::GetHeightmapImportDialogTypeString() const
 {
 	if (HeightmapImportDialogTypeString.IsEmpty())
diff --git a/Engine/Source/Editor/LandscapeEditor/Private/LandscapeSplineDetails.cpp b/Engine/Source/Editor/LandscapeEditor/Private/LandscapeSplineDetails.cpp
--- a/Engine/Source/Editor/LandscapeEditor/Private/LandscapeSplineDetails.cpp
+++ b/Engine/Source/Editor/LandscapeEditor/Private/LandscapeSplineDetails.cpp
@@ -6,6 +6,8 @@
 #include "LandscapeEdMode.h"
 #include "Widgets/Text/STextBlock.h"
 #include "Widgets/Input/SButton.h"
+#include "Widgets/Layout/SSeparator.h"
+#include "Widgets/Input/SButton.h"
 #include "EditorModeManager.h"
 #include "EditorModes.h"
 #include "DetailWidgetRow.h"
@@ -38,36 +40,141 @@ void FLandscapeSplineDetails::CustomizeDetails(IDetailLayoutBuilder& DetailBuild
 		]
 	];
 
-	LandscapeSplineCategory.AddCustomRow(FText::GetEmpty())
+	FMargin ButtonPadding(2.f, 0.f);
+
+	LandscapeSplineCategory.AddCustomRow(LOCTEXT("SelectSplineElements", "Select Spline Elements"))
+	.RowTag("SelectSplineElements")
+	.NameContent()
+	[
+		SNew(STextBlock)
+			.Font(IDetailLayoutBuilder::GetDetailFont())
+			.Text(IsSegmentSelectModeEnabled() ? LOCTEXT("SelectSplineSegments", "Select Spline Segments") : LOCTEXT("SelectSplineControlPoints", "Select Spline Control Points"))
+	]
+	.ValueContent()
+	.VAlign(VAlign_Fill)
+	.MaxDesiredWidth(170.f)
+	.MinDesiredWidth(170.f)
 	[
 		SNew(SHorizontalBox)
+		.Clipping(EWidgetClipping::ClipToBounds)
+
 		+ SHorizontalBox::Slot()
-		.Padding(0, 0, 2, 0)
 		.VAlign(VAlign_Center)
-		.FillWidth(1)
+		.AutoWidth()
+		.Padding(ButtonPadding)
 		[
-			SNew(STextBlock)
-			.Text(LOCTEXT("SelectAll", "Select all connected:"))
+			SNew(SButton)
+				.ButtonStyle(FAppStyle::Get(), "SplineComponentDetails.SelectFirst")
+				.ContentPadding(2.0f)
+				.ToolTipText(IsSegmentSelectModeEnabled() ? LOCTEXT("SelectFirstSplineSegmentToolTip", "Select first spline segment.") : LOCTEXT("SelectFirstSplinePointToolTip", "Select first spline control point."))
+				.OnClicked(this, &FLandscapeSplineDetails::OnSelectEndLinearSplineElementButtonClicked, ESplineNavigationFlags::DirectionBackward)
+				.IsEnabled(this, &FLandscapeSplineDetails::HasEndLinearSplineElement, ESplineNavigationFlags::DirectionBackward)
 		]
 		+ SHorizontalBox::Slot()
-		.Padding(0, 2, 0, 2)
-		.FillWidth(1)
+		.AutoWidth()
+		.VAlign(VAlign_Center)
+		.Padding(ButtonPadding)
 		[
 			SNew(SButton)
-			.Text(LOCTEXT("ControlPoints", "Control Points"))
-			.HAlign(HAlign_Center)
-			.VAlign(VAlign_Center)
-			.OnClicked(this, &FLandscapeSplineDetails::OnSelectConnectedControlPointsButtonClicked)
+				.ButtonStyle(FAppStyle::Get(), "SplineComponentDetails.AddPrev")
+				.ContentPadding(2.f)
+				.VAlign(VAlign_Center)
+				.HAlign(HAlign_Center)
+				.ToolTipText(IsSegmentSelectModeEnabled() ? LOCTEXT("SelectAddPrevSplineSegmentToolTip", "Add previous segment to current selection.") : LOCTEXT("SelectAddPrevSplinePointToolTip", "Add previous control point to current selection."))
+				.OnClicked(this, &FLandscapeSplineDetails::OnSelectAdjacentLinearSplineElementButtonClicked, (ESplineNavigationFlags::DirectionBackward | ESplineNavigationFlags::AddToSelection))
+				.IsEnabled(this, &FLandscapeSplineDetails::HasAdjacentLinearSplineElement, (ESplineNavigationFlags::DirectionBackward | ESplineNavigationFlags::AddToSelection))
 		]
 		+ SHorizontalBox::Slot()
-		.Padding(0, 2, 0, 2)
-		.FillWidth(1)
+		.VAlign(VAlign_Center)
+		.AutoWidth()
 		[
 			SNew(SButton)
-			.Text(LOCTEXT("Segments", "Segments"))
-			.HAlign(HAlign_Center)
-			.VAlign(VAlign_Center)
-			.OnClicked(this, &FLandscapeSplineDetails::OnSelectConnectedSegmentsButtonClicked)
+				.ButtonStyle(FAppStyle::Get(), "SplineComponentDetails.SelectPrev")
+				.ContentPadding(2.f)
+				.VAlign(VAlign_Center)
+				.HAlign(HAlign_Center)
+				.ToolTipText(IsSegmentSelectModeEnabled() ? LOCTEXT("SelectPrevSplineSegmentToolTip", "Select previous segment.") : LOCTEXT("SelectPrevPointToolTip", "Select previous control point."))
+				.OnClicked(this, &FLandscapeSplineDetails::OnSelectAdjacentLinearSplineElementButtonClicked, ESplineNavigationFlags::DirectionBackward)
+				.IsEnabled(this, &FLandscapeSplineDetails::HasAdjacentLinearSplineElement, ESplineNavigationFlags::DirectionBackward)
+
+		]
+		+ SHorizontalBox::Slot()
+		.VAlign(VAlign_Center)
+		.AutoWidth()
+		.Padding(ButtonPadding)
+		[
+			SNew(SButton)
+				.ButtonStyle(FAppStyle::Get(), "SplineComponentDetails.SelectAll")
+				.ContentPadding(2.f)
+				.VAlign(VAlign_Center)
+				.HAlign(HAlign_Center)
+				.ToolTipText(IsSegmentSelectModeEnabled() ? LOCTEXT("SelectAllSplineSegmentToolTip", "Select all segments.") : LOCTEXT("SelectAllSplinePointsToolTip", "Select all control points."))
+				.OnClicked(this, &FLandscapeSplineDetails::OnSelectAllConnectedSplineElementsButtonClicked)
+		]
+		+ SHorizontalBox::Slot()
+		.AutoWidth()
+		.VAlign(VAlign_Center)
+		.Padding(ButtonPadding)
+		[
+			SNew(SButton)
+				.ButtonStyle(FAppStyle::Get(), "SplineComponentDetails.SelectNext")
+				.ContentPadding(2.f)
+				.VAlign(VAlign_Center)
+				.HAlign(HAlign_Center)
+				.ToolTipText(IsSegmentSelectModeEnabled() ? LOCTEXT("SelectNextSplineSegmentToolTip", "Select next segment.") : LOCTEXT("SelectNextSplinePointToolTip", "Select next control point."))
+				.OnClicked(this, &FLandscapeSplineDetails::OnSelectAdjacentLinearSplineElementButtonClicked, ESplineNavigationFlags::DirectionForward)
+				.IsEnabled(this, &FLandscapeSplineDetails::HasAdjacentLinearSplineElement, ESplineNavigationFlags::DirectionForward)
+		]
+		+ SHorizontalBox::Slot()
+		.AutoWidth()
+		.VAlign(VAlign_Center)
+		.Padding(ButtonPadding)
+		[
+			SNew(SButton)
+				.ButtonStyle(FAppStyle::Get(), "SplineComponentDetails.AddNext")
+				.ContentPadding(2.f)
+				.VAlign(VAlign_Center)
+				.HAlign(HAlign_Center)
+				.ToolTipText(IsSegmentSelectModeEnabled() ? LOCTEXT("SelectAddNextSplineSegmentToolTip", "Add next segment to current selection.") : LOCTEXT("SelectAddNextSplinePointToolTip", "Add next control point to current selection."))
+				.OnClicked(this, &FLandscapeSplineDetails::OnSelectAdjacentLinearSplineElementButtonClicked, (ESplineNavigationFlags::DirectionForward | ESplineNavigationFlags::AddToSelection))
+				.IsEnabled(this, &FLandscapeSplineDetails::HasAdjacentLinearSplineElement, (ESplineNavigationFlags::DirectionForward | ESplineNavigationFlags::AddToSelection))
+		]
+		+ SHorizontalBox::Slot()
+		.VAlign(VAlign_Center)
+		.AutoWidth()
+		.Padding(ButtonPadding)
+		[
+			SNew(SButton)
+				.ButtonStyle(FAppStyle::Get(), "SplineComponentDetails.SelectLast")
+				.ContentPadding(2.0f)
+				.VAlign(VAlign_Center)
+				.HAlign(HAlign_Center)
+				.ToolTipText(IsSegmentSelectModeEnabled() ? LOCTEXT("SelectLastSplineSegmentToolTip", "Select last spline segment.") : LOCTEXT("SelectLastSplinePointToolTip", "Select last spline control point."))
+				.OnClicked(this, &FLandscapeSplineDetails::OnSelectEndLinearSplineElementButtonClicked, ESplineNavigationFlags::DirectionForward)
+				.IsEnabled(this, &FLandscapeSplineDetails::HasEndLinearSplineElement, ESplineNavigationFlags::DirectionForward)
+		]
+
+		// Add a vertical divider to separate navigation with conversion
+		+ SHorizontalBox::Slot()
+		.VAlign(VAlign_Center)
+		.AutoWidth()
+		.Padding(ButtonPadding * 3)
+		[
+			SNew(SSeparator)
+			.Orientation(EOrientation::Orient_Vertical)
+			.Thickness(2.f)
+		]
+		+ SHorizontalBox::Slot()
+		.VAlign(VAlign_Center)
+		.AutoWidth()
+		[
+			SNew(SButton)
+				.ButtonStyle(FAppStyle::Get(), "SplineComponentDetails.Convert")
+				.ContentPadding(ButtonPadding)
+				.VAlign(VAlign_Center)
+				.HAlign(HAlign_Center)
+				.ToolTipText(IsSegmentSelectModeEnabled() ? LOCTEXT("SelectConvertSplineSegmentsToPoints", "Convert selected spline segments into control points.") : LOCTEXT("SelectConvertSplinePointsToSegments", "Convert selected spline control points into segments."))
+				.OnClicked(this, &FLandscapeSplineDetails::OnToggleSplineSelectionTypeButtonClicked)
 		]
 	];
 
@@ -123,20 +230,59 @@ class FEdModeLandscape* FLandscapeSplineDetails::GetEditorMode() const
 	return (FEdModeLandscape*)GLevelEditorModeTools().GetActiveMode(FBuiltinEditorModes::EM_Landscape);
 }
 
+bool FLandscapeSplineDetails::IsSegmentSelectModeEnabled() const
+{
+	FEdModeLandscape* LandscapeEdMode = GetEditorMode();
+	if (LandscapeEdMode && LandscapeEdMode->CurrentToolTarget.LandscapeInfo.IsValid() && LandscapeEdMode->CurrentToolTarget.LandscapeInfo->GetCurrentLevelLandscapeProxy(true))
+	{
+		return LandscapeEdMode->HasSelectedSplineSegments();
+	}
+
+	return false;
+}
+
 FReply FLandscapeSplineDetails::OnFlipSegmentButtonClicked()
 {
 	FEdModeLandscape* LandscapeEdMode = GetEditorMode();
-	if (LandscapeEdMode)
+	if (LandscapeEdMode && LandscapeEdMode->CurrentToolTarget.LandscapeInfo.IsValid() && LandscapeEdMode->CurrentToolTarget.LandscapeInfo->GetCurrentLevelLandscapeProxy(true))
 	{
 		LandscapeEdMode->FlipSelectedSplineSegments();
 	}
 	return FReply::Handled();
 }
 
+bool FLandscapeSplineDetails::HasAdjacentLinearSplineElement(ESplineNavigationFlags Flags) const
+{
+	FEdModeLandscape* LandscapeEdMode = GetEditorMode();
+	if (LandscapeEdMode && LandscapeEdMode->CurrentToolTarget.LandscapeInfo.IsValid() && LandscapeEdMode->CurrentToolTarget.LandscapeInfo->GetCurrentLevelLandscapeProxy(true))
+	{
+		// check flags have one direction set
+		check(FMath::CountBits(static_cast<uint8>(Flags & ESplineNavigationFlags::DirectionMask)) == 1);
+		return LandscapeEdMode->HasAdjacentLinearSplineConnection(Flags | (IsSegmentSelectModeEnabled() ? ESplineNavigationFlags::SegmentSelectModeEnabled : ESplineNavigationFlags::ControlPointSelectModeEnabled));
+	}
+	return false;
+}
+
+bool FLandscapeSplineDetails::HasEndLinearSplineElement(ESplineNavigationFlags Flags) const
+{
+	FEdModeLandscape* LandscapeEdMode = GetEditorMode();
+	if (LandscapeEdMode && LandscapeEdMode->CurrentToolTarget.LandscapeInfo.IsValid() && LandscapeEdMode->CurrentToolTarget.LandscapeInfo->GetCurrentLevelLandscapeProxy(true))
+	{
+		// check flags have one direction set
+		check(FMath::CountBits(static_cast<uint8>(Flags & ESplineNavigationFlags::DirectionMask)) == 1);
+		return LandscapeEdMode->HasEndLinearSplineConnection(Flags | (IsSegmentSelectModeEnabled() ? ESplineNavigationFlags::SegmentSelectModeEnabled : ESplineNavigationFlags::ControlPointSelectModeEnabled));
+	}
+	return false;
+}
+
 bool FLandscapeSplineDetails::IsFlipSegmentButtonEnabled() const
 {
 	FEdModeLandscape* LandscapeEdMode = GetEditorMode();
-	return LandscapeEdMode && LandscapeEdMode->HasSelectedSplineSegments();
+	if (LandscapeEdMode && LandscapeEdMode->CurrentToolTarget.LandscapeInfo.IsValid() && LandscapeEdMode->CurrentToolTarget.LandscapeInfo->GetCurrentLevelLandscapeProxy(true))
+	{
+		return LandscapeEdMode->HasSelectedSplineSegments();
+	}
+	return false;
 }
 
 FText FLandscapeSplineDetails::OnGetSplineOwningLandscapeText() const
@@ -165,23 +311,49 @@ FText FLandscapeSplineDetails::OnGetSplineOwningLandscapeText() const
 	return FText::FromString("Owner: " + SplineOwnersStr);
 }
 
-FReply FLandscapeSplineDetails::OnSelectConnectedControlPointsButtonClicked()
+FReply FLandscapeSplineDetails::OnSelectAdjacentLinearSplineElementButtonClicked(ESplineNavigationFlags Flags) const
 {
 	FEdModeLandscape* LandscapeEdMode = GetEditorMode();
-	if (LandscapeEdMode && LandscapeEdMode->CurrentToolTarget.LandscapeInfo.IsValid())
+	if (LandscapeEdMode && LandscapeEdMode->CurrentToolTarget.LandscapeInfo.IsValid() && LandscapeEdMode->CurrentToolTarget.LandscapeInfo->GetCurrentLevelLandscapeProxy(true))
 	{
-		LandscapeEdMode->SelectAllConnectedSplineControlPoints();
+		// check flags have one direction set
+		check(FMath::CountBits(static_cast<uint8>(Flags & ESplineNavigationFlags::DirectionMask)) == 1);
+		LandscapeEdMode->SelectAdjacentLinearSplineElement(Flags | (IsSegmentSelectModeEnabled() ? ESplineNavigationFlags::SegmentSelectModeEnabled : ESplineNavigationFlags::ControlPointSelectModeEnabled));
 	}
 
 	return FReply::Handled();
 }
 
-FReply FLandscapeSplineDetails::OnSelectConnectedSegmentsButtonClicked()
+FReply FLandscapeSplineDetails::OnSelectEndLinearSplineElementButtonClicked(ESplineNavigationFlags Flags) const
 {
 	FEdModeLandscape* LandscapeEdMode = GetEditorMode();
-	if (LandscapeEdMode && LandscapeEdMode->CurrentToolTarget.LandscapeInfo.IsValid())
+	if (LandscapeEdMode && LandscapeEdMode->CurrentToolTarget.LandscapeInfo.IsValid() && LandscapeEdMode->CurrentToolTarget.LandscapeInfo->GetCurrentLevelLandscapeProxy(true))
 	{
-		LandscapeEdMode->SelectAllConnectedSplineSegments();
+		// check flags have one direction set
+		check(FMath::CountBits(static_cast<uint8>(Flags & ESplineNavigationFlags::DirectionMask)) == 1);
+		LandscapeEdMode->SelectEndLinearSplineElement(Flags | (IsSegmentSelectModeEnabled() ? ESplineNavigationFlags::SegmentSelectModeEnabled : ESplineNavigationFlags::ControlPointSelectModeEnabled));
+	}
+
+	return FReply::Handled();
+}
+
+FReply FLandscapeSplineDetails::OnToggleSplineSelectionTypeButtonClicked() const
+{
+	FEdModeLandscape* LandscapeEdMode = GetEditorMode();
+	if (LandscapeEdMode && LandscapeEdMode->CurrentToolTarget.LandscapeInfo.IsValid() && LandscapeEdMode->CurrentToolTarget.LandscapeInfo->GetCurrentLevelLandscapeProxy(true))
+	{
+		IsSegmentSelectModeEnabled() ? LandscapeEdMode->SelectSplineControlPointsFromCurrentSegmentSelection() : LandscapeEdMode->SelectSplineSegmentsFromCurrentControlPointSelection();
+	}
+
+	return FReply::Handled();
+}
+
+FReply FLandscapeSplineDetails::OnSelectAllConnectedSplineElementsButtonClicked() const
+{
+	FEdModeLandscape* LandscapeEdMode = GetEditorMode();
+	if (LandscapeEdMode && LandscapeEdMode->CurrentToolTarget.LandscapeInfo.IsValid() && LandscapeEdMode->CurrentToolTarget.LandscapeInfo->GetCurrentLevelLandscapeProxy(true))
+	{
+		IsSegmentSelectModeEnabled() ? LandscapeEdMode->SelectAllConnectedSplineSegments() : LandscapeEdMode->SelectAllConnectedSplineControlPoints();
 	}
 
 	return FReply::Handled();
@@ -223,7 +395,11 @@ FReply FLandscapeSplineDetails::OnUpdateSplineMeshLevelsButtonClicked()
 bool FLandscapeSplineDetails::IsUpdateSplineMeshLevelsButtonEnabled() const
 {
 	FEdModeLandscape* LandscapeEdMode = GetEditorMode();
-	return (LandscapeEdMode && LandscapeEdMode->CurrentToolTarget.LandscapeInfo.IsValid());
+	if (LandscapeEdMode && LandscapeEdMode->CurrentToolTarget.LandscapeInfo.IsValid())
+	{
+		return LandscapeEdMode->CurrentToolTarget.LandscapeInfo.IsValid();
+	}
+	return false;
 }
 
 #undef LOCTEXT_NAMESPACE
diff --git a/Engine/Source/Editor/LandscapeEditor/Private/LandscapeSplineDetails.h b/Engine/Source/Editor/LandscapeEditor/Private/LandscapeSplineDetails.h
--- a/Engine/Source/Editor/LandscapeEditor/Private/LandscapeSplineDetails.h
+++ b/Engine/Source/Editor/LandscapeEditor/Private/LandscapeSplineDetails.h
@@ -8,6 +8,7 @@
 
 class FEdModeLandscape;
 class IDetailLayoutBuilder;
+enum class ESplineNavigationFlags : uint8;
 
 class FLandscapeSplineDetails : public IDetailCustomization
 {
@@ -21,14 +22,21 @@ public:
 private:
 	FEdModeLandscape* GetEditorMode() const;
 	FText OnGetSplineOwningLandscapeText() const;
-	FReply OnSelectConnectedControlPointsButtonClicked();
-	FReply OnSelectConnectedSegmentsButtonClicked();
 	FReply OnMoveToCurrentLevelButtonClicked();
 	bool IsMoveToCurrentLevelButtonEnabled() const;
+	bool IsSegmentSelectModeEnabled() const;
 
 	FReply OnUpdateSplineMeshLevelsButtonClicked();
 	bool IsUpdateSplineMeshLevelsButtonEnabled() const;
 	
 	FReply OnFlipSegmentButtonClicked();
 	bool IsFlipSegmentButtonEnabled() const;
+
+	bool HasAdjacentLinearSplineElement(ESplineNavigationFlags Flags) const;
+	bool HasEndLinearSplineElement(ESplineNavigationFlags Flags) const;
+	FReply OnSelectAdjacentLinearSplineElementButtonClicked(ESplineNavigationFlags Flags) const;
+	FReply OnSelectEndLinearSplineElementButtonClicked(ESplineNavigationFlags Flags) const;
+
+	FReply OnSelectAllConnectedSplineElementsButtonClicked() const;
+	FReply OnToggleSplineSelectionTypeButtonClicked() const;
 };
