Add Tags to C++ Terrain Tests

diff --git a/Engine/Plugins/Experimental/Water/Source/Editor/Private/Tests/WaterInLandscapeEditLayerTest.cpp b/Engine/Plugins/Experimental/Water/Source/Editor/Private/Tests/WaterInLandscapeEditLayerTest.cpp
--- a/Engine/Plugins/Experimental/Water/Source/Editor/Private/Tests/WaterInLandscapeEditLayerTest.cpp
+++ b/Engine/Plugins/Experimental/Water/Source/Editor/Private/Tests/WaterInLandscapeEditLayerTest.cpp
@@ -110,6 +110,7 @@ void FWaterInLandscapeEditLayerTest::GetTests(TArray<FString>& OutBeautifiedName
 		TEXT("VerifyLandscapeOceanWaterLayer"),
 		TEXT("VerifyLandscapeIslandWaterLayer")
 	};
+
 	const TArray<FString> WaterLandscapeCommandsGetter
 	{
 		TEXT("Water Body River"),
@@ -117,10 +118,15 @@ void FWaterInLandscapeEditLayerTest::GetTests(TArray<FString>& OutBeautifiedName
 		TEXT("Water Body Ocean"),
 		TEXT("Water Body Island")
 	};
+
+	FAutomationTestFramework& Framework = FAutomationTestFramework::Get();
+
 	for (const FString& Test : WaterLandscapeNamesGetter)
 	{
 		OutBeautifiedNames.Add(Test);
+		Framework.RegisterComplexAutomationTestTags(this, Test, TEXT("[GraphicsTools][Terrain][Water][DailyEssential]"));
 	}
+
 	for (const FString& TestCmd : WaterLandscapeCommandsGetter)
 	{
 		OutTestCommands.Add(TestCmd);
diff --git a/Engine/Source/Editor/LandscapeEditor/Private/Tests/LandscapeAutomationTests.cpp b/Engine/Source/Editor/LandscapeEditor/Private/Tests/LandscapeAutomationTests.cpp
--- a/Engine/Source/Editor/LandscapeEditor/Private/Tests/LandscapeAutomationTests.cpp
+++ b/Engine/Source/Editor/LandscapeEditor/Private/Tests/LandscapeAutomationTests.cpp
@@ -32,6 +32,9 @@ DEFINE_LOG_CATEGORY_STATIC(LogLandscapeAutomationTests, Log, All);
 */
 namespace LandscapeTestUtils
 {
+	constexpr EAutomationTestFlags LandscapeTestFlags = EAutomationTestFlags::EditorContext | EAutomationTestFlags::NonNullRHI | EAutomationTestFlags::EngineFilter;
+	const FString LandscapeDailyTestTags = TEXT("[GraphicsTools][Terrain][Landscape][DailyEssential]");
+
 	/**
 	* Finds the viewport to use for the landscape tool
 	*/
@@ -214,7 +217,10 @@ bool FCheckHeight::Update()
 /**
 * Landscape creation / edit test
 */
-IMPLEMENT_SIMPLE_AUTOMATION_TEST(FLandscapeEditorTest, "Editor.Landscape.Create and Modify New Landscape", EAutomationTestFlags::EditorContext | EAutomationTestFlags::NonNullRHI | EAutomationTestFlags::EngineFilter);
+const FString LandscapeEditorTestName = TEXT("Editor.Landscape.Create and Modify New Landscape");
+
+IMPLEMENT_SIMPLE_AUTOMATION_TEST(FLandscapeEditorTest, LandscapeEditorTestName, LandscapeTestUtils::LandscapeTestFlags)
+REGISTER_SIMPLE_AUTOMATION_TEST_TAGS(FLandscapeEditorTest, LandscapeEditorTestName, LandscapeTestUtils::LandscapeDailyTestTags)
 bool FLandscapeEditorTest::RunTest(const FString& Parameters)
 {
 	// Create a new map and landscape
@@ -228,8 +234,14 @@ bool FLandscapeEditorTest::RunTest(const FString& Parameters)
 	return true;
 }
 
-IMPLEMENT_SIMPLE_AUTOMATION_TEST(FLandscapeEditorTestFooBar, "Editor.Landscape.Import Landscape", EAutomationTestFlags::EditorContext | EAutomationTestFlags::NonNullRHI | EAutomationTestFlags::EngineFilter);
-bool FLandscapeEditorTestFooBar::RunTest(const FString& Parameters)
+/**
+ * Landscape - Import Landscape Test
+ */
+const FString LandscapeEditorImportTestName = TEXT("Editor.Landscape.Import Landscape");
+
+IMPLEMENT_SIMPLE_AUTOMATION_TEST(FLandscapeEditorImportTest, LandscapeEditorImportTestName, LandscapeTestUtils::LandscapeTestFlags)
+REGISTER_SIMPLE_AUTOMATION_TEST_TAGS(FLandscapeEditorImportTest, LandscapeEditorImportTestName, LandscapeTestUtils::LandscapeDailyTestTags)
+bool FLandscapeEditorImportTest::RunTest(const FString& Parameters)
 {
 	// Create a new map and landscape
 	ADD_LATENT_AUTOMATION_COMMAND(FCreateNewMapWithLandscapeCommand(/* NumComponents = */ 8, /* QuadsPerComponent = */ 7));
@@ -249,9 +261,12 @@ bool FLandscapeEditorTestFooBar::RunTest(const FString& Parameters)
 
 /*
  * Verify that we can create a new edit layer for a landscape
- * See: TR = 290778158
+ * QMetry: UE-TC-4704
  */
-IMPLEMENT_SIMPLE_AUTOMATION_TEST(FLandscapeEditorCreateNewLayer, "Editor.Landscape.Create New Edit Layer", EAutomationTestFlags::EditorContext | EAutomationTestFlags::NonNullRHI | EAutomationTestFlags::EngineFilter);
+const FString LandscapeEditorNewLayerTestName = TEXT("Editor.Landscape.Create New Edit Layer");
+
+IMPLEMENT_SIMPLE_AUTOMATION_TEST(FLandscapeEditorCreateNewLayer, LandscapeEditorNewLayerTestName, LandscapeTestUtils::LandscapeTestFlags)
+REGISTER_SIMPLE_AUTOMATION_TEST_TAGS(FLandscapeEditorCreateNewLayer, LandscapeEditorNewLayerTestName, LandscapeTestUtils::LandscapeDailyTestTags)
 bool FLandscapeEditorCreateNewLayer::RunTest(const FString& Parameters)
 {
 	// Create a new map with landscape
