Removed the "normal" word from the weight-blended landscape layer info creation menus since they're actually the more troublesome type

#rnx

diff --git a/Engine/Source/Editor/LandscapeEditor/Private/LandscapeEditorDetailCustomization_ImportLayers.cpp b/Engine/Source/Editor/LandscapeEditor/Private/LandscapeEditorDetailCustomization_ImportLayers.cpp
--- a/Engine/Source/Editor/LandscapeEditor/Private/LandscapeEditorDetailCustomization_ImportLayers.cpp
+++ b/Engine/Source/Editor/LandscapeEditor/Private/LandscapeEditorDetailCustomization_ImportLayers.cpp
@@ -298,7 +298,7 @@ TSharedRef<SWidget> FLandscapeEditorStructCustomization_FLandscapeImportLayer::O
 {
 	FMenuBuilder MenuBuilder(true, nullptr);
 
-	MenuBuilder.AddMenuEntry(LOCTEXT("Target_Create_Blended", "Weight-Blended Layer (normal)"), FText(), FSlateIcon(),
+	MenuBuilder.AddMenuEntry(LOCTEXT("Target_Create_Blended", "Weight-Blended Layer"), FText(), FSlateIcon(),
 		FUIAction(FExecuteAction::CreateStatic(&OnImportLayerCreateClicked, PropertyHandle_LayerInfo, LayerName, false)));
 
 	MenuBuilder.AddMenuEntry(LOCTEXT("Target_Create_NoWeightBlend", "Non Weight-Blended Layer"), FText(), FSlateIcon(),
diff --git a/Engine/Source/Editor/LandscapeEditor/Private/LandscapeEditorDetailCustomization_TargetLayers.cpp b/Engine/Source/Editor/LandscapeEditor/Private/LandscapeEditorDetailCustomization_TargetLayers.cpp
--- a/Engine/Source/Editor/LandscapeEditor/Private/LandscapeEditorDetailCustomization_TargetLayers.cpp
+++ b/Engine/Source/Editor/LandscapeEditor/Private/LandscapeEditorDetailCustomization_TargetLayers.cpp
@@ -1560,7 +1560,7 @@ TSharedRef<SWidget> FLandscapeEditorCustomNodeBuilder_TargetLayers::OnGetTargetL
 {
 	FMenuBuilder MenuBuilder(true, nullptr);
 
-	MenuBuilder.AddMenuEntry(LOCTEXT("Menu_Create_Blended", "Weight-Blended Layer (normal)"), FText(), FSlateIcon(),
+	MenuBuilder.AddMenuEntry(LOCTEXT("Menu_Create_Blended", "Weight-Blended Layer"), FText(), FSlateIcon(),
 		FUIAction(FExecuteAction::CreateStatic(&FLandscapeEditorCustomNodeBuilder_TargetLayers::OnTargetLayerCreateClicked, Target, false)));
 
 	MenuBuilder.AddMenuEntry(LOCTEXT("Menu_Create_NoWeightBlend", "Non Weight-Blended Layer"), FText(), FSlateIcon(),
