Crash resistance for ALanscapeSplineActor when in a bad state.

diff --git a/Engine/Source/Editor/LandscapeEditor/Private/LandscapeEdModeSplineTools.cpp b/Engine/Source/Editor/LandscapeEditor/Private/LandscapeEdModeSplineTools.cpp
--- a/Engine/Source/Editor/LandscapeEditor/Private/LandscapeEdModeSplineTools.cpp
+++ b/Engine/Source/Editor/LandscapeEditor/Private/LandscapeEdModeSplineTools.cpp
@@ -1295,7 +1295,7 @@ public:
 			}
 
 			// No Spline Owner found
-			if (!SplineOwner)
+			if (!SplineOwner || !SplineOwner->IsSplineOwnerValid())
 			{
 				return false;
 			}
@@ -1668,11 +1668,14 @@ public:
 					if (HitProxy->IsA(HLandscapeSplineProxy_Segment::StaticGetType()))
 					{
 						HLandscapeSplineProxy_Segment* SplineProxy = (HLandscapeSplineProxy_Segment*)HitProxy;
-						ClickedSplineSegment = SplineProxy->SplineSegment;
-						ULandscapeSplinesComponent* SplineComponent = ClickedSplineSegment->GetOuterULandscapeSplinesComponent();
+						ULandscapeSplinesComponent* SplineComponent = SplineProxy->SplineSegment->GetOuterULandscapeSplinesComponent();
 						ILandscapeSplineInterface* SplineOwner = SplineComponent->GetSplineOwner();
 						check(SplineOwner);
-						LandscapeToSpline = SplineOwner->LandscapeActorToWorld().GetRelativeTransform(SplineComponent->GetComponentTransform());
+						if (SplineOwner->IsSplineOwnerValid())
+						{
+							ClickedSplineSegment = SplineProxy->SplineSegment;
+							LandscapeToSpline = SplineOwner->LandscapeActorToWorld().GetRelativeTransform(SplineComponent->GetComponentTransform());
+						}
 					}
 					else if (HitProxy->IsA(HActor::StaticGetType()))
 					{
@@ -1689,11 +1692,14 @@ public:
 								{
 									if (ULandscapeSplineSegment* SplineSegment = Cast<ULandscapeSplineSegment>(ComponentOwner))
 									{
-										ClickedSplineSegment = SplineSegment;
 										// Find actual SplineComponent owner of the Segment (not the SplineComponent owner of the mesh)
 										SplineComponent = SplineSegment->GetTypedOuter<ULandscapeSplinesComponent>();
 										ILandscapeSplineInterface* SplineOwner = SplineComponent->GetSplineOwner();
-										LandscapeToSpline = SplineOwner->LandscapeActorToWorld().GetRelativeTransform(SplineComponent->GetComponentTransform());
+										if (SplineOwner->IsSplineOwnerValid())
+										{
+											ClickedSplineSegment = SplineSegment;
+											LandscapeToSpline = SplineOwner->LandscapeActorToWorld().GetRelativeTransform(SplineComponent->GetComponentTransform());
+										}
 									}
 								}
 							}
diff --git a/Engine/Source/Runtime/Landscape/Classes/ILandscapeSplineInterface.h b/Engine/Source/Runtime/Landscape/Classes/ILandscapeSplineInterface.h
--- a/Engine/Source/Runtime/Landscape/Classes/ILandscapeSplineInterface.h
+++ b/Engine/Source/Runtime/Landscape/Classes/ILandscapeSplineInterface.h
@@ -25,7 +25,8 @@ public:
 	LANDSCAPE_API virtual FTransform LandscapeActorToWorld() const PURE_VIRTUAL(ILandscapeSplineInterface::LandscapeActorToWorld, return FTransform::Identity;);
 	LANDSCAPE_API virtual ULandscapeSplinesComponent* GetSplinesComponent() const PURE_VIRTUAL(ILandscapeSplineInterface::GetSplineComponent, return nullptr;);
 	LANDSCAPE_API virtual void UpdateSharedProperties(ULandscapeInfo* InLandscapeInfo) PURE_VIRTUAL(ILandscapeSplineInterface::UpdateSharedProperties);
-
+	LANDSCAPE_API virtual bool IsSplineOwnerValid() const PURE_VIRTUAL(ILandscapeSplineInterface::IsSplineOwnerValid, return true;);
+	 
 #if WITH_EDITOR
 	LANDSCAPE_API virtual bool SupportsForeignSplineMesh() const PURE_VIRTUAL(ILandscapeSplineInterface::SupportsForeignSplineMesh, return false;);
 	LANDSCAPE_API virtual void CreateSplineComponent() PURE_VIRTUAL(ILandscapeSplineInterface::CreateSplineComponent);
diff --git a/Engine/Source/Runtime/Landscape/Classes/LandscapeProxy.h b/Engine/Source/Runtime/Landscape/Classes/LandscapeProxy.h
--- a/Engine/Source/Runtime/Landscape/Classes/LandscapeProxy.h
+++ b/Engine/Source/Runtime/Landscape/Classes/LandscapeProxy.h
@@ -1241,6 +1241,7 @@ PRAGMA_ENABLE_DEPRECATION_WARNINGS
 	// ILandscapeSplineInterface
 	virtual ULandscapeSplinesComponent* GetSplinesComponent() const override { return SplineComponent; }
 	LANDSCAPE_API virtual void UpdateSharedProperties(ULandscapeInfo* InLandscapeInfo) override;
+	virtual bool IsSplineOwnerValid() const override { return true; }
 
 	// Retrieve the screen size at which each LOD should be rendered
 	LANDSCAPE_API TArray<float> GetLODScreenSizeArray() const;
diff --git a/Engine/Source/Runtime/Landscape/Classes/LandscapeSplineActor.h b/Engine/Source/Runtime/Landscape/Classes/LandscapeSplineActor.h
--- a/Engine/Source/Runtime/Landscape/Classes/LandscapeSplineActor.h
+++ b/Engine/Source/Runtime/Landscape/Classes/LandscapeSplineActor.h
@@ -32,6 +32,7 @@ public:
 	virtual FTransform LandscapeActorToWorld() const override;
 	virtual ULandscapeInfo* GetLandscapeInfo() const override;
 	virtual void UpdateSharedProperties(ULandscapeInfo* InLandscapeInfo) override;
+	virtual bool IsSplineOwnerValid() const override;
 		
 #if WITH_EDITOR
 	LANDSCAPE_API bool ShouldPartitionSpline() const;
diff --git a/Engine/Source/Runtime/Landscape/Private/LandscapeSplineActor.cpp b/Engine/Source/Runtime/Landscape/Private/LandscapeSplineActor.cpp
--- a/Engine/Source/Runtime/Landscape/Private/LandscapeSplineActor.cpp
+++ b/Engine/Source/Runtime/Landscape/Private/LandscapeSplineActor.cpp
@@ -50,6 +50,13 @@ void ALandscapeSplineActor::UpdateSharedProperties(ULandscapeInfo* InLandscapeIn
 #endif // WITH_EDITOR
 }
 
+// Is this in a valid state to be safe to call LandscapeActorToWorld?
+bool ALandscapeSplineActor::IsSplineOwnerValid() const
+{
+	ULandscapeInfo* Info = GetLandscapeInfo();
+	return Info != nullptr && Info->LandscapeActor != nullptr;
+}
+
 #if WITH_EDITOR
 void ALandscapeSplineActor::GetActorDescProperties(FPropertyPairsMap& PropertyPairsMap) const
 {
diff --git a/Engine/Source/Runtime/Landscape/Private/LandscapeSplineRaster.cpp b/Engine/Source/Runtime/Landscape/Private/LandscapeSplineRaster.cpp
--- a/Engine/Source/Runtime/Landscape/Private/LandscapeSplineRaster.cpp
+++ b/Engine/Source/Runtime/Landscape/Private/LandscapeSplineRaster.cpp
@@ -613,7 +613,7 @@ bool ULandscapeInfo::ApplySplinesInternal(bool bOnlySelected, TScriptInterface<I
 {
 	TRACE_CPUPROFILER_EVENT_SCOPE(LandscapeInfo_ApplySplinesInternal);
 
-	if (!SplineOwner)
+	if (!SplineOwner || !SplineOwner->IsSplineOwnerValid())
 	{
 		return false;
 	}
