Added project setting to disable landscape target layer thumbnails. Because of how the thumbnail renderer is implemented, this can cause several hiccups at the moment (synchronous GPU readbacks) and the thumbnails are not working well anyway because of the RVT usage in the landscape material, so this setting represents an easy kill switch

diff --git a/Engine/Source/Editor/LandscapeEditor/Private/SLandscapeEditor.cpp b/Engine/Source/Editor/LandscapeEditor/Private/SLandscapeEditor.cpp
--- a/Engine/Source/Editor/LandscapeEditor/Private/SLandscapeEditor.cpp
+++ b/Engine/Source/Editor/LandscapeEditor/Private/SLandscapeEditor.cpp
@@ -37,7 +37,8 @@ void SLandscapeAssetThumbnail::Construct(const FArguments& InArgs, UObject* Asse
 {
 	FIntPoint ThumbnailSize = InArgs._ThumbnailSize;
 
-	AssetThumbnail = MakeShareable(new FAssetThumbnail(Asset, ThumbnailSize.X, ThumbnailSize.Y, ThumbnailPool));
+	const bool bShouldDisplayTargetLayerThumbnails = GetDefault<ULandscapeSettings>()->ShouldDisplayTargetLayerThumbnails();
+	AssetThumbnail = MakeShareable(new FAssetThumbnail(bShouldDisplayTargetLayerThumbnails ? Asset : nullptr, ThumbnailSize.X, ThumbnailSize.Y, ThumbnailPool));
 	OnAccessAsset = InArgs._OnAccessAsset;
 
 	ChildSlot
@@ -50,10 +51,12 @@ void SLandscapeAssetThumbnail::Construct(const FArguments& InArgs, UObject* Asse
 		]
 	];
 
-	UMaterialInterface* MaterialInterface = Cast<UMaterialInterface>(Asset);
-	if (MaterialInterface)
+	if (bShouldDisplayTargetLayerThumbnails)
 	{
-		UMaterial::OnMaterialCompilationFinished().AddSP(this, &SLandscapeAssetThumbnail::OnMaterialCompilationFinished);
+		if (UMaterialInterface* MaterialInterface = Cast<UMaterialInterface>(Asset))
+		{
+			UMaterial::OnMaterialCompilationFinished().AddSP(this, &SLandscapeAssetThumbnail::OnMaterialCompilationFinished);
+		}
 	}
 }
 
diff --git a/Engine/Source/Runtime/Landscape/Private/LandscapeSettings.cpp b/Engine/Source/Runtime/Landscape/Private/LandscapeSettings.cpp
--- a/Engine/Source/Runtime/Landscape/Private/LandscapeSettings.cpp
+++ b/Engine/Source/Runtime/Landscape/Private/LandscapeSettings.cpp
@@ -13,7 +13,8 @@ void ULandscapeSettings::PostEditChangeProperty(FPropertyChangedEvent& PropertyC
 
 	const FName PropertyName = PropertyChangedEvent.MemberProperty ? PropertyChangedEvent.MemberProperty->GetFName() : NAME_None;
 	if ((PropertyName == GET_MEMBER_NAME_CHECKED(ULandscapeSettings, BrushSizeUIMax))
-		|| (PropertyName == GET_MEMBER_NAME_CHECKED(ULandscapeSettings, BrushSizeClampMax)))
+		|| (PropertyName == GET_MEMBER_NAME_CHECKED(ULandscapeSettings, BrushSizeClampMax))
+		|| (PropertyName == GET_MEMBER_NAME_CHECKED(ULandscapeSettings, bDisplayTargetLayerThumbnails)))
 	{
 		// If landscape mode is active, refresh the detail panel to apply the changes immediately : 
 		ILandscapeModule& LandscapeModule = FModuleManager::GetModuleChecked<ILandscapeModule>("Landscape");
@@ -21,4 +22,4 @@ void ULandscapeSettings::PostEditChangeProperty(FPropertyChangedEvent& PropertyC
 	}
 }
 
-#endif // WITH_EDITOR
+#endif // WITH_EDITOR
\ No newline at end of file
diff --git a/Engine/Source/Runtime/Landscape/Public/LandscapeSettings.h b/Engine/Source/Runtime/Landscape/Public/LandscapeSettings.h
--- a/Engine/Source/Runtime/Landscape/Public/LandscapeSettings.h
+++ b/Engine/Source/Runtime/Landscape/Public/LandscapeSettings.h
@@ -63,6 +63,8 @@ public:
 	/** Returns the default landscape layer info object that will be assigned to unset layers when creating a new landscape. */
 	TSoftObjectPtr<ULandscapeLayerInfoObject> GetDefaultLayerInfoObject() const { return DefaultLayerInfoObject; }
 
+	bool ShouldDisplayTargetLayerThumbnails() const { return bDisplayTargetLayerThumbnails; }
+
 #if WITH_EDITOR
 	//~ Begin UObject Interface.
 	virtual void PostEditChangeProperty(FPropertyChangedEvent& PropertyChangedEvent) override;
@@ -120,5 +122,8 @@ protected:
 
 	UPROPERTY(EditAnywhere, config, Category = "Spline", meta = (ClampMin = "32", UIMax = "2048"))
 	float SplineIconScale = 50.0f;
-};
 
+	UPROPERTY(EditAnywhere, config, Category = "Configuration", meta = (ToolTip = 
+		"When true, each target layer will have a representative thumbnail in landscape mode. However, the thumbnail system can be slow. Plus, certain landscape materials (e.g. using RVT) are not able to properly render thumbnails. Thus, setting this to false will skip needlessly rendering landscape layer thumbnails, which can improve the editing experience"))
+	bool bDisplayTargetLayerThumbnails = true;
+};
