Landscape Paint - create layers from assigned materials new workflow with autofill layers button and dialog

diff --git a/Engine/Source/Editor/EditorStyle/Private/StarshipStyle.cpp b/Engine/Source/Editor/EditorStyle/Private/StarshipStyle.cpp
--- a/Engine/Source/Editor/EditorStyle/Private/StarshipStyle.cpp
+++ b/Engine/Source/Editor/EditorStyle/Private/StarshipStyle.cpp
@@ -7151,7 +7151,8 @@ void FStarshipEditorStyle::FStyle::SetupLandscapeEditorStyle()
 		Set("LandscapeEditor.TargetList.DropZone.Below", new BOX_BRUSH("Common/VerticalBoxDropZoneIndicator_Below", FMargin(10.0f / 16.0f, 0, 0, 10.0f / 16.0f), SelectionColor_Subdued));
 
 		Set("LandscapeEditor.Layer.Sync", new IMAGE_BRUSH_SVG("Icons/icon_Landscape_Layers_From_Material_16x", Icon16x16));
-		
+		Set("LandscapeEditor.Layer.AutoFill", new IMAGE_BRUSH_SVG("Icons/icon_Landscape_Layers_AutoFill_16x", Icon16x16));
+
 		Set("LandscapeEditor.SpinBox", FSpinBoxStyle(GetWidgetStyle<FSpinBoxStyle>("SpinBox"))
 			.SetTextPadding(FMargin(0.f))
 			.SetBackgroundBrush(FSlateNoResource())
diff --git a/Engine/Source/Editor/LandscapeEditor/Private/LandscapeEdMode.cpp b/Engine/Source/Editor/LandscapeEditor/Private/LandscapeEdMode.cpp
--- a/Engine/Source/Editor/LandscapeEditor/Private/LandscapeEdMode.cpp
+++ b/Engine/Source/Editor/LandscapeEditor/Private/LandscapeEdMode.cpp
@@ -2314,6 +2314,17 @@ const TArray<FLandscapeListInfo>& FEdModeLandscape::GetLandscapeList()
 	return LandscapeList;
 }
 
+const FString FEdModeLandscape::GetTargetLayerAssetPackagePath() const
+{
+	if (UISettings->TargetLayerAssetFilePath.bUseAssetDirectoryPath)
+	{
+		return UISettings->TargetLayerAssetFilePath.DirectoryPath.Path;
+	}
+
+	const ULevel* Level = GetWorld()->GetCurrentLevel();
+	return UE::Landscape::GetSharedAssetsPath(Level);
+}
+
 int32 FEdModeLandscape::UpdateLandscapeList()
 {
 	LandscapeList.Empty();
diff --git a/Engine/Source/Editor/LandscapeEditor/Private/LandscapeEdMode.h b/Engine/Source/Editor/LandscapeEditor/Private/LandscapeEdMode.h
--- a/Engine/Source/Editor/LandscapeEditor/Private/LandscapeEdMode.h
+++ b/Engine/Source/Editor/LandscapeEditor/Private/LandscapeEdMode.h
@@ -558,6 +558,7 @@ public:
 	const TArray<FName>* GetTargetDisplayOrderList() const;
 	int32 GetTargetLayerStartingIndex() const;
 	const TArray<FLandscapeListInfo>& GetLandscapeList();
+	const FString GetTargetLayerAssetPackagePath() const;
 
 	int32 UpdateLandscapeList();
 	void UpdateTargetList();
diff --git a/Engine/Source/Editor/LandscapeEditor/Private/LandscapeEditorDetailCustomization_ImportLayers.cpp b/Engine/Source/Editor/LandscapeEditor/Private/LandscapeEditorDetailCustomization_ImportLayers.cpp
--- a/Engine/Source/Editor/LandscapeEditor/Private/LandscapeEditorDetailCustomization_ImportLayers.cpp
+++ b/Engine/Source/Editor/LandscapeEditor/Private/LandscapeEditorDetailCustomization_ImportLayers.cpp
@@ -315,12 +315,11 @@ void FLandscapeEditorStructCustomization_FLandscapeImportLayer::OnImportLayerCre
 		return;
 	}
 	
-	// Hack as we don't have a direct world pointer in the EdMode...
-	ULevel* Level = LandscapeEdMode->CurrentGizmoActor->GetWorld()->GetCurrentLevel();
+	ULevel* Level = LandscapeEdMode->GetWorld()->GetCurrentLevel();
 
 	// Build default layer object name and package name
 	FName LayerObjectName;
-	FString PackageName = UE::Landscape::GetLayerInfoObjectPackageName(Level, LayerName, LayerObjectName);
+	FString PackageName = UE::Landscape::GetLayerInfoObjectPackageName(LayerName, UE::Landscape::GetSharedAssetsPath(Level), LayerObjectName);
 
 	TSharedRef<SDlgPickAssetPath> NewLayerDlg =
 		SNew(SDlgPickAssetPath)
diff --git a/Engine/Source/Editor/LandscapeEditor/Private/LandscapeEditorDetailCustomization_TargetLayers.cpp b/Engine/Source/Editor/LandscapeEditor/Private/LandscapeEditorDetailCustomization_TargetLayers.cpp
--- a/Engine/Source/Editor/LandscapeEditor/Private/LandscapeEditorDetailCustomization_TargetLayers.cpp
+++ b/Engine/Source/Editor/LandscapeEditor/Private/LandscapeEditorDetailCustomization_TargetLayers.cpp
@@ -21,6 +21,8 @@
 #include "Widgets/Input/SMultiLineEditableTextBox.h"
 #include "Widgets/SBoxPanel.h"
 #include "Widgets/Text/SInlineEditableTextBlock.h"
+#include "Dialog/SCustomDialog.h"
+#include "SSimpleButton.h"
 #include "EditorModeManager.h"
 #include "EditorModes.h"
 #include "LandscapeEditorModule.h"
@@ -38,6 +40,7 @@
 
 #include "SLandscapeEditor.h"
 #include "Dialogs/DlgPickAssetPath.h"
+#include "Dialogs/DlgPickPath.h"
 #include "ObjectTools.h"
 #include "ScopedTransaction.h"
 #include "DesktopPlatformModule.h"
@@ -49,6 +52,130 @@
 
 #define LOCTEXT_NAMESPACE "LandscapeEditor.TargetLayers"
 
+TSharedRef<IPropertyTypeCustomization> FLandscapeEditorStructCustomization_FTargetLayerAssetPath::MakeInstance()
+{
+	return MakeShareable(new FLandscapeEditorStructCustomization_FTargetLayerAssetPath);
+}
+
+void FLandscapeEditorStructCustomization_FTargetLayerAssetPath::CustomizeHeader(TSharedRef<IPropertyHandle> StructPropertyHandle, FDetailWidgetRow& HeaderRow, IPropertyTypeCustomizationUtils& StructCustomizationUtils)
+{
+	const TSharedRef<IPropertyHandle> PropertyHandle_TargetLayerDirectoryPath = StructPropertyHandle->GetChildHandle(GET_MEMBER_NAME_CHECKED(FLandscapeTargetLayerAssetFilePath, DirectoryPath)).ToSharedRef();
+
+	HeaderRow
+		.NameContent()
+		[
+			SNew(SHorizontalBox)
+				+ SHorizontalBox::Slot()
+				.AutoWidth()
+				.Padding(5.0f, 0)
+				[
+					SNew(SCheckBox)
+						.IsChecked_Static(&FLandscapeEditorStructCustomization_FTargetLayerAssetPath::GetUseTargetLayerAssetPathCheckState)
+						.OnCheckStateChanged_Static(&FLandscapeEditorStructCustomization_FTargetLayerAssetPath::OnUseTargetLayerAssetPathCheckStateChanged)
+						.ToolTipText(LOCTEXT("TargetLayerAssetPathCheckbox_ToolTip", "Enable to override the default asset path"))
+				]
+				+ SHorizontalBox::Slot()
+				.FillWidth(1)
+				.HAlign(HAlign_Fill)
+				[
+					PropertyHandle_TargetLayerDirectoryPath->CreatePropertyNameWidget(LOCTEXT("TargetLayerAssetPath", "Default Layer Asset Path"), LOCTEXT("TargetLayerAssetPath_ToolTip", "Set the default Target Layer asset folder"))
+				]
+		]
+		.ValueContent()
+		.HAlign(HAlign_Fill)
+		.VAlign(VAlign_Center)
+		[
+			SNew(SHorizontalBox)
+				+ SHorizontalBox::Slot()
+				.FillWidth(1)
+				.VAlign(VAlign_Center)
+				[
+					SNew(SEditableTextBox)
+						.Font(IDetailLayoutBuilder::GetDetailFont())
+						.Text_Lambda(&FLandscapeEditorStructCustomization_FTargetLayerAssetPath::GetTargetLayerAssetFilePath)
+						.HintText(LOCTEXT("TargetLayerAssetPath_Hint", "(Specify a default path)"))
+						.IsEnabled_Static(&FLandscapeEditorStructCustomization_FTargetLayerAssetPath::IsUseTargetLayerAssetPathEnabled)
+						.IsReadOnly(true)
+				]
+				+ SHorizontalBox::Slot()
+				.HAlign(HAlign_Right)
+				.VAlign(VAlign_Center)
+				.AutoWidth()
+				[
+					SNew(SSimpleButton)
+						.OnClicked_Static(&FLandscapeEditorStructCustomization_FTargetLayerAssetPath::OnSetTargetLayerAssetFilePath)
+						.IsEnabled_Static(&FLandscapeEditorStructCustomization_FTargetLayerAssetPath::IsUseTargetLayerAssetPathEnabled)
+						.Icon(FAppStyle::Get().GetBrush("Icons.FolderOpen"))
+				]
+		];
+}
+
+void FLandscapeEditorStructCustomization_FTargetLayerAssetPath::CustomizeChildren(TSharedRef<IPropertyHandle> StructPropertyHandle, class IDetailChildrenBuilder& ChildBuilder, IPropertyTypeCustomizationUtils& StructCustomizationUtils)
+{
+	// Hide struct fields since the UI is all handled in the Header
+	const TSharedRef<IPropertyHandle> PropertyHandle_TargetLayerFilePath = StructPropertyHandle->GetChildHandle(GET_MEMBER_NAME_CHECKED(FLandscapeTargetLayerAssetFilePath, DirectoryPath)).ToSharedRef();
+	PropertyHandle_TargetLayerFilePath->MarkHiddenByCustomization();
+}
+
+FText FLandscapeEditorStructCustomization_FTargetLayerAssetPath::GetTargetLayerAssetFilePath() 
+{
+	const FEdModeLandscape* LandscapeEdMode = GetEditorMode();
+	if (LandscapeEdMode && LandscapeEdMode->UISettings)
+	{
+		return FText::FromString(LandscapeEdMode->GetTargetLayerAssetPackagePath());
+	}
+
+	return FText();
+}
+
+FReply FLandscapeEditorStructCustomization_FTargetLayerAssetPath::OnSetTargetLayerAssetFilePath()
+{
+	const TSharedRef<SDlgPickPath> NewPathDlg =
+		SNew(SDlgPickPath)
+		.Title(LOCTEXT("TargetLayerAssetFilePath_Dlg", "Set Target Layer Asset File Path"));
+
+	if (NewPathDlg->ShowModal() != EAppReturnType::Cancel)
+	{
+		const FEdModeLandscape* LandscapeEdMode = GetEditorMode();
+		if (LandscapeEdMode && LandscapeEdMode->UISettings)
+		{
+			LandscapeEdMode->UISettings->TargetLayerAssetFilePath.DirectoryPath.Path = NewPathDlg->GetPath().ToString();
+		}
+	}
+
+	return FReply::Handled();
+}
+
+bool FLandscapeEditorStructCustomization_FTargetLayerAssetPath::IsUseTargetLayerAssetPathEnabled()
+{
+	const FEdModeLandscape* LandscapeEdMode = GetEditorMode();
+	if (LandscapeEdMode && LandscapeEdMode->UISettings)
+	{
+		return LandscapeEdMode->UISettings->TargetLayerAssetFilePath.bUseAssetDirectoryPath;
+	}
+
+	return false;
+}
+
+ECheckBoxState FLandscapeEditorStructCustomization_FTargetLayerAssetPath::GetUseTargetLayerAssetPathCheckState()
+{
+	return IsUseTargetLayerAssetPathEnabled() ? ECheckBoxState::Checked : ECheckBoxState::Unchecked;
+}
+
+void FLandscapeEditorStructCustomization_FTargetLayerAssetPath::OnUseTargetLayerAssetPathCheckStateChanged(ECheckBoxState NewCheckedState)
+{
+	const FEdModeLandscape* LandscapeEdMode = GetEditorMode();
+	if (LandscapeEdMode && LandscapeEdMode->UISettings)
+	{
+		if (NewCheckedState != ECheckBoxState::Undetermined)
+		{
+			LandscapeEdMode->UISettings->TargetLayerAssetFilePath.bUseAssetDirectoryPath = (NewCheckedState == ECheckBoxState::Checked);
+		}
+	}
+}
+
+//////////////////////////////////////////////////////////////////////////
+
 TSharedRef<IDetailCustomization> FLandscapeEditorDetailCustomization_TargetLayers::MakeInstance()
 {
 	return MakeShareable(new FLandscapeEditorDetailCustomization_TargetLayers);
@@ -58,6 +185,7 @@ BEGIN_SLATE_FUNCTION_BUILD_OPTIMIZATION
 void FLandscapeEditorDetailCustomization_TargetLayers::CustomizeDetails(IDetailLayoutBuilder& DetailBuilder)
 {
 	TSharedRef<IPropertyHandle> PropertyHandle_PaintingRestriction = DetailBuilder.GetProperty(GET_MEMBER_NAME_CHECKED(ULandscapeEditorObject, PaintingRestriction));
+	TSharedRef<IPropertyHandle> PropertyHandle_TargetLayerAssetFilePath = DetailBuilder.GetProperty(GET_MEMBER_NAME_CHECKED(ULandscapeEditorObject, TargetLayerAssetFilePath));
 	TSharedRef<IPropertyHandle> PropertyHandle_TargetDisplayOrder = DetailBuilder.GetProperty(GET_MEMBER_NAME_CHECKED(ULandscapeEditorObject, TargetDisplayOrder));
 	PropertyHandle_TargetDisplayOrder->MarkHiddenByCustomization();
 
@@ -67,7 +195,7 @@ void FLandscapeEditorDetailCustomization_TargetLayers::CustomizeDetails(IDetailL
 	if (!ShouldShowTargetLayers())
 	{
 		PropertyHandle_PaintingRestriction->MarkHiddenByCustomization();
-
+		PropertyHandle_TargetLayerAssetFilePath->MarkHiddenByCustomization();
 		return;
 	}
 
@@ -79,6 +207,8 @@ void FLandscapeEditorDetailCustomization_TargetLayers::CustomizeDetails(IDetailL
 		.Visibility(TAttribute<EVisibility>::Create(TAttribute<EVisibility>::FGetter::CreateStatic(&FLandscapeEditorDetailCustomization_TargetLayers::GetPaintingRestrictionVisibility)))
 		.IsEnabled(TAttribute<bool>::Create(TAttribute<bool>::FGetter::CreateLambda([LandscapeEdMode]() { return LandscapeEdMode->HasValidLandscapeEditLayerSelection(); })));
 
+	TargetsCategory.AddProperty(PropertyHandle_TargetLayerAssetFilePath);
+
 	TargetsCategory.AddCustomRow(FText())
 		.Visibility(TAttribute<EVisibility>::Create(TAttribute<EVisibility>::FGetter::CreateStatic(&FLandscapeEditorDetailCustomization_TargetLayers::GetVisibilityMaskTipVisibility)))
 		[
@@ -281,6 +411,7 @@ void FLandscapeEditorCustomNodeBuilder_TargetLayers::GenerateHeaderRowContent(FD
 				.HAlign(HAlign_Center)
 				.ToolTipText(LOCTEXT("TargetLayerSortButtonTooltip", "Define how we want to sort the displayed layers"))
 				.OnGetMenuContent(this, &FLandscapeEditorCustomNodeBuilder_TargetLayers::GetTargetLayerDisplayOrderButtonMenuContent)
+				.IsEnabled(this, &FLandscapeEditorCustomNodeBuilder_TargetLayers::HasWeightmapTargetLayers)
 				.ButtonContent()
 				[
 					SNew(SHorizontalBox)
@@ -316,6 +447,7 @@ void FLandscapeEditorCustomNodeBuilder_TargetLayers::GenerateHeaderRowContent(FD
 				.HAlign(HAlign_Center)
 				.ToolTipText(LOCTEXT("TargetLayerUnusedLayerButtonTooltip", "Define if we want to display unused layers"))
 				.OnGetMenuContent(this, &FLandscapeEditorCustomNodeBuilder_TargetLayers::GetTargetLayerShowUnusedButtonMenuContent)
+				.IsEnabled(this, &FLandscapeEditorCustomNodeBuilder_TargetLayers::HasWeightmapTargetLayers)
 				.ButtonContent()
 				[
 					SNew(SHorizontalBox)
@@ -349,32 +481,42 @@ void FLandscapeEditorCustomNodeBuilder_TargetLayers::GenerateHeaderRowContent(FD
 			[
 				SNew(SButton)
 				.ButtonStyle( FAppStyle::Get(), "SimpleButton" )
-				.ToolTipText(NSLOCTEXT("Landscape", "CreateLayersFromMaterials", "Create Layers From Assigned Materials"))
+				.ToolTipText_Lambda([this]() { return HasUnassignedTargetLayers() ? LOCTEXT("TargetLayerCreateFromMaterialsToolTip", "Create Layers From Assigned Materials") : LOCTEXT("TargetLayerCreateFromMaterialsDisabledToolTip", "All Material Layers Created"); })
 				.OnClicked(this, &FLandscapeEditorCustomNodeBuilder_TargetLayers::HandleCreateLayersFromMaterials)
+				.IsEnabled(this, &FLandscapeEditorCustomNodeBuilder_TargetLayers::HasUnassignedTargetLayers)
 				[
 					SNew(SImage)
 					.Image(FAppStyle::GetBrush("LandscapeEditor.Layer.Sync"))
 				]
 			]
+
+			+ SHorizontalBox::Slot()
+			.AutoWidth()
+			.VAlign(VAlign_Center)
+			.Padding(5.0f, 0.0f, 0.0f, 0.0f)
+			[
+				SNew(SButton)
+				.ButtonStyle(FAppStyle::Get(), "SimpleButton")
+				.ToolTipText(LOCTEXT("TargetLayerAutoFillLayers", "Auto-Fill Target Layer Assets"))
+				.OnClicked(this, &FLandscapeEditorCustomNodeBuilder_TargetLayers::ShowAutoFillTargetLayerDialog)
+				.IsEnabled(this, &FLandscapeEditorCustomNodeBuilder_TargetLayers::HasWeightmapTargetLayers)
+				[
+					SNew(SImage)
+						.Image(FAppStyle::GetBrush("LandscapeEditor.Layer.AutoFill"))
+				]
+			]
 		];
 	}
 	NodeRow.IsEnabled(TAttribute<bool>::Create(TAttribute<bool>::FGetter::CreateLambda([LandscapeEdMode]() { return LandscapeEdMode->HasValidLandscapeEditLayerSelection(); })));
 }
 
-
-
-FReply FLandscapeEditorCustomNodeBuilder_TargetLayers::HandleCreateLayersFromMaterials()
+const TArray<const TSharedRef<FLandscapeTargetListInfo>> FLandscapeEditorCustomNodeBuilder_TargetLayers::GetUnassignedTargetLayersFromMaterial() const
 {
-	FScopedTransaction Transaction(LOCTEXT("LandscapeTargetLayer_CreateFromMaterials", "Create Target Layers from Assigned materials"));
-	
-	FEdModeLandscape* LandscapeEdMode = GetEditorMode(); 
-	if (LandscapeEdMode == nullptr)
-	{
-		return FReply::Handled(); 
-	}
-	
+	const FEdModeLandscape* LandscapeEdMode = GetEditorMode();
+	check(LandscapeEdMode);
+
 	ALandscape* LandscapeActor = LandscapeEdMode->CurrentToolTarget.LandscapeInfo->LandscapeActor.Get();
-	
+
 	TSet<FName> LayerNames;
 	LandscapeActor->GetLandscapeInfo()->ForEachLandscapeProxy([&LayerNames](ALandscapeProxy* Proxy)
 	{
@@ -382,18 +524,51 @@ FReply FLandscapeEditorCustomNodeBuilder_TargetLayers::HandleCreateLayersFromMat
 		return true;
 	});
 
-	UE::Landscape::FLayerInfoFinder LayerInfoFinder;
+	TArray<const TSharedRef<FLandscapeTargetListInfo>> TargetLayerList;
+
 	for (const FName& LayerName : LayerNames)
 	{
 		if (!LandscapeActor->GetTargetLayers().Find(LayerName))
 		{
-			ULandscapeLayerInfoObject* LandscapeLayerInfo = LayerInfoFinder.Find(LayerName);
-			LandscapeActor->AddTargetLayer(LayerName, FLandscapeTargetLayerSettings(LandscapeLayerInfo));
+			const FLandscapeInfoLayerSettings LayerSettings = FLandscapeInfoLayerSettings(LayerName, LandscapeActor);
+			const TSharedRef<FLandscapeTargetListInfo> Target = MakeShareable(new FLandscapeTargetListInfo(FText::FromName(LayerName), ELandscapeToolTargetType::Weightmap, LayerSettings, LandscapeEdMode->UISettings->CurrentLayerIndex));
+
+			TargetLayerList.Add(Target);
 		}
 	}
+
+	return TargetLayerList;
+}
+
+bool FLandscapeEditorCustomNodeBuilder_TargetLayers::HasUnassignedTargetLayers() const
+{
+	return !GetUnassignedTargetLayersFromMaterial().IsEmpty();
+}
+
+FReply FLandscapeEditorCustomNodeBuilder_TargetLayers::HandleCreateLayersFromMaterials() const
+{
+	FScopedTransaction Transaction(LOCTEXT("LandscapeTargetLayer_CreateFromMaterials", "Create Target Layers from Assigned materials"));
 	
-	if (LandscapeEdMode)
+	FEdModeLandscape* LandscapeEdMode = GetEditorMode(); 
+	if (LandscapeEdMode == nullptr)
+	{
+		return FReply::Handled(); 
+	}
+
+	ALandscape* LandscapeActor = LandscapeEdMode->CurrentToolTarget.LandscapeInfo->LandscapeActor.Get();
+	const TArray<const TSharedRef<FLandscapeTargetListInfo>> TargetLayerList = GetUnassignedTargetLayersFromMaterial();
+
+	if (!TargetLayerList.IsEmpty()) 
 	{
+		for (const TSharedRef<FLandscapeTargetListInfo>& TargetLayer : TargetLayerList)
+		{
+			// The user may have created a new layer in the dialog, only add layers that are not yet assigned
+			if (!LandscapeActor->GetTargetLayers().Find(TargetLayer->GetLayerName()))
+			{
+				LandscapeActor->AddTargetLayer(TargetLayer->GetLayerName(), FLandscapeTargetLayerSettings(TargetLayer->LayerInfoObj.Get()));
+			}
+		}
+
 		LandscapeEdMode->GetLandscape()->GetLandscapeInfo()->UpdateLayerInfoMap();
 		LandscapeEdMode->UpdateTargetList();
 	}
@@ -401,6 +576,93 @@ FReply FLandscapeEditorCustomNodeBuilder_TargetLayers::HandleCreateLayersFromMat
 	return FReply::Handled();
 }
 
+
+FReply FLandscapeEditorCustomNodeBuilder_TargetLayers::ShowAutoFillTargetLayerDialog() const
+{
+	const FEdModeLandscape* LandscapeEdMode = GetEditorMode();
+	check(LandscapeEdMode);
+
+	const FString& DefaultTargetLayerAssetPath = LandscapeEdMode->GetTargetLayerAssetPackagePath();
+
+	const TSharedPtr<SCustomDialog> DialogWindow = SNew(SCustomDialog)
+		.Title(LOCTEXT("TargetLayerAutoFillLayers", "Auto-Fill Target Layer Assets"))
+		.Buttons({
+			SCustomDialog::FButton(LOCTEXT("Unassigned Layers Only", "Unassigned Layers Only")).SetPrimary(true),
+			SCustomDialog::FButton(LOCTEXT("All Layers", "All Layers")),
+			SCustomDialog::FButton(LOCTEXT("Cancel", "Cancel"))
+			})
+		.ContentAreaPadding(10.0f)
+		.Content()
+		[
+			SNew(SVerticalBox)
+				+ SVerticalBox::Slot()
+				[
+					SNew(STextBlock)
+						.Text(LOCTEXT("TargetLayerAutoFillDialog_Begin", "This operation will assign the layer info assets found in the default asset folder"))
+				]
+				+ SVerticalBox::Slot()
+				.Padding(0, 2.0f)
+				[
+					SNew(SEditableTextBox)
+						.Font(IDetailLayoutBuilder::GetDetailFont())
+						.Text(FText::FromString(DefaultTargetLayerAssetPath))
+						.Padding(2.0f)
+						.IsReadOnly(true)
+				]
+				+ SVerticalBox::Slot()
+				[
+					SNew(STextBlock)
+						.Text(LOCTEXT("LandscapeTargetLayerAutoFillDialog_End", " to the landscape's target layers."))
+				]
+				+ SVerticalBox::Slot()
+				.Padding(0, 20.0f, 0, 0) // Create spacing between details and prompt
+				[
+					SNew(STextBlock)
+						.Text(LOCTEXT("TargetLayerAutoFillDialog_Prompt", "Please specify which target layers should be set."))
+				]
+			
+		];
+
+	const int32 ButtonPressedIndex = DialogWindow->ShowModal();
+	
+	// -1 is cancel, 0 empty layers, 1 all layers
+	if (ButtonPressedIndex == 0 || ButtonPressedIndex == 1)
+	{
+		const bool bUpdateAllLayers = (ButtonPressedIndex == 1);
+		HandleAutoFillTargetLayers(bUpdateAllLayers);
+	}
+
+	return FReply::Handled();
+}
+
+FReply FLandscapeEditorCustomNodeBuilder_TargetLayers::HandleAutoFillTargetLayers(bool bUpdateAllLayers) const
+{
+	FScopedTransaction Transaction(LOCTEXT("LandscapeTargetLayer_AutoFillTargetLayers", "Auto-Fill Target Layer Assets"));
+
+	const FEdModeLandscape* LandscapeEdMode = GetEditorMode();
+	if (LandscapeEdMode == nullptr)
+	{
+		return FReply::Handled();
+	}
+
+	for (const TSharedRef<FLandscapeTargetListInfo>& TargetInfo : PrepareTargetLayerList(/*bInSort = */true, /*bInFilter = */false))
+	{
+		// Auto fill unassigned weightmap target layers or all target layers when flag is true
+		if (TargetInfo->TargetType == ELandscapeToolTargetType::Weightmap && (!TargetInfo->LayerInfoObj.IsValid() || bUpdateAllLayers))
+		{
+			const FString TargetLayerAssetFilePath = LandscapeEdMode->GetTargetLayerAssetPackagePath();
+			const FAssetData* AssetData = LandscapeEditorUtils::FindLandscapeTargetLayerInfoAsset(TargetInfo->LayerName, TargetLayerAssetFilePath);
+			
+			if (AssetData != nullptr)
+			{
+				OnTargetLayerSetObject(*AssetData, TargetInfo);
+			}
+		}
+	}
+
+	return FReply::Handled();
+}
+
 void FLandscapeEditorCustomNodeBuilder_TargetLayers::HandleCreateLayer()
 {
 	FEdModeLandscape* LandscapeEdMode = GetEditorMode();
@@ -605,7 +867,7 @@ void FLandscapeEditorCustomNodeBuilder_TargetLayers::OnFilterTextCommitted(const
 	}
 }
 
-EVisibility FLandscapeEditorCustomNodeBuilder_TargetLayers::GetLayersFilterVisibility() const
+bool FLandscapeEditorCustomNodeBuilder_TargetLayers::HasWeightmapTargetLayers() const
 {
 	FEdModeLandscape* LandscapeEdMode = GetEditorMode();
 
@@ -615,17 +877,20 @@ EVisibility FLandscapeEditorCustomNodeBuilder_TargetLayers::GetLayersFilterVisib
 		if ((LandscapeEdMode->CurrentToolTarget.TargetType == ELandscapeToolTargetType::Weightmap)
 			|| (LandscapeEdMode->CurrentToolTarget.TargetType == ELandscapeToolTargetType::Invalid))
 		{
-			const bool bContainsWeightmapLayers = LandscapeEdMode->GetTargetList().ContainsByPredicate([](TSharedRef<FLandscapeTargetListInfo> InInfo)
+			return LandscapeEdMode->GetTargetList().ContainsByPredicate([](TSharedRef<FLandscapeTargetListInfo> InInfo)
 				{
 					FName LayerName = InInfo->GetLayerName();
 					return (LayerName != NAME_None) && (LayerName != UMaterialExpressionLandscapeVisibilityMask::ParameterName);
 				});
-
-			return bContainsWeightmapLayers ? EVisibility::Visible : EVisibility::Collapsed;
 		}
 	}
 
-	return EVisibility::Collapsed;
+	return false;
+}
+
+EVisibility FLandscapeEditorCustomNodeBuilder_TargetLayers::GetLayersFilterVisibility() const
+{
+	return HasWeightmapTargetLayers() ? EVisibility::Visible : EVisibility::Collapsed;
 }
 
 FText FLandscapeEditorCustomNodeBuilder_TargetLayers::GetLayersFilterText() const
@@ -885,6 +1150,12 @@ TSharedPtr<SWidget> FLandscapeEditorCustomNodeBuilder_TargetLayers::GenerateRow(
 													return false;
 												}
 
+												if (NewName.IsNone())
+												{
+													OutErrorMessage = LOCTEXT("LandscapeTargetLayer_RenameFailed_ReservedName", "Target layer name cannot be empty");
+													return false;
+												}
+
 												ALandscape* Landscape = Cast<ALandscape>(Target->Owner);
 												if (Landscape->HasTargetLayer(NewName))
 												{
@@ -1579,11 +1850,8 @@ void FLandscapeEditorCustomNodeBuilder_TargetLayers::OnTargetLayerCreateClicked(
 	if (LandscapeEdMode)
 	{
 		FName LayerName = Target->GetLayerName();
-		ULevel* Level = Target->Owner->GetLevel();
-
-		// Build default layer object name and package name
 		FName LayerObjectName;
-		FString PackageName = UE::Landscape::GetLayerInfoObjectPackageName(Level, LayerName, LayerObjectName);
+		FString PackageName = UE::Landscape::GetLayerInfoObjectPackageName(LayerName, LandscapeEdMode->GetTargetLayerAssetPackagePath(), LayerObjectName);
 
 		TSharedRef<SDlgPickAssetPath> NewLayerDlg =
 			SNew(SDlgPickAssetPath)
diff --git a/Engine/Source/Editor/LandscapeEditor/Private/LandscapeEditorDetailCustomization_TargetLayers.h b/Engine/Source/Editor/LandscapeEditor/Private/LandscapeEditorDetailCustomization_TargetLayers.h
--- a/Engine/Source/Editor/LandscapeEditor/Private/LandscapeEditorDetailCustomization_TargetLayers.h
+++ b/Engine/Source/Editor/LandscapeEditor/Private/LandscapeEditorDetailCustomization_TargetLayers.h
@@ -27,6 +27,23 @@ class SDragAndDropVerticalBox;
  * Slate widgets customizer for the target layers list in the Landscape Editor
  */
 
+class FLandscapeEditorStructCustomization_FTargetLayerAssetPath : public FLandscapeEditorStructCustomization_Base
+{
+public:
+	static TSharedRef<IPropertyTypeCustomization> MakeInstance();
+
+	virtual void CustomizeHeader(TSharedRef<IPropertyHandle> StructPropertyHandle, class FDetailWidgetRow& HeaderRow, IPropertyTypeCustomizationUtils& StructCustomizationUtils) override;
+	virtual void CustomizeChildren(TSharedRef<IPropertyHandle> StructPropertyHandle, class IDetailChildrenBuilder& ChildBuilder, IPropertyTypeCustomizationUtils& StructCustomizationUtils) override;
+
+private: 
+	static FText GetTargetLayerAssetFilePath();
+	static FReply OnSetTargetLayerAssetFilePath();
+
+	static bool IsUseTargetLayerAssetPathEnabled();
+	static ECheckBoxState GetUseTargetLayerAssetPathCheckState();
+	static void OnUseTargetLayerAssetPathCheckStateChanged(ECheckBoxState NewCheckedState);
+};
+
 class FLandscapeEditorDetailCustomization_TargetLayers : public FLandscapeEditorDetailCustomization_Base
 {
 public:
@@ -108,7 +125,15 @@ protected:
 	TOptional<SDragAndDropVerticalBox::EItemDropZone> HandleCanAcceptDrop(const FDragDropEvent& DragDropEvent, SDragAndDropVerticalBox::EItemDropZone DropZone, SVerticalBox::FSlot* Slot);
 	FReply HandleAcceptDrop(FDragDropEvent const& DragDropEvent, SDragAndDropVerticalBox::EItemDropZone DropZone, int32 SlotIndex, SVerticalBox::FSlot* Slot);
 
-	FReply HandleCreateLayersFromMaterials();
+	// Create Layers From Materials
+	const TArray<const TSharedRef<FLandscapeTargetListInfo>> GetUnassignedTargetLayersFromMaterial() const;
+	bool HasUnassignedTargetLayers() const;
+	FReply HandleCreateLayersFromMaterials() const;
+
+	// Auto-Fill Target layers
+	FReply ShowAutoFillTargetLayerDialog() const;
+	FReply HandleAutoFillTargetLayers(bool bUpdateAllLayers) const;
+
 	void HandleCreateLayer();
 	const FSlateBrush* GetTargetLayerDisplayOrderBrush() const;
 	TSharedRef<SWidget> GetTargetLayerDisplayOrderButtonMenuContent();
@@ -123,6 +148,7 @@ protected:
 
 	void OnFilterTextChanged(const FText& InFilterText);
 	void OnFilterTextCommitted(const FText& InFilterText, ETextCommit::Type InCommitType);
+	bool HasWeightmapTargetLayers() const;
 	EVisibility GetLayersFilterVisibility() const;
 	FText GetLayersFilterText() const;
 };
diff --git a/Engine/Source/Editor/LandscapeEditor/Private/LandscapeEditorModule.cpp b/Engine/Source/Editor/LandscapeEditor/Private/LandscapeEditorModule.cpp
--- a/Engine/Source/Editor/LandscapeEditor/Private/LandscapeEditorModule.cpp
+++ b/Engine/Source/Editor/LandscapeEditor/Private/LandscapeEditorModule.cpp
@@ -31,6 +31,7 @@
 #include "LandscapeEditorDetailCustomization_NewLandscape.h"
 #include "LandscapeEditorDetailCustomization_CopyPaste.h"
 #include "LandscapeEditorDetailCustomization_ImportLayers.h"
+#include "LandscapeEditorDetailCustomization_TargetLayers.h"
 #include "LandscapeSplineDetails.h"
 #include "LandscapeModule.h"
 
@@ -88,6 +89,7 @@ public:
 		PropertyModule.RegisterCustomClassLayout("LandscapeEditorObject", FOnGetDetailCustomizationInstance::CreateStatic(&FLandscapeEditorDetails::MakeInstance));
 		PropertyModule.RegisterCustomPropertyTypeLayout("GizmoImportLayer", FOnGetPropertyTypeCustomizationInstance::CreateStatic(&FLandscapeEditorStructCustomization_FGizmoImportLayer::MakeInstance));
 		PropertyModule.RegisterCustomPropertyTypeLayout("LandscapeImportLayer", FOnGetPropertyTypeCustomizationInstance::CreateStatic(&FLandscapeEditorStructCustomization_FLandscapeImportLayer::MakeInstance));
+		PropertyModule.RegisterCustomPropertyTypeLayout("LandscapeTargetLayerAssetFilePath", FOnGetPropertyTypeCustomizationInstance::CreateStatic(&FLandscapeEditorStructCustomization_FTargetLayerAssetPath::MakeInstance));
 
 		PropertyModule.RegisterCustomClassLayout("LandscapeSplineControlPoint", FOnGetDetailCustomizationInstance::CreateStatic(&FLandscapeSplineDetails::MakeInstance));
 		PropertyModule.RegisterCustomClassLayout("LandscapeSplineSegment", FOnGetDetailCustomizationInstance::CreateStatic(&FLandscapeSplineDetails::MakeInstance));
@@ -173,6 +175,7 @@ public:
 		PropertyModule.UnregisterCustomClassLayout("LandscapeEditorObject");
 		PropertyModule.UnregisterCustomPropertyTypeLayout("GizmoImportLayer");
 		PropertyModule.UnregisterCustomPropertyTypeLayout("LandscapeImportLayer");
+		PropertyModule.UnregisterCustomPropertyTypeLayout("LandscapeTargetLayerAssetFilePath");
 
 		PropertyModule.UnregisterCustomClassLayout("LandscapeSplineControlPoint");
 		PropertyModule.UnregisterCustomClassLayout("LandscapeSplineSegment");
diff --git a/Engine/Source/Editor/LandscapeEditor/Private/LandscapeEditorObject.cpp b/Engine/Source/Editor/LandscapeEditor/Private/LandscapeEditorObject.cpp
--- a/Engine/Source/Editor/LandscapeEditor/Private/LandscapeEditorObject.cpp
+++ b/Engine/Source/Editor/LandscapeEditor/Private/LandscapeEditorObject.cpp
@@ -214,6 +214,8 @@ void ULandscapeEditorObject::Load()
 	GConfig->GetInt(TEXT("LandscapeEdit"), TEXT("BrushComponentSize"), BrushComponentSize, GEditorPerProjectIni);
 	GConfig->GetFloat(TEXT("LandscapeEdit"), TEXT("BrushFalloff"), BrushFalloff, GEditorPerProjectIni);
 	GConfig->GetFloat(TEXT("LandscapeEdit"), TEXT("PaintBrushFalloff"), PaintBrushFalloff, GEditorPerProjectIni);
+	GConfig->GetString(TEXT("LandscapeEdit"), TEXT("TargetLayerAssetFilePath"), TargetLayerAssetFilePath.DirectoryPath.Path, GEditorPerProjectIni);
+	GConfig->GetBool(TEXT("LandscapeEdit"), TEXT("bUseTargetLayerAssetFilePath"), TargetLayerAssetFilePath.bUseAssetDirectoryPath, GEditorPerProjectIni);
 	bool InbUseClayBrush = bUseClayBrush;
 	GConfig->GetBool(TEXT("LandscapeEdit"), TEXT("bUseClayBrush"), InbUseClayBrush, GEditorPerProjectIni);
 	bUseClayBrush = InbUseClayBrush;
@@ -394,6 +396,8 @@ void ULandscapeEditorObject::Save()
 	GConfig->SetInt(TEXT("LandscapeEdit"), TEXT("BrushComponentSize"), BrushComponentSize, GEditorPerProjectIni);
 	GConfig->SetFloat(TEXT("LandscapeEdit"), TEXT("BrushFalloff"), BrushFalloff, GEditorPerProjectIni);
 	GConfig->SetFloat(TEXT("LandscapeEdit"), TEXT("PaintBrushFalloff"), PaintBrushFalloff, GEditorPerProjectIni);
+	GConfig->SetString(TEXT("LandscapeEdit"), TEXT("TargetLayerAssetFilePath"), *TargetLayerAssetFilePath.DirectoryPath.Path, GEditorPerProjectIni);
+	GConfig->SetBool(TEXT("LandscapeEdit"), TEXT("bUseTargetLayerAssetFilePath"), TargetLayerAssetFilePath.bUseAssetDirectoryPath, GEditorPerProjectIni);
 	GConfig->SetBool(TEXT("LandscapeEdit"), TEXT("bUseClayBrush"), bUseClayBrush, GEditorPerProjectIni);
 	GConfig->SetFloat(TEXT("LandscapeEdit"), TEXT("AlphaBrushScale"), AlphaBrushScale, GEditorPerProjectIni);
 	GConfig->SetBool(TEXT("LandscapeEdit"), TEXT("AlphaBrushAutoRotate"), bAlphaBrushAutoRotate, GEditorPerProjectIni);
diff --git a/Engine/Source/Editor/LandscapeEditor/Private/LandscapeEditorUtils.cpp b/Engine/Source/Editor/LandscapeEditor/Private/LandscapeEditorUtils.cpp
--- a/Engine/Source/Editor/LandscapeEditor/Private/LandscapeEditorUtils.cpp
+++ b/Engine/Source/Editor/LandscapeEditor/Private/LandscapeEditorUtils.cpp
@@ -19,6 +19,9 @@
 #include "WorldPartition/WorldPartitionHelpers.h"
 #include "Framework/Application/SlateApplication.h"
 #include "Misc/MessageDialog.h"
+#include "AssetRegistry/AssetData.h"
+#include "AssetRegistry/AssetRegistryModule.h"
+#include "Algo/LevenshteinDistance.h"
 
 namespace LandscapeEditorUtils
 {
@@ -94,6 +97,52 @@ namespace LandscapeEditorUtils
 		UWorldPartition::FDisableNonDirtyActorTrackingScope Scope(World->GetWorldPartition(), true);
 		LandscapeEditorUtils::SaveObjects(Proxies);
 	}
+
+	const TArray<FAssetData> GetLandscapeTargetLayerInfoAssets()
+	{
+		TArray<FAssetData> LayerInfoAssets;
+
+		const UClass* AssetClass = ULandscapeLayerInfoObject::StaticClass();
+		const FAssetRegistryModule& AssetRegistryModule = FModuleManager::LoadModuleChecked<FAssetRegistryModule>("AssetRegistry");
+		FARFilter Filter;
+		const FName PackageName = *AssetClass->GetPackage()->GetName();
+		const FName AssetName = AssetClass->GetFName();
+
+		Filter.ClassPaths.Add(FTopLevelAssetPath(PackageName, AssetName));
+		AssetRegistryModule.Get().GetAssets(Filter, LayerInfoAssets);
+
+		return LayerInfoAssets;
+	}
+
+	const FAssetData* FindLandscapeTargetLayerInfoAsset(const FName& LayerName, const FString& TargetLayerAssetPackagePath)
+	{
+		const TArray<FAssetData>& LayerInfoAssets = GetLandscapeTargetLayerInfoAssets();
+
+		int32 BestScore = -1.0f;
+		const FAssetData* BestLayerInfoObj = nullptr;
+
+		for (const FAssetData& LayerInfoAsset : LayerInfoAssets)
+		{
+			const ULandscapeLayerInfoObject* LayerInfo = CastChecked<ULandscapeLayerInfoObject>(LayerInfoAsset.GetAsset());
+
+			if (LayerInfo && LayerInfo->LayerName == LayerName)
+			{
+				const FString CurrentPackagePath = LayerInfoAsset.PackagePath.ToString();
+				float WorstCase = static_cast<float>(TargetLayerAssetPackagePath.Len() + CurrentPackagePath.Len());
+				WorstCase = FMath::Max(WorstCase, 1.f);
+
+				const float Score = 1.0f - (static_cast<float>(Algo::LevenshteinDistance(TargetLayerAssetPackagePath, CurrentPackagePath)) / WorstCase);
+				if (Score > BestScore)
+				{
+					BestScore = Score;
+					BestLayerInfoObj = &LayerInfoAsset;
+				}
+			}
+		}
+
+		return BestLayerInfoObj;
+	}
+
 }
 
 
diff --git a/Engine/Source/Editor/LandscapeEditor/Public/LandscapeEditorObject.h b/Engine/Source/Editor/LandscapeEditor/Public/LandscapeEditorObject.h
--- a/Engine/Source/Editor/LandscapeEditor/Public/LandscapeEditorObject.h
+++ b/Engine/Source/Editor/LandscapeEditor/Public/LandscapeEditorObject.h
@@ -259,6 +259,24 @@ struct FLandscapePatternBrushWorldSpaceSettings
 	{}
 };
 
+USTRUCT()
+struct FLandscapeTargetLayerAssetFilePath
+{
+	GENERATED_USTRUCT_BODY()
+
+	UPROPERTY(Category = "Target Layers", EditAnywhere, NonTransactional, meta=(InlineEditConditionToggle))
+	bool bUseAssetDirectoryPath = false;
+
+	UPROPERTY(Category = "Target Layers", EditAnywhere, NonTransactional, meta = (EditCondition = bUseAssetDirectoryPath, LongPackageName))
+	FDirectoryPath DirectoryPath;
+
+	FLandscapeTargetLayerAssetFilePath()
+		: bUseAssetDirectoryPath(false)
+		, DirectoryPath()
+	{
+	}
+};
+
 UENUM()
 enum class ELandscapeImportExportMode
 {
@@ -702,6 +720,9 @@ public:
 	UPROPERTY(Category = "Target Layers", EditAnywhere, NonTransactional, meta=(ShowForTargetTypes="Weightmap,Visibility"))
 	ELandscapeLayerPaintingRestriction PaintingRestriction;
 
+	UPROPERTY(Category = "Target Layers", EditAnywhere, NonTransactional, meta = (ShowForTargetTypes = "Weightmap"))
+	FLandscapeTargetLayerAssetFilePath TargetLayerAssetFilePath;
+
 	// This allows to hide layer names that do not match the filter
 	UPROPERTY(Transient)
 	FString TargetLayersFilterString;
diff --git a/Engine/Source/Editor/LandscapeEditor/Public/LandscapeEditorUtils.h b/Engine/Source/Editor/LandscapeEditor/Public/LandscapeEditorUtils.h
--- a/Engine/Source/Editor/LandscapeEditor/Public/LandscapeEditorUtils.h
+++ b/Engine/Source/Editor/LandscapeEditor/Public/LandscapeEditorUtils.h
@@ -26,4 +26,7 @@ namespace LandscapeEditorUtils
 	}
 
 	void SaveLandscapeProxies(UWorld* InWorld, TArrayView<ALandscapeProxy*> Proxies);
+
+	const TArray<FAssetData> GetLandscapeTargetLayerInfoAssets();
+	const FAssetData* FindLandscapeTargetLayerInfoAsset(const FName& LayerName, const FString& TargetLayerAssetPackagePath);
 }
diff --git a/Engine/Source/Runtime/Landscape/Classes/LandscapeProxy.h b/Engine/Source/Runtime/Landscape/Classes/LandscapeProxy.h
--- a/Engine/Source/Runtime/Landscape/Classes/LandscapeProxy.h
+++ b/Engine/Source/Runtime/Landscape/Classes/LandscapeProxy.h
@@ -1331,6 +1331,15 @@ PRAGMA_ENABLE_DEPRECATION_WARNINGS
 	// Query all the LandscapeLayerInfo Objects from the weight layers allocated on this proxy.
 	TMap<FName, ULandscapeLayerInfoObject*> RetrieveTargetLayerInfosFromAllocations() const;
 
+	/**
+	* Creates a new LandscapeLayerInfoObject
+	* @param	LayerName	The name of the created Layer
+	* @param	InFilePath	New asset file path, should be Current Level SharedAssetPath or the LandscapeEditorObject->TargetLayerAssetFilePath->DirectoryPath
+	* @param	InTemplate	If set, this object will be cloned as a basis for the new LandscapeLayerInfoObject.
+	* @return a new LandscapeLayerInfoObject if created successfully.
+	*/
+	LANDSCAPE_API static ULandscapeLayerInfoObject* CreateLayerInfo(const TCHAR* InLayerName, const FString& InFilePath, const ULandscapeLayerInfoObject* InTemplate = nullptr);
+
 	/**
 	* Creates a new LandscapeLayerInfoObject
 	* @param	LayerName	The name of the created Layer
@@ -1338,6 +1347,7 @@ PRAGMA_ENABLE_DEPRECATION_WARNINGS
 	* @param	InTemplate	If set, this object will be cloned as a basis for the new LandscapeLayerInfoObject.
 	* @return a new LandscapeLayerInfoObject if created successfully.
 	*/
+	UE_DEPRECATED(5.6, "This create function is deprecated. Please use new CreateLayerInfo and pass full asset file path.")
 	LANDSCAPE_API static ULandscapeLayerInfoObject* CreateLayerInfo(const TCHAR* InLayerName, const ULevel* InLevel, const ULandscapeLayerInfoObject* InTemplate = nullptr);
 
 	LANDSCAPE_API ULandscapeLayerInfoObject* CreateLayerInfo(const TCHAR* InLayerName, const ULandscapeLayerInfoObject* InTemplate = nullptr);
diff --git a/Engine/Source/Runtime/Landscape/Private/LandscapeEdit.cpp b/Engine/Source/Runtime/Landscape/Private/LandscapeEdit.cpp
--- a/Engine/Source/Runtime/Landscape/Private/LandscapeEdit.cpp
+++ b/Engine/Source/Runtime/Landscape/Private/LandscapeEdit.cpp
@@ -3223,10 +3223,10 @@ TMap<FName, ULandscapeLayerInfoObject*> ALandscapeProxy::RetrieveTargetLayerInfo
 	return InfoObjects;
 }
 
-ULandscapeLayerInfoObject* ALandscapeProxy::CreateLayerInfo(const TCHAR* InLayerName, const ULevel* InLevel, const ULandscapeLayerInfoObject* InTemplate)
+ULandscapeLayerInfoObject* ALandscapeProxy::CreateLayerInfo(const TCHAR* InLayerName, const FString& InFilePath, const ULandscapeLayerInfoObject* InTemplate)
 {
 	FName LayerObjectName;
-	FString PackageName = UE::Landscape::GetLayerInfoObjectPackageName(InLevel, InLayerName, LayerObjectName);
+	FString PackageName = UE::Landscape::GetLayerInfoObjectPackageName(InLayerName, InFilePath, LayerObjectName);
 	UPackage* Package = CreatePackage(*PackageName);
 	ULandscapeLayerInfoObject* LayerInfo = nullptr;
 	check(Package != nullptr);
@@ -3249,9 +3249,15 @@ ULandscapeLayerInfoObject* ALandscapeProxy::CreateLayerInfo(const TCHAR* InLayer
 	return LayerInfo;
 }
 
+// Deprecated
+ULandscapeLayerInfoObject* ALandscapeProxy::CreateLayerInfo(const TCHAR* InLayerName, const ULevel* InLevel, const ULandscapeLayerInfoObject* InTemplate)
+{
+	return CreateLayerInfo(InLayerName, UE::Landscape::GetSharedAssetsPath(InLevel), InTemplate);
+}
+
 ULandscapeLayerInfoObject* ALandscapeProxy::CreateLayerInfo(const TCHAR* InLayerName, const ULandscapeLayerInfoObject* InTemplate)
 {
-	ULandscapeLayerInfoObject* LayerInfo = ALandscapeProxy::CreateLayerInfo(InLayerName, GetLevel(), InTemplate);
+	ULandscapeLayerInfoObject* LayerInfo = ALandscapeProxy::CreateLayerInfo(InLayerName, UE::Landscape::GetSharedAssetsPath(GetLevel()), InTemplate);
 
 	check(LayerInfo);
 
diff --git a/Engine/Source/Runtime/Landscape/Private/LandscapeUtils.cpp b/Engine/Source/Runtime/Landscape/Private/LandscapeUtils.cpp
--- a/Engine/Source/Runtime/Landscape/Private/LandscapeUtils.cpp
+++ b/Engine/Source/Runtime/Landscape/Private/LandscapeUtils.cpp
@@ -210,20 +210,19 @@ FString GetSharedAssetsPath(const ULevel* InLevel)
 	return GetSharedAssetsPath(InLevel->GetOutermost()->GetName());
 }
 
-FString GetLayerInfoObjectPackageName(const ULevel* InLevel, const FName& InLayerName, FName& OutLayerObjectName)
+FString GetLayerInfoObjectPackageName(const FName& InLayerName, const FString& InPackagePath, FName& OutLayerObjectName)
 {
 	FString PackageName;
 	FString PackageFilename;
-	FString SharedAssetsPath = GetSharedAssetsPath(InLevel);
 	int32 Suffix = 1;
 
 	OutLayerObjectName = FName(*FString::Printf(TEXT("%s_LayerInfo"), *ObjectTools::SanitizeInvalidChars(*InLayerName.ToString(), INVALID_LONGPACKAGE_CHARACTERS)));
-	FPackageName::TryConvertFilenameToLongPackageName(SharedAssetsPath / OutLayerObjectName.ToString(), PackageName);
+	FPackageName::TryConvertFilenameToLongPackageName(InPackagePath / OutLayerObjectName.ToString(), PackageName);
 
 	while (FPackageName::DoesPackageExist(PackageName, &PackageFilename))
 	{
 		OutLayerObjectName = FName(*FString::Printf(TEXT("%s_LayerInfo_%d"), *ObjectTools::SanitizeInvalidChars(*InLayerName.ToString(), INVALID_LONGPACKAGE_CHARACTERS), Suffix));
-		if (!FPackageName::TryConvertFilenameToLongPackageName(SharedAssetsPath / OutLayerObjectName.ToString(), PackageName))
+		if (!FPackageName::TryConvertFilenameToLongPackageName(InPackagePath / OutLayerObjectName.ToString(), PackageName))
 		{
 			break;
 		}
@@ -234,6 +233,12 @@ FString GetLayerInfoObjectPackageName(const ULevel* InLevel, const FName& InLaye
 	return PackageName;
 }
 
+// Deprecated
+FString GetLayerInfoObjectPackageName(const ULevel* InLevel, const FName& InLayerName, FName& OutLayerObjectName)
+{
+	return GetLayerInfoObjectPackageName(InLayerName, GetSharedAssetsPath(InLevel), OutLayerObjectName);
+}
+
 bool IsVisibilityLayer(const ULandscapeLayerInfoObject* InLayerInfoObject)
 {
 	return (ALandscapeProxy::VisibilityLayer != nullptr) && (ALandscapeProxy::VisibilityLayer == InLayerInfoObject);
@@ -407,6 +412,7 @@ bool FBatchTextureCopy::ProcessTextureCopies()
 	return true;
 }
 
+PRAGMA_DISABLE_DEPRECATION_WARNINGS
 FLayerInfoFinder::FLayerInfoFinder()
 {
 	const UClass* AssetClass = ULandscapeLayerInfoObject::StaticClass();
@@ -431,7 +437,7 @@ ULandscapeLayerInfoObject* FLayerInfoFinder::Find(const FName& LayerName) const
 	}
 	return nullptr;
 }
-
+PRAGMA_ENABLE_DEPRECATION_WARNINGS
 
 UMaterialInstance* CreateToolLandscapeMaterialInstanceConstant(UMaterialInterface* BaseMaterial)
 {
diff --git a/Engine/Source/Runtime/Landscape/Public/LandscapeUtils.h b/Engine/Source/Runtime/Landscape/Public/LandscapeUtils.h
--- a/Engine/Source/Runtime/Landscape/Public/LandscapeUtils.h
+++ b/Engine/Source/Runtime/Landscape/Public/LandscapeUtils.h
@@ -293,6 +293,14 @@ LANDSCAPE_API FString GetSharedAssetsPath(const FString& InPath);
 */
 LANDSCAPE_API FString GetSharedAssetsPath(const ULevel* InLevel);
 
+/**
+ * Returns a generated package name for a Layer Info Object
+ * @param	InLayerName	The LayerName of the Layer Info Object
+ * @param	InPackagePath	Base package path that will be used, should be Current Level SharedAssetPath or the LandscapeEditorObject->TargetLayerAssetFilePath
+ * @param	OutLayerObjectName	The generated object name for Layer Info Object
+*/
+LANDSCAPE_API FString GetLayerInfoObjectPackageName(const FName& InLayerName, const FString& InPackagePath, FName& OutLayerObjectName);
+
 /**
  * Returns a generated package name for a Layer Info Object
  * @param	InLevel		Level's Path will be used as a basis to generate package's path. If /Temp/, it will be replaced by the last valid path used for level.
@@ -300,12 +308,13 @@ LANDSCAPE_API FString GetSharedAssetsPath(const ULevel* InLevel);
  * @param	OutLayerObjectName	The generated object name for Layer Info Object
  * @return
 */
+UE_DEPRECATED(5.6, "This Get function is deprecated. Please use new GetLayerInfo and pass in a full asset package path.")
 LANDSCAPE_API FString GetLayerInfoObjectPackageName(const ULevel* InLevel, const FName& InLayerName, FName& OutLayerObjectName);
 
 /** Returns true if the provided layer info object is the current visibility layer. */
 LANDSCAPE_API bool IsVisibilityLayer(const ULandscapeLayerInfoObject* InLayerInfoObject);
 
-struct FLayerInfoFinder
+struct UE_DEPRECATED(5.6, "This helper struct is deprecated. Please use utility methods in LandscapeEditorUtils.") FLayerInfoFinder
 {
 	LANDSCAPE_API FLayerInfoFinder();
 	LANDSCAPE_API ~FLayerInfoFinder() = default;
