Landscape Paint - Auto-Fill target layers: Added Create New Assets checkbox, improved dialog formatting, asset Find only checks within package path

diff --git a/Engine/Source/Editor/LandscapeEditor/Private/LandscapeEdMode.cpp b/Engine/Source/Editor/LandscapeEditor/Private/LandscapeEdMode.cpp
--- a/Engine/Source/Editor/LandscapeEditor/Private/LandscapeEdMode.cpp
+++ b/Engine/Source/Editor/LandscapeEditor/Private/LandscapeEdMode.cpp
@@ -2322,11 +2322,16 @@ const TArray<FLandscapeListInfo>& FEdModeLandscape::GetLandscapeList()
 	return LandscapeList;
 }
 
-const FString FEdModeLandscape::GetTargetLayerAssetPackagePath() const
+const FString FEdModeLandscape::GetTargetLayerAssetPackagePath(const bool bIsEmptyPathValid) const
 {
 	if (UISettings->TargetLayerAssetFilePath.bUseAssetDirectoryPath)
 	{
-		return UISettings->TargetLayerAssetFilePath.DirectoryPath.Path;
+		// When default path is enabled but no path has been chosen we should use the shared assets path
+		// UI methods can optionally return an empty path to show a HintText
+		if (UISettings->TargetLayerAssetFilePath.DirectoryPath.Path.IsEmpty() && bIsEmptyPathValid)
+		{
+			return UISettings->TargetLayerAssetFilePath.DirectoryPath.Path;
+		}
 	}
 
 	const ULevel* Level = GetWorld()->GetCurrentLevel();
diff --git a/Engine/Source/Editor/LandscapeEditor/Private/LandscapeEdMode.h b/Engine/Source/Editor/LandscapeEditor/Private/LandscapeEdMode.h
--- a/Engine/Source/Editor/LandscapeEditor/Private/LandscapeEdMode.h
+++ b/Engine/Source/Editor/LandscapeEditor/Private/LandscapeEdMode.h
@@ -558,7 +558,7 @@ public:
 	const TArray<FName>* GetTargetDisplayOrderList() const;
 	int32 GetTargetLayerStartingIndex() const;
 	const TArray<FLandscapeListInfo>& GetLandscapeList();
-	const FString GetTargetLayerAssetPackagePath() const;
+	const FString GetTargetLayerAssetPackagePath(const bool bIsEmptyPathValid = false) const;
 
 	int32 UpdateLandscapeList();
 	void UpdateTargetList(bool bRegenerateThumbnails = false);
diff --git a/Engine/Source/Editor/LandscapeEditor/Private/LandscapeEditorDetailCustomization_TargetLayers.cpp b/Engine/Source/Editor/LandscapeEditor/Private/LandscapeEditorDetailCustomization_TargetLayers.cpp
--- a/Engine/Source/Editor/LandscapeEditor/Private/LandscapeEditorDetailCustomization_TargetLayers.cpp
+++ b/Engine/Source/Editor/LandscapeEditor/Private/LandscapeEditorDetailCustomization_TargetLayers.cpp
@@ -92,7 +92,7 @@ void FLandscapeEditorStructCustomization_FTargetLayerAssetPath::CustomizeHeader(
 				[
 					SNew(SEditableTextBox)
 						.Font(IDetailLayoutBuilder::GetDetailFont())
-						.Text_Lambda(&FLandscapeEditorStructCustomization_FTargetLayerAssetPath::GetTargetLayerAssetFilePath)
+						.Text_Static(&FLandscapeEditorStructCustomization_FTargetLayerAssetPath::GetTargetLayerAssetFilePath)
 						.HintText(LOCTEXT("TargetLayerAssetPath_Hint", "(Specify a default path)"))
 						.IsEnabled_Static(&FLandscapeEditorStructCustomization_FTargetLayerAssetPath::IsUseTargetLayerAssetPathEnabled)
 						.IsReadOnly(true)
@@ -122,7 +122,7 @@ FText FLandscapeEditorStructCustomization_FTargetLayerAssetPath::GetTargetLayerA
 	const FEdModeLandscape* LandscapeEdMode = GetEditorMode();
 	if (LandscapeEdMode && LandscapeEdMode->UISettings)
 	{
-		return FText::FromString(LandscapeEdMode->GetTargetLayerAssetPackagePath());
+		return FText::FromString(LandscapeEdMode->GetTargetLayerAssetPackagePath(/* bIsEmptyPathValid */true));
 	}
 
 	return FText();
@@ -581,8 +581,8 @@ FReply FLandscapeEditorCustomNodeBuilder_TargetLayers::ShowAutoFillTargetLayerDi
 {
 	const FEdModeLandscape* LandscapeEdMode = GetEditorMode();
 	check(LandscapeEdMode);
-
-	const FString& DefaultTargetLayerAssetPath = LandscapeEdMode->GetTargetLayerAssetPackagePath();
+	
+	bool bIsCreateNewTargetLayersChecked = false;
 
 	const TSharedPtr<SCustomDialog> DialogWindow = SNew(SCustomDialog)
 		.Title(LOCTEXT("TargetLayerAutoFillLayers", "Auto-Fill Target Layer Assets"))
@@ -596,31 +596,39 @@ FReply FLandscapeEditorCustomNodeBuilder_TargetLayers::ShowAutoFillTargetLayerDi
 		[
 			SNew(SVerticalBox)
 				+ SVerticalBox::Slot()
+				.AutoHeight()
 				[
 					SNew(STextBlock)
-						.Text(LOCTEXT("TargetLayerAutoFillDialog_Begin", "This operation will assign the layer info assets found in the default asset folder"))
+						.WrapTextAt(512.0f)
+						.Text(FText::Format(LOCTEXT("TargetLayerAutoFillDialog", "This operation will assign the layer info assets found within the default asset folder {0} to the landscape's target layers."),
+							FText::FromString(LandscapeEdMode->GetTargetLayerAssetPackagePath())))
 				]
 				+ SVerticalBox::Slot()
-				.Padding(0, 2.0f)
+				.AutoHeight()
+				.Padding(0, 10.0f, 0, 0)
 				[
-					SNew(SEditableTextBox)
-						.Font(IDetailLayoutBuilder::GetDetailFont())
-						.Text(FText::FromString(DefaultTargetLayerAssetPath))
-						.Padding(2.0f)
-						.IsReadOnly(true)
-				]
-				+ SVerticalBox::Slot()
-				[
-					SNew(STextBlock)
-						.Text(LOCTEXT("LandscapeTargetLayerAutoFillDialog_End", " to the landscape's target layers."))
+					SNew(SHorizontalBox)
+						+ SHorizontalBox::Slot()
+						.AutoWidth()
+						[
+							SNew(SCheckBox)
+							.IsChecked_Lambda([&bIsCreateNewTargetLayersChecked]() { return bIsCreateNewTargetLayersChecked ? ECheckBoxState::Checked : ECheckBoxState::Unchecked; })
+							.OnCheckStateChanged_Lambda([&bIsCreateNewTargetLayersChecked](ECheckBoxState InState) { bIsCreateNewTargetLayersChecked = (InState == ECheckBoxState::Checked); })
+						]
+						+ SHorizontalBox::Slot()
+						.AutoWidth()
+						[
+							SNew(STextBlock)
+								.Text(LOCTEXT("TargetLayerAutoFillDialog_CreateNew", "Create new assets in the default folder if none are found."))
+						]
 				]
 				+ SVerticalBox::Slot()
+				.AutoHeight()
 				.Padding(0, 20.0f, 0, 0) // Create spacing between details and prompt
 				[
 					SNew(STextBlock)
 						.Text(LOCTEXT("TargetLayerAutoFillDialog_Prompt", "Please specify which target layers should be set."))
 				]
-			
 		];
 
 	const int32 ButtonPressedIndex = DialogWindow->ShowModal();
@@ -629,13 +637,13 @@ FReply FLandscapeEditorCustomNodeBuilder_TargetLayers::ShowAutoFillTargetLayerDi
 	if (ButtonPressedIndex == 0 || ButtonPressedIndex == 1)
 	{
 		const bool bUpdateAllLayers = (ButtonPressedIndex == 1);
-		HandleAutoFillTargetLayers(bUpdateAllLayers);
+		HandleAutoFillTargetLayers(bUpdateAllLayers, bIsCreateNewTargetLayersChecked);
 	}
 
 	return FReply::Handled();
 }
 
-FReply FLandscapeEditorCustomNodeBuilder_TargetLayers::HandleAutoFillTargetLayers(bool bUpdateAllLayers) const
+FReply FLandscapeEditorCustomNodeBuilder_TargetLayers::HandleAutoFillTargetLayers(const bool bUpdateAllLayers, const bool bCreateNewTargetLayers) const
 {
 	FScopedTransaction Transaction(LOCTEXT("LandscapeTargetLayer_AutoFillTargetLayers", "Auto-Fill Target Layer Assets"));
 
@@ -657,6 +665,13 @@ FReply FLandscapeEditorCustomNodeBuilder_TargetLayers::HandleAutoFillTargetLayer
 			{
 				OnTargetLayerSetObject(*AssetData, TargetInfo);
 			}
+			else if (bCreateNewTargetLayers)
+			{
+				FName LayerObjectName;
+				const FString PackageName = UE::Landscape::GetLayerInfoObjectPackageName(TargetInfo->LayerName, TargetLayerAssetFilePath, LayerObjectName);
+
+				CreateTargetLayerInfoAsset(TargetInfo, PackageName, LayerObjectName, /* bNoWeightBlend */ true);
+			}
 		}
 	}
 
@@ -1152,7 +1167,7 @@ TSharedPtr<SWidget> FLandscapeEditorCustomNodeBuilder_TargetLayers::GenerateRow(
 
 												if (NewName.IsNone())
 												{
-													OutErrorMessage = LOCTEXT("LandscapeTargetLayer_RenameFailed_ReservedName", "Target layer name cannot be empty");
+													OutErrorMessage = LOCTEXT("LandscapeTargetLayer_RenameFailed_EmptyName", "Target layer name cannot be empty");
 													return false;
 												}
 
@@ -1846,8 +1861,7 @@ void FLandscapeEditorCustomNodeBuilder_TargetLayers::OnTargetLayerCreateClicked(
 
 	FScopedTransaction Transaction(LOCTEXT("Undo_Create", "Creating New Landscape Layer"));
 
-	FEdModeLandscape* LandscapeEdMode = GetEditorMode();
-	if (LandscapeEdMode)
+	if (FEdModeLandscape* LandscapeEdMode = GetEditorMode())
 	{
 		FName LayerName = Target->GetLayerName();
 		FName LayerObjectName;
@@ -1863,56 +1877,63 @@ void FLandscapeEditorCustomNodeBuilder_TargetLayers::OnTargetLayerCreateClicked(
 			PackageName = NewLayerDlg->GetFullAssetPath().ToString();
 			LayerObjectName = FName(*NewLayerDlg->GetAssetName().ToString());
 
-			UPackage* Package = CreatePackage( *PackageName);
+			CreateTargetLayerInfoAsset(Target, PackageName, LayerObjectName, bNoWeightBlend);
+		}
+	}
+}
+
+void FLandscapeEditorCustomNodeBuilder_TargetLayers::CreateTargetLayerInfoAsset(const TSharedRef<FLandscapeTargetListInfo> Target, const FString& PackageName, const FName& LayerObjectName, const bool bNoWeightBlend)
+{
+	if (FEdModeLandscape* LandscapeEdMode = GetEditorMode())
+	{
+		const FName LayerName = Target->GetLayerName();
 
-			// Do not pass RF_Transactional to NewObject, or the asset will mark itself as garbage on Undo (which is not a well-supported path, potentially causing crashes)
-			ULandscapeLayerInfoObject* LayerInfo = NewObject<ULandscapeLayerInfoObject>(Package, LayerObjectName, RF_Public | RF_Standalone);
-			LayerInfo->SetFlags(RF_Transactional);	// we add RF_Transactional after creation, so that future edits _are_ recorded in undo
-			LayerInfo->LayerName = LayerName;
-			LayerInfo->bNoWeightBlend = bNoWeightBlend;
+		UPackage* Package = CreatePackage(*PackageName);
 
-			ULandscapeInfo* LandscapeInfo = Target->LandscapeInfo.Get();
-			LandscapeInfo->Modify();
-			int32 Index = LandscapeInfo->GetLayerInfoIndex(LayerName, Target->Owner.Get());
-			if (Index == INDEX_NONE)
-			{
-				LandscapeInfo->Layers.Add(FLandscapeInfoLayerSettings(LayerInfo, Target->Owner.Get()));
-			}
-			else
-			{
-				LandscapeInfo->Layers[Index].LayerInfoObj = LayerInfo;
-			}
+		// Do not pass RF_Transactional to NewObject, or the asset will mark itself as garbage on Undo (which is not a well-supported path, potentially causing crashes)
+		ULandscapeLayerInfoObject* LayerInfo = NewObject<ULandscapeLayerInfoObject>(Package, LayerObjectName, RF_Public | RF_Standalone);
+		LayerInfo->SetFlags(RF_Transactional);	// we add RF_Transactional after creation, so that future edits _are_ recorded in undo
+		LayerInfo->LayerName = LayerName;
+		LayerInfo->bNoWeightBlend = bNoWeightBlend;
 
-			if (LandscapeEdMode->CurrentToolTarget.LayerName == Target->LayerName
-				&& LandscapeEdMode->CurrentToolTarget.LayerInfo == Target->LayerInfoObj)
-			{
-				LandscapeEdMode->SetCurrentTargetLayer(Target->LayerName, Target->LayerInfoObj);
-			}
+		ULandscapeInfo* LandscapeInfo = Target->LandscapeInfo.Get();
+		LandscapeInfo->Modify();
+		int32 Index = LandscapeInfo->GetLayerInfoIndex(LayerName, Target->Owner.Get());
+		if (Index == INDEX_NONE)
+		{
+			LandscapeInfo->Layers.Add(FLandscapeInfoLayerSettings(LayerInfo, Target->Owner.Get()));
+		}
+		else
+		{
+			LandscapeInfo->Layers[Index].LayerInfoObj = LayerInfo;
+		}
+
+		if (LandscapeEdMode->CurrentToolTarget.LayerName == Target->LayerName
+			&& LandscapeEdMode->CurrentToolTarget.LayerInfo == Target->LayerInfoObj)
+		{
+			LandscapeEdMode->SetCurrentTargetLayer(Target->LayerName, Target->LayerInfoObj);
+		}
 
-			Target->LayerInfoObj = LayerInfo;
-			Target->LandscapeInfo->CreateTargetLayerSettingsFor(LayerInfo);
+		Target->LayerInfoObj = LayerInfo;
+		Target->LandscapeInfo->CreateTargetLayerSettingsFor(LayerInfo);
 
-			// Notify the asset registry
-			FAssetRegistryModule::AssetCreated(LayerInfo);
+		// Notify the asset registry
+		FAssetRegistryModule::AssetCreated(LayerInfo);
 
-			// Mark the package dirty...
-			Package->MarkPackageDirty();
+		// Mark the package dirty...
+		Package->MarkPackageDirty();
 
-			// Show in the content browser
-			TArray<UObject*> Objects;
-			Objects.Add(LayerInfo);
-			GEditor->SyncBrowserToObjects(Objects);
+		// Show in the content browser
+		TArray<UObject*> Objects;
+		Objects.Add(LayerInfo);
+		GEditor->SyncBrowserToObjects(Objects);
 
-			ALandscape* LandscapeActor = Target->LandscapeInfo->LandscapeActor.Get();
-			LandscapeActor->UpdateTargetLayer(FName(LayerName), FLandscapeTargetLayerSettings(LayerInfo));
+		ALandscape* LandscapeActor = Target->LandscapeInfo->LandscapeActor.Get();
+		LandscapeActor->UpdateTargetLayer(FName(LayerName), FLandscapeTargetLayerSettings(LayerInfo));
 
-			FillEmptyLayers(LandscapeInfo, LayerInfo);
+		FillEmptyLayers(LandscapeInfo, LayerInfo);
 
-			if (FEdModeLandscape* Mode = (FEdModeLandscape*)GLevelEditorModeTools().GetActiveMode(FBuiltinEditorModes::EM_Landscape))
-			{
-				Mode->UpdateTargetList();
-			}
-		}
+		LandscapeEdMode->UpdateTargetList();
 	}
 }
 
diff --git a/Engine/Source/Editor/LandscapeEditor/Private/LandscapeEditorDetailCustomization_TargetLayers.h b/Engine/Source/Editor/LandscapeEditor/Private/LandscapeEditorDetailCustomization_TargetLayers.h
--- a/Engine/Source/Editor/LandscapeEditor/Private/LandscapeEditorDetailCustomization_TargetLayers.h
+++ b/Engine/Source/Editor/LandscapeEditor/Private/LandscapeEditorDetailCustomization_TargetLayers.h
@@ -106,6 +106,7 @@ protected:
 	static EVisibility GetTargetLayerDeleteVisibility(const TSharedRef<FLandscapeTargetListInfo> Target);
 	static TSharedRef<SWidget> OnGetTargetLayerCreateMenu(const TSharedRef<FLandscapeTargetListInfo> Target);
 	static void OnTargetLayerCreateClicked(const TSharedRef<FLandscapeTargetListInfo> Target, bool bNoWeightBlend);
+	static void CreateTargetLayerInfoAsset(const TSharedRef<FLandscapeTargetListInfo> Target, const FString& PackageName, const FName& LayerObjectName, const bool bNoWeightBlend);
 	static FReply OnTargetLayerMakePublicClicked(const TSharedRef<FLandscapeTargetListInfo> Target);
 	static FReply OnTargetLayerDeleteClicked(const TSharedRef<FLandscapeTargetListInfo> Target);
 	static FSlateColor GetLayerUsageDebugColor(const TSharedRef<FLandscapeTargetListInfo> Target);
@@ -132,7 +133,7 @@ protected:
 
 	// Auto-Fill Target layers
 	FReply ShowAutoFillTargetLayerDialog() const;
-	FReply HandleAutoFillTargetLayers(bool bUpdateAllLayers) const;
+	FReply HandleAutoFillTargetLayers(const bool bUpdateAllLayers, const bool bCreateTargetLayers) const;
 
 	void HandleCreateLayer();
 	const FSlateBrush* GetTargetLayerDisplayOrderBrush() const;
diff --git a/Engine/Source/Editor/LandscapeEditor/Private/LandscapeEditorUtils.cpp b/Engine/Source/Editor/LandscapeEditor/Private/LandscapeEditorUtils.cpp
--- a/Engine/Source/Editor/LandscapeEditor/Private/LandscapeEditorUtils.cpp
+++ b/Engine/Source/Editor/LandscapeEditor/Private/LandscapeEditorUtils.cpp
@@ -125,7 +125,8 @@ namespace LandscapeEditorUtils
 		{
 			const ULandscapeLayerInfoObject* LayerInfo = CastChecked<ULandscapeLayerInfoObject>(LayerInfoAsset.GetAsset());
 
-			if (LayerInfo && LayerInfo->LayerName == LayerName)
+			// Only include assets in the package or its sub folders
+			if (LayerInfo && LayerInfo->LayerName == LayerName && LayerInfoAsset.PackagePath.ToString().Contains(TargetLayerAssetPackagePath))
 			{
 				const FString CurrentPackagePath = LayerInfoAsset.PackagePath.ToString();
 				float WorstCase = static_cast<float>(TargetLayerAssetPackagePath.Len() + CurrentPackagePath.Len());
