Landscape New Tab - Added new variable to separate New/Import target layers, added auto-fill button for target layer info assets, hid Layers/LayersBrushStack categories

diff --git a/Engine/Source/Editor/LandscapeEditor/Private/LandscapeEditorDetailCustomization_ImportLayers.cpp b/Engine/Source/Editor/LandscapeEditor/Private/LandscapeEditorDetailCustomization_ImportLayers.cpp
--- a/Engine/Source/Editor/LandscapeEditor/Private/LandscapeEditorDetailCustomization_ImportLayers.cpp
+++ b/Engine/Source/Editor/LandscapeEditor/Private/LandscapeEditorDetailCustomization_ImportLayers.cpp
@@ -315,11 +315,9 @@ void FLandscapeEditorStructCustomization_FLandscapeImportLayer::OnImportLayerCre
 		return;
 	}
 	
-	ULevel* Level = LandscapeEdMode->GetWorld()->GetCurrentLevel();
-
 	// Build default layer object name and package name
 	FName LayerObjectName;
-	FString PackageName = UE::Landscape::GetLayerInfoObjectPackageName(LayerName, UE::Landscape::GetSharedAssetsPath(Level), LayerObjectName);
+	FString PackageName = UE::Landscape::GetLayerInfoObjectPackageName(LayerName, LandscapeEdMode->GetTargetLayerAssetPackagePath(), LayerObjectName);
 
 	TSharedRef<SDlgPickAssetPath> NewLayerDlg =
 		SNew(SDlgPickAssetPath)
diff --git a/Engine/Source/Editor/LandscapeEditor/Private/LandscapeEditorDetailCustomization_NewLandscape.cpp b/Engine/Source/Editor/LandscapeEditor/Private/LandscapeEditorDetailCustomization_NewLandscape.cpp
--- a/Engine/Source/Editor/LandscapeEditor/Private/LandscapeEditorDetailCustomization_NewLandscape.cpp
+++ b/Engine/Source/Editor/LandscapeEditor/Private/LandscapeEditorDetailCustomization_NewLandscape.cpp
@@ -17,6 +17,7 @@
 #include "Widgets/Notifications/SErrorText.h"
 #include "Widgets/Input/SComboButton.h"
 #include "Widgets/Input/SCheckBox.h"
+#include "Dialog/SCustomDialog.h"
 #include "SWarningOrErrorBox.h"
 #include "LandscapeEditorModule.h"
 #include "LandscapeEditorObject.h"
@@ -63,6 +64,8 @@
 #include "SourceControlHelpers.h"
 
 #include "Misc/ScopedSlowTask.h"
+#include "EditorModeManager.h"
+#include "EditorModes.h"
 
 #include "Math/Box.h"
 #include "LandscapeEditorUtils.h"
@@ -70,6 +73,124 @@
 #define LOCTEXT_NAMESPACE "LandscapeEditor.NewLandscape"
 
 
+void FLandscapeEditorArrayBuilder_NewTargetLayer::GenerateHeaderRowContent(FDetailWidgetRow& NodeRow)
+{
+	const TSharedPtr<IPropertyHandle> PropertyHandle = GetPropertyHandle();
+	check(PropertyHandle);
+
+	NodeRow
+		.FilterString(PropertyHandle->GetPropertyDisplayName())
+		.NameContent()
+		[
+			PropertyHandle->CreatePropertyNameWidget()
+		]
+		.ValueContent()
+		[
+			SNew(SHorizontalBox)
+				+ SHorizontalBox::Slot()
+				.AutoWidth()
+				.VAlign(VAlign_Center)
+				.Padding(5.0f, 0.0f, 0.0f, 0.0f)
+				[
+					SNew(SButton)
+						.ButtonStyle(FAppStyle::Get(), "SimpleButton")
+						.ToolTipText(LOCTEXT("NewLandscape_AutoFillTargetLayers_Tip", "Auto-Fill Target Layer Assets"))
+						.OnClicked_Static(&FLandscapeEditorArrayBuilder_NewTargetLayer::ShowAutoFillTargetLayerDialog)
+						.IsEnabled_Static(&FLandscapeEditorArrayBuilder_NewTargetLayer::HasWeightmapTargetLayers)
+						[
+							SNew(SImage)
+								.Image(FAppStyle::GetBrush("LandscapeEditor.Layer.AutoFill"))
+						]
+				]
+		];
+}
+
+FReply FLandscapeEditorArrayBuilder_NewTargetLayer::HandleAutoFillTargetLayers(const bool bCreateAssetForEmptyLayers)
+{
+	if (const FEdModeLandscape* LandscapeEdMode = static_cast<FEdModeLandscape*>(GLevelEditorModeTools().GetActiveMode(FBuiltinEditorModes::EM_Landscape)))
+	{
+		FScopedTransaction Transaction(LOCTEXT("NewLandscape_AutoFillTargetLayers", "Auto-Fill Target Layer Assets"));
+
+		LandscapeEdMode->UISettings->AutoFillTargetLayers(LandscapeEdMode->GetTargetLayerAssetPackagePath(), bCreateAssetForEmptyLayers);
+	}
+
+	return FReply::Handled();
+}
+
+FReply FLandscapeEditorArrayBuilder_NewTargetLayer::ShowAutoFillTargetLayerDialog()
+{
+	if (const FEdModeLandscape* LandscapeEdMode = static_cast<FEdModeLandscape*>(GLevelEditorModeTools().GetActiveMode(FBuiltinEditorModes::EM_Landscape)))
+	{
+		const FString& DefaultTargetLayerAssetPath = LandscapeEdMode->GetTargetLayerAssetPackagePath();
+		
+		bool bIsCreateAssetForEmptyLayersChecked = false;
+
+		const TSharedPtr<SCustomDialog> DialogWindow = SNew(SCustomDialog)
+			.Title(LOCTEXT("TargetLayerAutoFillLayers", "Auto-Fill Target Layer Assets"))
+			.Buttons({
+				SCustomDialog::FButton(LOCTEXT("OK", "OK")).SetPrimary(true),
+				SCustomDialog::FButton(LOCTEXT("Cancel", "Cancel"))
+				})
+			.ContentAreaPadding(10.0f)
+			.Content()
+			[
+				SNew(SVerticalBox)
+					+ SVerticalBox::Slot()
+					.AutoHeight()
+					[
+						SNew(STextBlock)
+							.WrapTextAt(512.0f)
+							.Text(FText::Format(LOCTEXT("TargetLayerAutoFillDialog", "This operation will assign the layer info assets found within the default asset folder {0} to the landscape's target layers."),
+								FText::FromString(LandscapeEdMode->GetTargetLayerAssetPackagePath())))
+					]
+					+ SVerticalBox::Slot()
+					.AutoHeight()
+					.Padding(0, 10.0f, 0, 0)
+					[
+						SNew(SHorizontalBox)
+							+ SHorizontalBox::Slot()
+							.AutoWidth()
+							[
+								SNew(SCheckBox)
+									.IsChecked_Lambda([&bIsCreateAssetForEmptyLayersChecked]() { return bIsCreateAssetForEmptyLayersChecked ? ECheckBoxState::Checked : ECheckBoxState::Unchecked; })
+									.OnCheckStateChanged_Lambda([&bIsCreateAssetForEmptyLayersChecked](ECheckBoxState InState) { bIsCreateAssetForEmptyLayersChecked = (InState == ECheckBoxState::Checked); })
+							]
+							+ SHorizontalBox::Slot()
+							.AutoWidth()
+							[
+								SNew(STextBlock)
+									.Text(LOCTEXT("TargetLayerAutoFillDialog_CreateNew", "Create new assets in the default folder if none are found."))
+							]
+					]
+					+ SVerticalBox::Slot()
+					.AutoHeight()
+					.Padding(0, 20.0f, 0, 0) // Create spacing between details and prompt
+					[
+						SNew(STextBlock)
+							.Text(LOCTEXT("TargetLayerAutoFillDialog_Prompt", "Do you want to auto-fill all target layers?"))
+					]
+			];
+
+		// -1 is cancel, 0 all layers
+		if (DialogWindow->ShowModal() == 0)
+		{
+			HandleAutoFillTargetLayers(bIsCreateAssetForEmptyLayersChecked);
+		}
+	}
+	
+	return FReply::Handled();
+}
+
+bool FLandscapeEditorArrayBuilder_NewTargetLayer::HasWeightmapTargetLayers()
+{
+	if (const FEdModeLandscape* LandscapeEdMode = static_cast<FEdModeLandscape*>(GLevelEditorModeTools().GetActiveMode(FBuiltinEditorModes::EM_Landscape)))
+	{
+		return !LandscapeEdMode->UISettings->NewLandscape_Layers.IsEmpty();
+	}
+
+	return false;
+}
+
 TSharedRef<IDetailCustomization> FLandscapeEditorDetailCustomization_NewLandscape::MakeInstance()
 {
 	return MakeShareable(new FLandscapeEditorDetailCustomization_NewLandscape);
@@ -249,28 +370,31 @@ void FLandscapeEditorDetailCustomization_NewLandscape::CustomizeDetails(IDetailL
 
 	MaterialPropertyRow.OverrideResetToDefault(ResetOverride);
 
-	NewLandscapeCategory.AddCustomRow(LOCTEXT("LayersLabel", "Layers"))
-	.Visibility(TAttribute<EVisibility>(this, &FLandscapeEditorDetailCustomization_NewLandscape::GetMaterialTipVisibility))
-	[
-		SNew(SHorizontalBox)
-		+ SHorizontalBox::Slot()
-		.Padding(15, 12, 0, 12)
-		[
-			SNew(STextBlock)
-			.Font(DetailBuilder.GetDetailFont())
-			.Text(LOCTEXT("Material_Tip","Hint: Assign a material to see landscape layers"))
-		]
-	];
-
 	TSharedRef<IPropertyHandle> PropertyHandle_AlphamapType = DetailBuilder.GetProperty(GET_MEMBER_NAME_CHECKED(ULandscapeEditorObject, ImportLandscape_AlphamapType));
 	NewLandscapeCategory.AddProperty(PropertyHandle_AlphamapType)
 	.Visibility(TAttribute<EVisibility>::Create(TAttribute<EVisibility>::FGetter::CreateStatic(&GetVisibilityOnlyInNewLandscapeMode, ENewLandscapePreviewMode::ImportLandscape)));
 
-	TSharedRef<IPropertyHandle> PropertyHandle_Layers = DetailBuilder.GetProperty(GET_MEMBER_NAME_CHECKED(ULandscapeEditorObject, ImportLandscape_Layers));
+	TSharedRef<IPropertyHandle> PropertyHandle_TargetLayerAssetFilePath = DetailBuilder.GetProperty(GET_MEMBER_NAME_CHECKED(ULandscapeEditorObject, TargetLayerAssetFilePath));
+	NewLandscapeCategory.AddProperty(PropertyHandle_TargetLayerAssetFilePath);
 
-	TSharedRef<FDetailArrayBuilder> ArrayBuilder = MakeShared<FDetailArrayBuilder>(PropertyHandle_Layers, /*InGenerateHeader*/ true, /*InDisplayResetToDefault*/ false, /*InDisplayElementNum*/ false);
+	TSharedRef<IPropertyHandle> PropertyHandle_Layers = DetailBuilder.GetProperty(GET_MEMBER_NAME_CHECKED(ULandscapeEditorObject, NewLandscape_Layers));
+	TSharedPtr<FLandscapeEditorArrayBuilder_NewTargetLayer> ArrayBuilder = MakeShared<FLandscapeEditorArrayBuilder_NewTargetLayer>(PropertyHandle_Layers);
 	ArrayBuilder->OnGenerateArrayElementWidget(FOnGenerateArrayElementWidget::CreateSP(this, &FLandscapeEditorDetailCustomization_NewLandscape::GenerateLayersArrayElementWidget));
-	NewLandscapeCategory.AddCustomBuilder(ArrayBuilder, /*bForAdvanced*/ false);
+
+	NewLandscapeCategory.AddCustomBuilder(ArrayBuilder.ToSharedRef());
+
+	NewLandscapeCategory.AddCustomRow(LOCTEXT("LayersLabel", "Layers"))
+		.Visibility(TAttribute<EVisibility>(this, &FLandscapeEditorDetailCustomization_NewLandscape::GetMaterialTipVisibility))
+		[
+			SNew(SHorizontalBox)
+				+ SHorizontalBox::Slot()
+				.Padding(15, 12, 0, 12)
+				[
+					SNew(STextBlock)
+						.Font(DetailBuilder.GetDetailFont())
+						.Text(LOCTEXT("Material_Tip", "Hint: Assign a material to see the landscape layers it handles. Use Auto-Fill button to populate layer assets."))
+				]
+		];
 
 	TSharedRef<IPropertyHandle> PropertyHandle_Location = DetailBuilder.GetProperty(GET_MEMBER_NAME_CHECKED(ULandscapeEditorObject, NewLandscape_Location));
 	TSharedRef<IPropertyHandle> PropertyHandle_Location_X = PropertyHandle_Location->GetChildHandle("X").ToSharedRef();
@@ -1129,7 +1253,7 @@ FReply FLandscapeEditorDetailCustomization_NewLandscape::OnCreateButtonClicked()
 	LandscapeInfo->UpdateLayerInfoMap(Landscape);
 
 	// Import doesn't fill in the LayerInfo for layers with no data, do that now
-	const TArray<FLandscapeImportLayer>& ImportLandscapeLayersList = UISettings->ImportLandscape_Layers;
+	const TArray<FLandscapeImportLayer>& ImportLandscapeLayersList = UISettings->NewLandscape_Layers;
 	const ULandscapeSettings* Settings = GetDefault<ULandscapeSettings>();
 	TSoftObjectPtr<ULandscapeLayerInfoObject> DefaultLayerInfoObject = Settings->GetDefaultLayerInfoObject().LoadSynchronous();
 
@@ -1348,9 +1472,9 @@ bool FLandscapeEditorDetailCustomization_NewLandscape::GetImportButtonIsEnabled(
 			return false;
 		}
 
-		for (int32 i = 0; i < LandscapeEdMode->UISettings->ImportLandscape_Layers.Num(); ++i)
+		for (int32 i = 0; i < LandscapeEdMode->UISettings->NewLandscape_Layers.Num(); ++i)
 		{
-			if (LandscapeEdMode->UISettings->ImportLandscape_Layers[i].ImportResult == ELandscapeImportResult::Error)
+			if (LandscapeEdMode->UISettings->NewLandscape_Layers[i].ImportResult == ELandscapeImportResult::Error)
 			{
 				return false;
 			}
@@ -1501,7 +1625,7 @@ EVisibility FLandscapeEditorDetailCustomization_NewLandscape::GetMaterialTipVisi
 	FEdModeLandscape* LandscapeEdMode = GetEditorMode();
 	if (LandscapeEdMode != nullptr)
 	{
-		if (LandscapeEdMode->UISettings->ImportLandscape_Layers.Num() == 0)
+		if (LandscapeEdMode->UISettings->NewLandscape_Layers.Num() == 0)
 		{
 			return EVisibility::Visible;
 		}
diff --git a/Engine/Source/Editor/LandscapeEditor/Private/LandscapeEditorDetailCustomization_NewLandscape.h b/Engine/Source/Editor/LandscapeEditor/Private/LandscapeEditorDetailCustomization_NewLandscape.h
--- a/Engine/Source/Editor/LandscapeEditor/Private/LandscapeEditorDetailCustomization_NewLandscape.h
+++ b/Engine/Source/Editor/LandscapeEditor/Private/LandscapeEditorDetailCustomization_NewLandscape.h
@@ -11,6 +11,7 @@
 #include "LandscapeFileFormatInterface.h"
 #include "IDetailCustomization.h"
 #include "IPropertyTypeCustomization.h"
+#include "PropertyCustomizationHelpers.h"
 #include "PropertyHandle.h"
 #include "LandscapeEditorDetailCustomization_Base.h"
 #include "LandscapeSubsystem.h"
@@ -23,6 +24,22 @@ enum class ENewLandscapePreviewMode : uint8;
  * Slate widgets customizer for the "New Landscape" tool
  */
 
+class FLandscapeEditorArrayBuilder_NewTargetLayer : public FDetailArrayBuilder
+{
+public:
+	FLandscapeEditorArrayBuilder_NewTargetLayer(TSharedRef<IPropertyHandle> InBaseProperty)
+		: FDetailArrayBuilder(InBaseProperty)
+	{
+	}
+
+	virtual void GenerateHeaderRowContent(FDetailWidgetRow& NodeRow) override;
+
+private:
+	static FReply ShowAutoFillTargetLayerDialog();
+	static FReply HandleAutoFillTargetLayers(const bool bCreateAssetForEmptyLayers);
+	static bool HasWeightmapTargetLayers();
+};
+
 class FLandscapeEditorDetailCustomization_NewLandscape : public FLandscapeEditorDetailCustomization_Base
 {
 public:
diff --git a/Engine/Source/Editor/LandscapeEditor/Private/LandscapeEditorDetails.cpp b/Engine/Source/Editor/LandscapeEditor/Private/LandscapeEditorDetails.cpp
--- a/Engine/Source/Editor/LandscapeEditor/Private/LandscapeEditorDetails.cpp
+++ b/Engine/Source/Editor/LandscapeEditor/Private/LandscapeEditorDetails.cpp
@@ -201,7 +201,8 @@ void FLandscapeEditorDetails::CustomizeDetails(IDetailLayoutBuilder& DetailBuild
 	Customization_AlphaBrush = MakeShareable(new FLandscapeEditorDetailCustomization_AlphaBrush);
 	Customization_AlphaBrush->CustomizeDetails(DetailBuilder);
 
-	if (LandscapeEdMode->CanHaveLandscapeLayersContent())
+	// Hide Layers/LayersBrushStack when New Landscape tab is active
+	if (LandscapeEdMode->CanHaveLandscapeLayersContent() && LandscapeEdMode->CurrentTool->GetToolName() != FName("NewLandscape"))
 	{
 		// Layers
 		Customization_Layers = MakeShareable(new FLandscapeEditorDetailCustomization_Layers);
diff --git a/Engine/Source/Editor/LandscapeEditor/Private/LandscapeEditorObject.cpp b/Engine/Source/Editor/LandscapeEditor/Private/LandscapeEditorObject.cpp
--- a/Engine/Source/Editor/LandscapeEditor/Private/LandscapeEditorObject.cpp
+++ b/Engine/Source/Editor/LandscapeEditor/Private/LandscapeEditorObject.cpp
@@ -14,10 +14,12 @@
 #include "LandscapeTiledImage.h"
 #include "LandscapeMaterialInstanceConstant.h"
 #include "LandscapeUtils.h"
+#include "LandscapeEditorUtils.h"
 #include "Misc/ConfigCacheIni.h"
 #include "EngineUtils.h"
 #include "RenderUtils.h"
 #include "Misc/MessageDialog.h"
+#include "AssetRegistry/AssetRegistryModule.h"
 
 static TAutoConsoleVariable<bool> CVarLandscapeSimulateAlphaBrushTextureLoadFailure(
 	TEXT("landscape.SimulateAlphaBrushTextureLoadFailure"),
@@ -775,10 +777,10 @@ ELandscapeImportResult ULandscapeEditorObject::CreateImportLayersInfo(TArray<FLa
 		return ELandscapeImportResult::Error;
 	}
 
-	OutImportLayerInfos.Reserve(ImportLandscape_Layers.Num());
+	OutImportLayerInfos.Reserve(NewLandscape_Layers.Num());
 
 	// Fill in LayerInfos array and allocate data
-	for (FLandscapeImportLayer& UIImportLayer : ImportLandscape_Layers)
+	for (FLandscapeImportLayer& UIImportLayer : NewLandscape_Layers)
 	{
 		OutImportLayerInfos.Add((const FLandscapeImportLayer&)UIImportLayer); //slicing is fine here
 		FLandscapeImportLayerInfo& ImportLayer = OutImportLayerInfos.Last();
@@ -809,10 +811,10 @@ ELandscapeImportResult ULandscapeEditorObject::CreateNewLayersInfo(TArray<FLands
 	const int32 SizeX = NewLandscape_ComponentCount.X * QuadsPerComponent + 1;
 	const int32 SizeY = NewLandscape_ComponentCount.Y * QuadsPerComponent + 1;
 
-	OutNewLayerInfos.Reset(ImportLandscape_Layers.Num());
+	OutNewLayerInfos.Reset(NewLandscape_Layers.Num());
 
 	// Fill in LayerInfos array and allocate data
-	for (const FLandscapeImportLayer& UIImportLayer : ImportLandscape_Layers)
+	for (const FLandscapeImportLayer& UIImportLayer : NewLandscape_Layers)
 	{
 		FLandscapeImportLayerInfo ImportLayer = FLandscapeImportLayerInfo(UIImportLayer.LayerName);
 		ImportLayer.LayerInfo = UIImportLayer.LayerInfo;
@@ -908,8 +910,9 @@ void ULandscapeEditorObject::RefreshImportLayersList(bool bRefreshFromTarget, bo
 		LayerNames = ALandscapeProxy::GetLayersFromMaterial(Material);
 	}
 
-	const TArray<FLandscapeImportLayer> OldLayersList = MoveTemp(ImportLandscape_Layers);
-	ImportLandscape_Layers.Reset(LayerNames.Num());
+	TArray<FLandscapeImportLayer>& LandscapeLayers = bRefreshFromTarget ? ImportLandscape_Layers : NewLandscape_Layers;
+	const TArray<FLandscapeImportLayer> OldLayersList = MoveTemp(LandscapeLayers);
+	LandscapeLayers.Reset(LayerNames.Num());
 
 	// Don't recreate the render state of everything, only update the materials context
 	FMaterialUpdateContext MaterialUpdateContext(FMaterialUpdateContext::EOptions::Default & ~FMaterialUpdateContext::EOptions::RecreateRenderStates);
@@ -955,7 +958,42 @@ void ULandscapeEditorObject::RefreshImportLayersList(bool bRefreshFromTarget, bo
 
 			RefreshLayerImport(NewImportLayer);
 
-			ImportLandscape_Layers.Add(MoveTemp(NewImportLayer));
+			LandscapeLayers.Add(MoveTemp(NewImportLayer));
+		}
+	}
+}
+
+void ULandscapeEditorObject::AutoFillTargetLayers(const FString& DefaultTargetLayerAssetPath, const bool bCreateAssetForEmptyLayers)
+{
+	for (FLandscapeImportLayer& NewLayer : NewLandscape_Layers)
+	{
+		if (const FAssetData* AssetData = LandscapeEditorUtils::FindLandscapeTargetLayerInfoAsset(NewLayer.LayerName, DefaultTargetLayerAssetPath))
+		{
+			NewLayer.LayerInfo = CastChecked<ULandscapeLayerInfoObject>(AssetData->GetAsset());
+		}
+		else if (bCreateAssetForEmptyLayers)
+		{
+			const FName LayerName = NewLayer.LayerName;
+			FName LayerObjectName;
+			const FString PackageName = UE::Landscape::GetLayerInfoObjectPackageName(LayerName, DefaultTargetLayerAssetPath, LayerObjectName);
+
+			UPackage* Package = CreatePackage(*PackageName);
+			ULandscapeLayerInfoObject* LayerInfo = NewObject<ULandscapeLayerInfoObject>(Package, LayerObjectName, RF_Public | RF_Standalone | RF_Transactional);
+			LayerInfo->LayerName = LayerName;
+			LayerInfo->bNoWeightBlend = true;
+
+			NewLayer.LayerInfo = LayerInfo;
+
+			// Notify the asset registry
+			FAssetRegistryModule::AssetCreated(LayerInfo);
+
+			// Mark the package dirty...
+			Package->MarkPackageDirty();
+
+			// Show in the content browser
+			TArray<UObject*> Objects;
+			Objects.Add(LayerInfo);
+			GEditor->SyncBrowserToObjects(Objects);
 		}
 	}
 }
diff --git a/Engine/Source/Editor/LandscapeEditor/Public/LandscapeEditorObject.h b/Engine/Source/Editor/LandscapeEditor/Public/LandscapeEditorObject.h
--- a/Engine/Source/Editor/LandscapeEditor/Public/LandscapeEditorObject.h
+++ b/Engine/Source/Editor/LandscapeEditor/Public/LandscapeEditorObject.h
@@ -634,10 +634,14 @@ public:
 	UPROPERTY(Category="New Landscape", EditAnywhere, NonTransactional, meta=(DisplayName="Layer Alphamap Type", ShowForTools="NewLandscape,ImportExport"))
 	ELandscapeImportAlphamapType ImportLandscape_AlphamapType;
 
-	// The landscape layers that will be created. Only layer names referenced in the material assigned above are shown here. Modify the material to add more layers.
-	UPROPERTY(Category="New Landscape", EditAnywhere, NonTransactional, EditFixedSize, meta=(DisplayName="Layers", ShowForTools="NewLandscape,ImportExport"))
+	// The landscape layers that will be created via Import tab. Only layer names referenced in the material assigned above are shown here. Modify the material to add more layers.
+	UPROPERTY(Category="New Landscape", EditAnywhere, NonTransactional, EditFixedSize, meta=(DisplayName="Layers", ShowForTools="ImportExport"))
 	TArray<FLandscapeImportLayer> ImportLandscape_Layers;
 
+	// The landscape layers that will be created via New tab. Only layer names referenced in the material assigned above are shown here. Modify the material to add more layers.
+	UPROPERTY(Category="New Landscape", EditAnywhere, NonTransactional, EditFixedSize, meta=(DisplayName="Layers", ShowForTools="NewLandscape"))
+	TArray<FLandscapeImportLayer> NewLandscape_Layers;
+
 	// Common Brush Settings:
 
 	// The radius of the sculpt brush, in unreal units
@@ -766,6 +770,7 @@ public:
 	void ChooseBestComponentSizeForImport();
 	void ImportLandscapeData();
 	void RefreshImportLayersList(bool bRefreshFromTarget = false, bool bRegenerateThumbnails = false);
+	void AutoFillTargetLayers(const FString& DefaultTargetLayerAssetPath, const bool bCreateAssetForEmptyLayers);
 	ELandscapeImportResult CreateImportLayersInfo(TArray<FLandscapeImportLayerInfo>& OutImportLayerInfos);
 	ELandscapeImportResult CreateNewLayersInfo(TArray<FLandscapeImportLayerInfo>& OutNewLayerInfos);
 	void InitializeDefaultHeightData(TArray<uint16>& OutData);
