* Moved TargetDisplayOrder and TargetDisplayOrderList out of ALandscapeProxy and into ALandscape. This is just a UI setting that doesn't need to be duplicated to all streaming proxies
* Also, fixed landscape thumbnails disappearing anytime we unload a region

diff --git a/Engine/Source/Editor/LandscapeEditor/Private/LandscapeEdMode.cpp b/Engine/Source/Editor/LandscapeEditor/Private/LandscapeEdMode.cpp
--- a/Engine/Source/Editor/LandscapeEditor/Private/LandscapeEdMode.cpp
+++ b/Engine/Source/Editor/LandscapeEditor/Private/LandscapeEdMode.cpp
@@ -2573,11 +2573,10 @@ void FEdModeLandscape::UpdateTargetList(bool bRegenerateThumbnails)
 
 	if (CurrentToolTarget.LandscapeInfo.IsValid())
 	{
-		ALandscapeProxy* LandscapeProxy = CurrentToolTarget.LandscapeInfo->GetLandscapeProxy();
-
-		if (LandscapeProxy != nullptr)
+		ALandscape* Landscape = CurrentToolTarget.LandscapeInfo->LandscapeActor.Get();
+		if (Landscape != nullptr)
 		{
-			CachedLandscapeMaterial = LandscapeProxy->GetLandscapeMaterial();
+			CachedLandscapeMaterial = Landscape->GetLandscapeMaterial();
 
 			bool bFoundSelected = false;
 
@@ -2590,7 +2589,7 @@ void FEdModeLandscape::UpdateTargetList(bool bRegenerateThumbnails)
 			}
 
 			// Add visibility
-			FLandscapeInfoLayerSettings VisibilitySettings(ALandscapeProxy::VisibilityLayer, LandscapeProxy);
+			FLandscapeInfoLayerSettings VisibilitySettings(ALandscapeProxy::VisibilityLayer, Landscape);
 			LandscapeTargetList.Add(MakeShareable(new FLandscapeTargetListInfo(LOCTEXT("Visibility", "Visibility"), ELandscapeToolTargetType::Visibility, VisibilitySettings, UISettings->CurrentLayerIndex)));
 
 			if (CurrentToolTarget.TargetType == ELandscapeToolTargetType::Visibility)
@@ -2639,7 +2638,7 @@ void FEdModeLandscape::UpdateTargetList(bool bRegenerateThumbnails)
 				SetCurrentTargetLayer(NAME_None, nullptr);
 			}
 
-			UISettings->TargetDisplayOrder = LandscapeProxy->TargetDisplayOrder;
+			UISettings->TargetDisplayOrder = Landscape->TargetDisplayOrder;
 
 			UpdateTargetLayerDisplayOrder(UISettings->TargetDisplayOrder);
 		}
@@ -2657,9 +2656,8 @@ void FEdModeLandscape::UpdateTargetLayerDisplayOrder(ELandscapeLayerDisplayMode
 		return;
 	}
 
-	ALandscapeProxy* LandscapeProxy = CurrentToolTarget.LandscapeInfo->GetLandscapeProxy();
-
-	if (LandscapeProxy == nullptr)
+	ALandscape* Landscape = CurrentToolTarget.LandscapeInfo->LandscapeActor.Get();
+	if (Landscape == nullptr)
 	{
 		return;
 	}
@@ -2667,8 +2665,8 @@ void FEdModeLandscape::UpdateTargetLayerDisplayOrder(ELandscapeLayerDisplayMode
 	bool DetailPanelRefreshRequired = false;
 
 	// Save value to landscape
-	LandscapeProxy->TargetDisplayOrder = InTargetDisplayOrder;
-	TArray<FName>& SavedTargetNameList = LandscapeProxy->TargetDisplayOrderList;
+	Landscape->TargetDisplayOrder = InTargetDisplayOrder;
+	TArray<FName>& SavedTargetNameList = Landscape->TargetDisplayOrderList;
 
 	switch (InTargetDisplayOrder)
 	{
@@ -2895,14 +2893,13 @@ const TArray<FName>* FEdModeLandscape::GetTargetDisplayOrderList() const
 		return nullptr;
 	}
 
-	ALandscapeProxy* LandscapeProxy = CurrentToolTarget.LandscapeInfo->GetLandscapeProxy();
-
-	if (LandscapeProxy == nullptr)
+	ALandscape* Landscape = CurrentToolTarget.LandscapeInfo->LandscapeActor.Get();
+	if (Landscape == nullptr)
 	{
 		return nullptr;
 	}
 
-	return &LandscapeProxy->TargetDisplayOrderList;
+	return &Landscape->TargetDisplayOrderList;
 }
 
 void FEdModeLandscape::MoveTargetLayerDisplayOrder(int32 IndexToMove, int32 IndexToDestination)
@@ -2912,18 +2909,17 @@ void FEdModeLandscape::MoveTargetLayerDisplayOrder(int32 IndexToMove, int32 Inde
 		return;
 	}
 
-	ALandscapeProxy* LandscapeProxy = CurrentToolTarget.LandscapeInfo->GetLandscapeProxy();
-
-	if (LandscapeProxy == nullptr)
+	ALandscape* Landscape = CurrentToolTarget.LandscapeInfo->LandscapeActor.Get();
+	if (Landscape == nullptr)
 	{
 		return;
 	}
 	
-	FName Data = LandscapeProxy->TargetDisplayOrderList[IndexToMove];
-	LandscapeProxy->TargetDisplayOrderList.RemoveAt(IndexToMove);
-	LandscapeProxy->TargetDisplayOrderList.Insert(Data, IndexToDestination);
+	FName Data = Landscape->TargetDisplayOrderList[IndexToMove];
+	Landscape->TargetDisplayOrderList.RemoveAt(IndexToMove);
+	Landscape->TargetDisplayOrderList.Insert(Data, IndexToDestination);
 
-	LandscapeProxy->TargetDisplayOrder = ELandscapeLayerDisplayMode::UserSpecific;
+	Landscape->TargetDisplayOrder = ELandscapeLayerDisplayMode::UserSpecific;
 	UISettings->TargetDisplayOrder = ELandscapeLayerDisplayMode::UserSpecific;
 
 	RefreshDetailPanel();
diff --git a/Engine/Source/Runtime/Landscape/Classes/Landscape.h b/Engine/Source/Runtime/Landscape/Classes/Landscape.h
--- a/Engine/Source/Runtime/Landscape/Classes/Landscape.h
+++ b/Engine/Source/Runtime/Landscape/Classes/Landscape.h
@@ -671,6 +671,14 @@ public:
 	UPROPERTY(Transient, DuplicateTransient, TextExportTransient)	
 	TArray<TWeakObjectPtr<UTexture2D>> TrackedStreamingInTextures;
 
+	/** Display Order of the targets */
+	UPROPERTY(NonTransactional)
+	TArray<FName> TargetDisplayOrderList;
+
+	/** Display Order mode for the targets */
+	UPROPERTY(NonTransactional)
+	ELandscapeLayerDisplayMode TargetDisplayOrder = ELandscapeLayerDisplayMode::Default;
+
 private:
 	UPROPERTY()
 	TArray<FLandscapeLayer> LandscapeEditLayers;
diff --git a/Engine/Source/Runtime/Landscape/Classes/LandscapeProxy.h b/Engine/Source/Runtime/Landscape/Classes/LandscapeProxy.h
--- a/Engine/Source/Runtime/Landscape/Classes/LandscapeProxy.h
+++ b/Engine/Source/Runtime/Landscape/Classes/LandscapeProxy.h
@@ -541,14 +541,6 @@ public:
 	/** LOD level to use when exporting the landscape to obj or FBX */
 	UPROPERTY(EditAnywhere, Category=LOD, AdvancedDisplay, meta = (LandscapeOverridable))
 	int32 ExportLOD;
-
-	/** Display Order of the targets */
-	UPROPERTY(NonTransactional, meta = (LandscapeInherited))
-	TArray<FName> TargetDisplayOrderList;
-
-	/** Display Order mode for the targets */
-	UPROPERTY(NonTransactional, meta = (LandscapeInherited))
-	ELandscapeLayerDisplayMode TargetDisplayOrder;
 #endif
 
 	/** LOD level to use when running lightmass (increase to 1 or 2 for large landscapes to stop lightmass crashing) */
diff --git a/Engine/Source/Runtime/Landscape/Private/Landscape.cpp b/Engine/Source/Runtime/Landscape/Private/Landscape.cpp
--- a/Engine/Source/Runtime/Landscape/Private/Landscape.cpp
+++ b/Engine/Source/Runtime/Landscape/Private/Landscape.cpp
@@ -1581,9 +1581,6 @@ TArray<ALandscapeProxy*> ALandscapeProxy::LandscapeProxies;
 
 ALandscapeProxy::ALandscapeProxy(const FObjectInitializer& ObjectInitializer)
 	: Super(ObjectInitializer)
-#if WITH_EDITORONLY_DATA
-	, TargetDisplayOrder(ELandscapeLayerDisplayMode::Default)
-#endif // WITH_EDITORONLY_DATA
 {
 	bReplicates = false;
 	SetNetUpdateFrequency(10.0f);
@@ -4097,13 +4094,20 @@ bool ULandscapeInfo::UpdateLayerInfoMapInternal(ALandscapeProxy* Proxy)
 		Proxy->TargetLayersForFixup.Empty();
 	}
 
-	Layers.Empty();
+	// Keep a temp copy of the previous layers to keep the thumbnail MICs alive : 
+	TArray<FLandscapeInfoLayerSettings> PreviousLayers;
+	Swap(Layers, PreviousLayers);
 
 	for (const TTuple<FName, FLandscapeTargetLayerSettings>& TargetLayer : LandscapeActor->GetTargetLayers())
 	{
 		FLandscapeInfoLayerSettings InfoLayerSettings (TargetLayer.Key, LandscapeActor.Get());
 		InfoLayerSettings.bValid = true;
 		InfoLayerSettings.LayerInfoObj = TargetLayer.Value.LayerInfoObj;
+		if (const FLandscapeInfoLayerSettings* PreviousLayer = Algo::FindBy(PreviousLayers, TargetLayer.Key, &FLandscapeInfoLayerSettings::LayerName))
+		{
+			InfoLayerSettings.ThumbnailMIC = PreviousLayer->ThumbnailMIC;
+		}
+
 		Layers.Add(InfoLayerSettings);
 	}
 	
