Landscape Paint - Added Rename option to target layer quick action menu, refactored action menu to show disabled states, automatically select/rename new target layers on creation

diff --git a/Engine/Source/Editor/LandscapeEditor/Private/LandscapeEdMode.cpp b/Engine/Source/Editor/LandscapeEditor/Private/LandscapeEdMode.cpp
--- a/Engine/Source/Editor/LandscapeEditor/Private/LandscapeEdMode.cpp
+++ b/Engine/Source/Editor/LandscapeEditor/Private/LandscapeEdMode.cpp
@@ -104,6 +104,7 @@ FEdModeLandscape::FEdModeLandscape()
 	, CurrentToolIndex(INDEX_NONE)
 	, CurrentBrushSetIndex(0)
 	, CurrentLandscapeViewMode(ELandscapeViewMode::Normal)
+	, PendingRenameTargetLayerIndex(INDEX_NONE)
 	, NewLandscapePreviewMode(ENewLandscapePreviewMode::None)
 	, ImportExportMode(EImportExportMode::Import)
 	, CurrentGizmoActor(nullptr)
diff --git a/Engine/Source/Editor/LandscapeEditor/Private/LandscapeEdMode.h b/Engine/Source/Editor/LandscapeEditor/Private/LandscapeEdMode.h
--- a/Engine/Source/Editor/LandscapeEditor/Private/LandscapeEdMode.h
+++ b/Engine/Source/Editor/LandscapeEditor/Private/LandscapeEdMode.h
@@ -315,6 +315,8 @@ public:
 	int32 CurrentBrushSetIndex;
 	// Persistent View Mode when toggling landscape editor
 	ELandscapeViewMode::Type CurrentLandscapeViewMode;
+	// Persistent Target Layer index when entering a layer rename
+	int32 PendingRenameTargetLayerIndex;
 
 	ENewLandscapePreviewMode NewLandscapePreviewMode;
 	EImportExportMode ImportExportMode;
diff --git a/Engine/Source/Editor/LandscapeEditor/Private/LandscapeEditorDetailCustomization_TargetLayers.cpp b/Engine/Source/Editor/LandscapeEditor/Private/LandscapeEditorDetailCustomization_TargetLayers.cpp
--- a/Engine/Source/Editor/LandscapeEditor/Private/LandscapeEditorDetailCustomization_TargetLayers.cpp
+++ b/Engine/Source/Editor/LandscapeEditor/Private/LandscapeEditorDetailCustomization_TargetLayers.cpp
@@ -35,6 +35,7 @@
 #include "DetailCategoryBuilder.h"
 #include "PropertyCustomizationHelpers.h"
 #include "Subsystems/AssetEditorSubsystem.h"
+#include "Algo/Count.h"
 #include "AssetRegistry/AssetRegistryModule.h"
 #include "Materials/MaterialExpressionLandscapeVisibilityMask.h"
 
@@ -685,7 +686,7 @@ void FLandscapeEditorCustomNodeBuilder_TargetLayers::HandleCreateLayer()
 	{
 		return; 
 	}
-	
+
 	TWeakObjectPtr<ALandscape> Landscape = LandscapeEdMode->CurrentToolTarget.LandscapeInfo->LandscapeActor;
 
 	if (!Landscape.Get())
@@ -694,12 +695,20 @@ void FLandscapeEditorCustomNodeBuilder_TargetLayers::HandleCreateLayer()
 	}
 
 	FScopedTransaction Transaction(LOCTEXT("LandscapeTargetLayer_Create", "Create a Target Layer"));
-	
+
 	Landscape->AddTargetLayer();
-	
+
 	LandscapeEdMode->CurrentToolTarget.LandscapeInfo->UpdateLayerInfoMap();
 	LandscapeEdMode->UpdateTargetList();
+
+	// Trigger a rename for the new layer, enters edit mode on next tick
+	LandscapeEdMode->PendingRenameTargetLayerIndex = GetWeightmapTargetLayerCount() - 1;
+
 	LandscapeEdMode->RefreshDetailPanel();
+
+	// Auto-select new target layer in details panel
+	OnTargetSelectionChanged(LandscapeEdMode->GetTargetList().Last());
+	
 }
 
 TSharedRef<SWidget> FLandscapeEditorCustomNodeBuilder_TargetLayers::GetTargetLayerDisplayOrderButtonMenuContent()
@@ -882,17 +891,15 @@ void FLandscapeEditorCustomNodeBuilder_TargetLayers::OnFilterTextCommitted(const
 	}
 }
 
-bool FLandscapeEditorCustomNodeBuilder_TargetLayers::HasWeightmapTargetLayers() const
+int32 FLandscapeEditorCustomNodeBuilder_TargetLayers::GetWeightmapTargetLayerCount() const
 {
-	FEdModeLandscape* LandscapeEdMode = GetEditorMode();
-
-	if (LandscapeEdMode && LandscapeEdMode->CurrentToolMode)
+	if (const FEdModeLandscape* LandscapeEdMode = GetEditorMode(); LandscapeEdMode->CurrentToolMode)
 	{
 		// ELandscapeToolTargetType::Invalid means "weightmap with no valid paint layer" so we still want to display that property if it has been marked to be displayed in Weightmap target type, to be consistent 
 		if ((LandscapeEdMode->CurrentToolTarget.TargetType == ELandscapeToolTargetType::Weightmap)
 			|| (LandscapeEdMode->CurrentToolTarget.TargetType == ELandscapeToolTargetType::Invalid))
 		{
-			return LandscapeEdMode->GetTargetList().ContainsByPredicate([](TSharedRef<FLandscapeTargetListInfo> InInfo)
+			return Algo::CountIf(LandscapeEdMode->GetTargetList(), [](TSharedRef<FLandscapeTargetListInfo> InInfo)
 				{
 					FName LayerName = InInfo->GetLayerName();
 					return (LayerName != NAME_None) && (LayerName != UMaterialExpressionLandscapeVisibilityMask::ParameterName);
@@ -900,7 +907,12 @@ bool FLandscapeEditorCustomNodeBuilder_TargetLayers::HasWeightmapTargetLayers()
 		}
 	}
 
-	return false;
+	return 0;
+}
+
+bool FLandscapeEditorCustomNodeBuilder_TargetLayers::HasWeightmapTargetLayers() const
+{
+	return GetWeightmapTargetLayerCount() > 0;
 }
 
 EVisibility FLandscapeEditorCustomNodeBuilder_TargetLayers::GetLayersFilterVisibility() const
@@ -997,6 +1009,9 @@ void FLandscapeEditorCustomNodeBuilder_TargetLayers::GenerateChildContent(IDetai
 
 		// Generate a row for all target layers, including those that will be filtered and let the row's visibility lambda to compute their visibility dynamically. This allows 
 		//  filtering to work without refreshing the details panel (which causes the search box to lose focus) :
+		InlineTextBlocks.Empty();
+		InlineTextBlocks.Reserve(GetWeightmapTargetLayerCount());
+
 		for (const TSharedRef<FLandscapeTargetListInfo>& TargetInfo : PrepareTargetLayerList(/*bInSort = */true, /*bInFilter = */false))
 		{
 			TSharedPtr<SWidget> GeneratedRowWidget = GenerateRow(TargetInfo);
@@ -1013,13 +1028,21 @@ void FLandscapeEditorCustomNodeBuilder_TargetLayers::GenerateChildContent(IDetai
 }
 END_SLATE_FUNCTION_BUILD_OPTIMIZATION
 
+void FLandscapeEditorCustomNodeBuilder_TargetLayers::Tick(float DeltaTime)
+{
+	if (FEdModeLandscape* LandscapeEdMode = GetEditorMode(); LandscapeEdMode->PendingRenameTargetLayerIndex != INDEX_NONE)
+	{
+		OnRenameLayer(LandscapeEdMode->PendingRenameTargetLayerIndex);
+		LandscapeEdMode->PendingRenameTargetLayerIndex = INDEX_NONE;
+	}
+}
+
 BEGIN_SLATE_FUNCTION_BUILD_OPTIMIZATION
 TSharedPtr<SWidget> FLandscapeEditorCustomNodeBuilder_TargetLayers::GenerateRow(const TSharedRef<FLandscapeTargetListInfo> Target)
 {
 	TSharedPtr<SWidget> RowWidget;
 
-	FEdModeLandscape* LandscapeEdMode = GetEditorMode();
-	if (LandscapeEdMode)
+	if (FEdModeLandscape* LandscapeEdMode = GetEditorMode())
 	{
 		if ((LandscapeEdMode->CurrentTool->GetSupportedTargetTypes() & LandscapeEdMode->CurrentToolMode->SupportedTargetTypes & ELandscapeToolTargetTypeMask::FromType(Target->TargetType)) == 0)
 		{
@@ -1032,7 +1055,7 @@ TSharedPtr<SWidget> FLandscapeEditorCustomNodeBuilder_TargetLayers::GenerateRow(
 		RowWidget = SNew(SLandscapeEditorSelectableBorder)
 			.Padding(0.0f)
 			.VAlign(VAlign_Center)
-			.OnContextMenuOpening_Static(&FLandscapeEditorCustomNodeBuilder_TargetLayers::OnTargetLayerContextMenuOpening, Target)
+			.OnContextMenuOpening(this, &FLandscapeEditorCustomNodeBuilder_TargetLayers::OnTargetLayerContextMenuOpening, Target, static_cast<int32>(INDEX_NONE))
 			.OnSelected_Static(&FLandscapeEditorCustomNodeBuilder_TargetLayers::OnTargetSelectionChanged, Target)
 			.IsSelected_Static(&FLandscapeEditorCustomNodeBuilder_TargetLayers::GetTargetLayerIsSelected, Target)
 			.Visibility(this, &FLandscapeEditorCustomNodeBuilder_TargetLayers::ShouldShowLayer, Target)
@@ -1068,10 +1091,12 @@ TSharedPtr<SWidget> FLandscapeEditorCustomNodeBuilder_TargetLayers::GenerateRow(
 	{
 		static const FSlateColorBrush SolidWhiteBrush = FSlateColorBrush(FColorList::White);
 
+		InlineTextBlocks.Add(TSharedPtr<SInlineEditableTextBlock>());
+
 		RowWidget = SNew(SLandscapeEditorSelectableBorder)
 			.Padding(0.0f)
 			.VAlign(VAlign_Center)
-			.OnContextMenuOpening_Static(&FLandscapeEditorCustomNodeBuilder_TargetLayers::OnTargetLayerContextMenuOpening, Target)
+			.OnContextMenuOpening(this, &FLandscapeEditorCustomNodeBuilder_TargetLayers::OnTargetLayerContextMenuOpening, Target, InlineTextBlocks.Num() - 1)
 			.OnSelected_Static(&FLandscapeEditorCustomNodeBuilder_TargetLayers::OnTargetSelectionChanged, Target)
 			.IsSelected_Static(&FLandscapeEditorCustomNodeBuilder_TargetLayers::GetTargetLayerIsSelected, Target)
 			.Visibility(this, &FLandscapeEditorCustomNodeBuilder_TargetLayers::ShouldShowLayer, Target)
@@ -1146,7 +1171,7 @@ TSharedPtr<SWidget> FLandscapeEditorCustomNodeBuilder_TargetLayers::GenerateRow(
 								SNew(SHorizontalBox)
 									+ SHorizontalBox::Slot()
 									[
-										SNew(SInlineEditableTextBlock)
+										SAssignNew(InlineTextBlocks.Last(), SInlineEditableTextBlock)
 											.Font(IDetailLayoutBuilder::GetDetailFontBold())
 											.Text(Target->TargetLayerDisplayName)
 											.ColorAndOpacity_Static(&FLandscapeEditorCustomNodeBuilder_TargetLayers::GetTargetTextColor, Target)
@@ -1472,53 +1497,73 @@ void FLandscapeEditorCustomNodeBuilder_TargetLayers::OnTargetSelectionChanged(co
 	}
 }
 
-TSharedPtr<SWidget> FLandscapeEditorCustomNodeBuilder_TargetLayers::OnTargetLayerContextMenuOpening(const TSharedRef<FLandscapeTargetListInfo> Target)
+TSharedPtr<SWidget> FLandscapeEditorCustomNodeBuilder_TargetLayers::OnTargetLayerContextMenuOpening(const TSharedRef<FLandscapeTargetListInfo> Target, const int32 InLayerIndex)
 {
-	if (Target->TargetType == ELandscapeToolTargetType::Heightmap || Target->LayerInfoObj != nullptr)
+	FMenuBuilder MenuBuilder(true, nullptr);
+
+	MenuBuilder.BeginSection("LandscapeEditorLayerActions", LOCTEXT("LayerContextMenu.Heading", "Layer Actions"));
 	{
-		FMenuBuilder MenuBuilder(true, nullptr);
-		
-		MenuBuilder.BeginSection("LandscapeEditorLayerActions", LOCTEXT("LayerContextMenu.Heading", "Layer Actions"));
+		if (FEdModeLandscape* LandscapeEdMode = GetEditorMode())
 		{
-			FEdModeLandscape* LandscapeEdMode = GetEditorMode();
-			if (LandscapeEdMode)
-			{
-				FUIAction LandscapeHeightmapChangeToolsAction = FUIAction(FExecuteAction::CreateStatic(&FLandscapeEditorCustomNodeBuilder_TargetLayers::OnHeightmapLayerContextMenu, Target));
-				MenuBuilder.AddMenuEntry(LOCTEXT("LayerContextMenu.Heightmap", "Import From/Export To File..."),
-										LOCTEXT("LayerContextMenu.HeightmapToolTip", "Opens the Landscape Import tool in order to import / export heightmaps from / to external files."), FSlateIcon(), LandscapeHeightmapChangeToolsAction);
-			}
+			FUIAction LandscapeHeightmapChangeToolsAction = FUIAction(FExecuteAction::CreateStatic(&FLandscapeEditorCustomNodeBuilder_TargetLayers::OnHeightmapLayerContextMenu, Target));
+			MenuBuilder.AddMenuEntry(LOCTEXT("LayerContextMenu.Heightmap", "Import From/Export To File..."),
+				LOCTEXT("LayerContextMenu.HeightmapToolTip", "Opens the Landscape Import tool in order to import / export layers from / to external files."), FSlateIcon(), LandscapeHeightmapChangeToolsAction);
 
-			if (Target->TargetType == ELandscapeToolTargetType::Weightmap)
+			if (Target->TargetType == ELandscapeToolTargetType::Weightmap && InLayerIndex != INDEX_NONE)
 			{
+				// Rebuild material instances
+				FUIAction RebuildAction = FUIAction(FExecuteAction::CreateStatic(&FLandscapeEditorCustomNodeBuilder_TargetLayers::OnRebuildMICs, Target));
+				MenuBuilder.AddMenuEntry(LOCTEXT("LayerContextMenu.Rebuild", "Rebuild Materials"), LOCTEXT("LayerContextMenu.Rebuild_Tooltip", "Rebuild material instances used for this landscape."), FSlateIcon(), RebuildAction);
+
+				// Rename Layer
+				FUIAction RenameAction = FUIAction(FExecuteAction::CreateSP(this, &FLandscapeEditorCustomNodeBuilder_TargetLayers::OnRenameLayer, InLayerIndex));
+				MenuBuilder.AddMenuEntry(LOCTEXT("LayerContextMenu.Rename", "Rename Layer"), LOCTEXT("LayerContextMenu.Rename_Tooltip", "Rename this target layer."), FSlateIcon(), RenameAction);
+
+				// Separate Generic vs Content based actions
 				MenuBuilder.AddMenuSeparator();
 
 				// Fill
-				FUIAction FillAction = FUIAction(FExecuteAction::CreateStatic(&FLandscapeEditorCustomNodeBuilder_TargetLayers::OnFillLayer, Target));
-				MenuBuilder.AddMenuEntry(LOCTEXT("LayerContextMenu.Fill", "Fill Layer"), LOCTEXT("LayerContextMenu.Fill_Tooltip", "Fills this layer to 100% across the entire landscape. If this is a weight-blended (normal) layer, all other weight-blended layers will be cleared."), FSlateIcon(), FillAction);
+				FUIAction FillAction = FUIAction(FExecuteAction::CreateStatic(&FLandscapeEditorCustomNodeBuilder_TargetLayers::OnFillLayer, Target),
+					FCanExecuteAction::CreateLambda([Target]()
+						{
+							return Target->LayerInfoObj != nullptr;
+						}));
 
-				// Clear
-				FUIAction ClearAction = FUIAction(FExecuteAction::CreateStatic(&FLandscapeEditorCustomNodeBuilder_TargetLayers::OnClearLayer, Target));
-				MenuBuilder.AddMenuEntry(LOCTEXT("LayerContextMenu.Clear", "Clear Layer"), LOCTEXT("LayerContextMenu.Clear_Tooltip", "Clears this layer to 0% across the entire landscape. If this is a weight-blended (normal) layer, other weight-blended layers will be adjusted to compensate."), FSlateIcon(), ClearAction);
+				MenuBuilder.AddMenuEntry(LOCTEXT("LayerContextMenu.Fill", "Fill Layer"),
+					LOCTEXT("LayerContextMenu.Fill_Tooltip", "Fills this layer to 100% across the entire landscape. If this is a weight-blended layer, all other weight-blended layers will be cleared."),
+					FSlateIcon(),
+					FillAction);
 
-				// Rebuild material instances
-				FUIAction RebuildAction = FUIAction(FExecuteAction::CreateStatic(&FLandscapeEditorCustomNodeBuilder_TargetLayers::OnRebuildMICs, Target));
-				MenuBuilder.AddMenuEntry(LOCTEXT("LayerContextMenu.Rebuild", "Rebuild Materials"), LOCTEXT("LayerContextMenu.Rebuild_Tooltip", "Rebuild material instances used for this landscape."), FSlateIcon(), RebuildAction);
+				// Clear
+				FUIAction ClearAction = FUIAction(FExecuteAction::CreateStatic(&FLandscapeEditorCustomNodeBuilder_TargetLayers::OnClearLayer, Target),
+					FCanExecuteAction::CreateLambda([Target]()
+						{
+							return Target->LayerInfoObj != nullptr;
+						}));
+
+				MenuBuilder.AddMenuEntry(LOCTEXT("LayerContextMenu.Clear", "Clear Layer"),
+					LOCTEXT("LayerContextMenu.Clear_Tooltip", "Clears this layer to 0% across the entire landscape. If this is a weight-blended layer, other weight-blended layers will be adjusted to compensate."),
+					FSlateIcon(),
+					ClearAction);
 			}
 			else if (Target->TargetType == ELandscapeToolTargetType::Visibility)
 			{
 				MenuBuilder.AddMenuSeparator();
 
 				// Clear
-				FUIAction ClearAction = FUIAction(FExecuteAction::CreateStatic(&FLandscapeEditorCustomNodeBuilder_TargetLayers::OnClearLayer, Target));
-				MenuBuilder.AddMenuEntry(LOCTEXT("LayerContextMenu.ClearHoles", "Remove all Holes"), FText(), FSlateIcon(), ClearAction);
+				FUIAction ClearAction = FUIAction(FExecuteAction::CreateStatic(&FLandscapeEditorCustomNodeBuilder_TargetLayers::OnClearLayer, Target),
+					FCanExecuteAction::CreateLambda([Target]()
+						{
+							return Target->LayerInfoObj != nullptr;
+						}));
+
+				MenuBuilder.AddMenuEntry(LOCTEXT("LayerContextMenu.ClearHoles", "Remove All Holes"), FText(), FSlateIcon(), ClearAction);
 			}
 		}
-		MenuBuilder.EndSection();
-
-		return MenuBuilder.MakeWidget();
 	}
+	MenuBuilder.EndSection();
 
-	return nullptr;
+	return MenuBuilder.MakeWidget();
 }
 
 void FLandscapeEditorCustomNodeBuilder_TargetLayers::OnExportLayer(const TSharedRef<FLandscapeTargetListInfo> Target)
@@ -1659,6 +1704,15 @@ void FLandscapeEditorCustomNodeBuilder_TargetLayers::OnHeightmapLayerContextMenu
 	}
 }
 
+void FLandscapeEditorCustomNodeBuilder_TargetLayers::OnRenameLayer(const int32 InLayerIndex)
+{
+	// On Rename can be called from Tick or Quick Actions menu
+	if (FEdModeLandscape* LandscapeEdMode = GetEditorMode(); InlineTextBlocks.IsValidIndex(InLayerIndex) && InlineTextBlocks[InLayerIndex].IsValid())
+	{
+		InlineTextBlocks[InLayerIndex]->EnterEditingMode();
+	}
+}
+
 void FLandscapeEditorCustomNodeBuilder_TargetLayers::OnFillLayer(const TSharedRef<FLandscapeTargetListInfo> Target)
 {
 	FScopedTransaction Transaction(LOCTEXT("Undo_FillLayer", "Filling Landscape Layer"));
diff --git a/Engine/Source/Editor/LandscapeEditor/Private/LandscapeEditorDetailCustomization_TargetLayers.h b/Engine/Source/Editor/LandscapeEditor/Private/LandscapeEditorDetailCustomization_TargetLayers.h
--- a/Engine/Source/Editor/LandscapeEditor/Private/LandscapeEditorDetailCustomization_TargetLayers.h
+++ b/Engine/Source/Editor/LandscapeEditor/Private/LandscapeEditorDetailCustomization_TargetLayers.h
@@ -70,8 +70,8 @@ public:
 	virtual void SetOnRebuildChildren( FSimpleDelegate InOnRegenerateChildren ) override;
 	virtual void GenerateHeaderRowContent( FDetailWidgetRow& NodeRow ) override;
 	virtual void GenerateChildContent( IDetailChildrenBuilder& ChildrenBuilder ) override;
-	virtual void Tick( float DeltaTime ) override {}
-	virtual bool RequiresTick() const override { return false; }
+	virtual void Tick(float DeltaTime) override;
+	virtual bool RequiresTick() const override { return true; }
 	virtual bool InitiallyCollapsed() const override { return false; }
 	virtual FName GetName() const override { return "TargetLayers"; }
 
@@ -82,6 +82,8 @@ protected:
 	TSharedRef<IPropertyHandle> TargetDisplayOrderPropertyHandle;
 	TSharedRef<IPropertyHandle> TargetShowUnusedLayersPropertyHandle;
 	TSharedPtr<SSearchBox> LayersFilterSearchBox;
+	/** Widgets for displaying and editing the target layer name */
+	TArray<TSharedPtr<SInlineEditableTextBlock>> InlineTextBlocks;
 
 	static class FEdModeLandscape* GetEditorMode();
 
@@ -89,7 +91,8 @@ protected:
 
 	static bool GetTargetLayerIsSelected(const TSharedRef<FLandscapeTargetListInfo> Target);
 	static void OnTargetSelectionChanged(const TSharedRef<FLandscapeTargetListInfo> Target);
-	static TSharedPtr<SWidget> OnTargetLayerContextMenuOpening(const TSharedRef<FLandscapeTargetListInfo> Target);
+	TSharedPtr<SWidget> OnTargetLayerContextMenuOpening(const TSharedRef<FLandscapeTargetListInfo> Target, const int32 InLayerIndex);
+	void OnRenameLayer(const int32 InLayerIndex);
 	static void OnExportLayer(const TSharedRef<FLandscapeTargetListInfo> Target);
 	static void OnImportLayer(const TSharedRef<FLandscapeTargetListInfo> Target);
 	static void OnReimportLayer(const TSharedRef<FLandscapeTargetListInfo> Target);
@@ -99,7 +102,7 @@ protected:
 	static void OnClearLayer(const TSharedRef<FLandscapeTargetListInfo> Target);
 	static void OnRebuildMICs(const TSharedRef<FLandscapeTargetListInfo> Target);
 	static bool ShouldFilterLayerInfo(const struct FAssetData& AssetData, FName LayerName);
-	static void OnTargetLayerSetObject(const FAssetData& AssetData, const TSharedRef<FLandscapeTargetListInfo> Target);
+	static void OnTargetLayerSetObject(const FAssetData& AssetData, const TSharedRef<FLandscapeTargetListInfo> Target); 
 	static EVisibility GetTargetLayerInfoSelectorVisibility(const TSharedRef<FLandscapeTargetListInfo> Target);
 	static bool GetTargetLayerCreateEnabled(const TSharedRef<FLandscapeTargetListInfo> Target);
 	static EVisibility GetTargetLayerMakePublicVisibility(const TSharedRef<FLandscapeTargetListInfo> Target);
@@ -149,6 +152,7 @@ protected:
 
 	void OnFilterTextChanged(const FText& InFilterText);
 	void OnFilterTextCommitted(const FText& InFilterText, ETextCommit::Type InCommitType);
+	int32 GetWeightmapTargetLayerCount() const;
 	bool HasWeightmapTargetLayers() const;
 	EVisibility GetLayersFilterVisibility() const;
 	FText GetLayersFilterText() const;
