Landscape Paint [Bug] - Fix invalid access error when entering Paint mode in a level with no Landscape Actor(s)

diff --git a/Engine/Source/Editor/LandscapeEditor/Private/LandscapeEditorDetailCustomization_TargetLayers.cpp b/Engine/Source/Editor/LandscapeEditor/Private/LandscapeEditorDetailCustomization_TargetLayers.cpp
--- a/Engine/Source/Editor/LandscapeEditor/Private/LandscapeEditorDetailCustomization_TargetLayers.cpp
+++ b/Engine/Source/Editor/LandscapeEditor/Private/LandscapeEditorDetailCustomization_TargetLayers.cpp
@@ -294,8 +294,7 @@ EVisibility FLandscapeEditorDetailCustomization_TargetLayers::GetPaintingRestric
 
 EVisibility FLandscapeEditorDetailCustomization_TargetLayers::GetVisibilityMaskTipVisibility()
 {
-	FEdModeLandscape* LandscapeEdMode = GetEditorMode();
-	if (LandscapeEdMode && LandscapeEdMode->CurrentToolTarget.LandscapeInfo.IsValid())
+	if (FEdModeLandscape* LandscapeEdMode = GetEditorMode(); (LandscapeEdMode != nullptr) && LandscapeEdMode->CurrentToolTarget.LandscapeInfo.IsValid())
 	{
 		if (LandscapeEdMode->CurrentToolTarget.TargetType == ELandscapeToolTargetType::Visibility)
 		{
@@ -315,8 +314,7 @@ EVisibility FLandscapeEditorDetailCustomization_TargetLayers::GetVisibilityMaskT
 
 EVisibility FLandscapeEditorDetailCustomization_TargetLayers::GetPopulateTargetLayersInfoTipVisibility()
 {
-	FEdModeLandscape* LandscapeEdMode = GetEditorMode();
-	if (LandscapeEdMode && LandscapeEdMode->CurrentToolTarget.LandscapeInfo.IsValid())
+	if (FEdModeLandscape* LandscapeEdMode = GetEditorMode(); (LandscapeEdMode != nullptr) && LandscapeEdMode->CurrentToolTarget.LandscapeInfo.IsValid())
 	{
 		if ((LandscapeEdMode->CurrentToolTarget.TargetType == ELandscapeToolTargetType::Weightmap)
 			|| (LandscapeEdMode->CurrentToolTarget.TargetType == ELandscapeToolTargetType::Invalid)) // ELandscapeToolTargetType::Invalid means "weightmap with no valid paint layer" 
@@ -331,8 +329,7 @@ EVisibility FLandscapeEditorDetailCustomization_TargetLayers::GetPopulateTargetL
 
 EVisibility FLandscapeEditorDetailCustomization_TargetLayers::GetFilteredTargetLayersListInfoTipVisibility()
 {
-	FEdModeLandscape* LandscapeEdMode = GetEditorMode();
-	if (LandscapeEdMode && LandscapeEdMode->CurrentToolTarget.LandscapeInfo.IsValid())
+	if (FEdModeLandscape* LandscapeEdMode = GetEditorMode(); (LandscapeEdMode != nullptr) && LandscapeEdMode->CurrentToolTarget.LandscapeInfo.IsValid())
 	{
 		if ((LandscapeEdMode->CurrentToolTarget.TargetType == ELandscapeToolTargetType::Weightmap)
 			|| (LandscapeEdMode->CurrentToolTarget.TargetType == ELandscapeToolTargetType::Invalid)) // ELandscapeToolTargetType::Invalid means "weightmap with no valid paint layer" 
@@ -513,31 +510,36 @@ void FLandscapeEditorCustomNodeBuilder_TargetLayers::GenerateHeaderRowContent(FD
 
 const TArray<const TSharedRef<FLandscapeTargetListInfo>> FLandscapeEditorCustomNodeBuilder_TargetLayers::GetUnassignedTargetLayersFromMaterial() const
 {
-	const FEdModeLandscape* LandscapeEdMode = GetEditorMode();
-	check(LandscapeEdMode);
-
-	ALandscape* LandscapeActor = LandscapeEdMode->CurrentToolTarget.LandscapeInfo->LandscapeActor.Get();
+	TArray<const TSharedRef<FLandscapeTargetListInfo>> TargetLayerList;
 
-	TSet<FName> LayerNames;
-	LandscapeActor->GetLandscapeInfo()->ForEachLandscapeProxy([&LayerNames](ALandscapeProxy* Proxy)
+	if (const FEdModeLandscape* LandscapeEdMode = GetEditorMode(); (LandscapeEdMode != nullptr) && LandscapeEdMode->CurrentToolTarget.LandscapeInfo.IsValid())
 	{
-		LayerNames.Append(Proxy->RetrieveTargetLayerNamesFromMaterials());
-		return true;
-	});
+		TWeakObjectPtr<ALandscape> LandscapeActor = LandscapeEdMode->CurrentToolTarget.LandscapeInfo->LandscapeActor;
 
-	TArray<const TSharedRef<FLandscapeTargetListInfo>> TargetLayerList;
+		if (!LandscapeActor.Get())
+		{
+			return TargetLayerList;
+		}
 
-	for (const FName& LayerName : LayerNames)
-	{
-		if (!LandscapeActor->GetTargetLayers().Find(LayerName))
+		TSet<FName> LayerNames;
+		LandscapeActor->GetLandscapeInfo()->ForEachLandscapeProxy([&LayerNames](ALandscapeProxy* Proxy)
+			{
+				LayerNames.Append(Proxy->RetrieveTargetLayerNamesFromMaterials());
+				return true;
+			});
+
+		for (const FName& LayerName : LayerNames)
 		{
-			const FLandscapeInfoLayerSettings LayerSettings = FLandscapeInfoLayerSettings(LayerName, LandscapeActor);
-			const TSharedRef<FLandscapeTargetListInfo> Target = MakeShareable(new FLandscapeTargetListInfo(FText::FromName(LayerName), ELandscapeToolTargetType::Weightmap, LayerSettings, LandscapeEdMode->UISettings->CurrentLayerIndex));
+			if (!LandscapeActor->GetTargetLayers().Find(LayerName))
+			{
+				const FLandscapeInfoLayerSettings LayerSettings = FLandscapeInfoLayerSettings(LayerName, LandscapeActor.Get());
+				const TSharedRef<FLandscapeTargetListInfo> Target = MakeShareable(new FLandscapeTargetListInfo(FText::FromName(LayerName), ELandscapeToolTargetType::Weightmap, LayerSettings, LandscapeEdMode->UISettings->CurrentLayerIndex));
 
-			TargetLayerList.Add(Target);
+				TargetLayerList.Add(Target);
+			}
 		}
 	}
-
+	
 	return TargetLayerList;
 }
 
@@ -550,30 +552,33 @@ FReply FLandscapeEditorCustomNodeBuilder_TargetLayers::HandleCreateLayersFromMat
 {
 	FScopedTransaction Transaction(LOCTEXT("LandscapeTargetLayer_CreateFromMaterials", "Create Target Layers from Assigned materials"));
 	
-	FEdModeLandscape* LandscapeEdMode = GetEditorMode(); 
-	if (LandscapeEdMode == nullptr)
+	if (FEdModeLandscape* LandscapeEdMode = GetEditorMode(); (LandscapeEdMode != nullptr) && LandscapeEdMode->CurrentToolTarget.LandscapeInfo.IsValid())
 	{
-		return FReply::Handled(); 
-	}
+		TWeakObjectPtr<ALandscape> LandscapeActor = LandscapeEdMode->CurrentToolTarget.LandscapeInfo->LandscapeActor;
 
-	ALandscape* LandscapeActor = LandscapeEdMode->CurrentToolTarget.LandscapeInfo->LandscapeActor.Get();
-	const TArray<const TSharedRef<FLandscapeTargetListInfo>> TargetLayerList = GetUnassignedTargetLayersFromMaterial();
+		if (!LandscapeActor.Get())
+		{
+			return FReply::Handled();
+		}
 
-	if (!TargetLayerList.IsEmpty()) 
-	{
-		for (const TSharedRef<FLandscapeTargetListInfo>& TargetLayer : TargetLayerList)
+		const TArray<const TSharedRef<FLandscapeTargetListInfo>> TargetLayerList = GetUnassignedTargetLayersFromMaterial();
+
+		if (!TargetLayerList.IsEmpty())
 		{
-			// The user may have created a new layer in the dialog, only add layers that are not yet assigned
-			if (!LandscapeActor->GetTargetLayers().Find(TargetLayer->GetLayerName()))
+			for (const TSharedRef<FLandscapeTargetListInfo>& TargetLayer : TargetLayerList)
 			{
-				LandscapeActor->AddTargetLayer(TargetLayer->GetLayerName(), FLandscapeTargetLayerSettings(TargetLayer->LayerInfoObj.Get()));
+				// The user may have created a new layer in the dialog, only add layers that are not yet assigned
+				if (!LandscapeActor->GetTargetLayers().Find(TargetLayer->GetLayerName()))
+				{
+					LandscapeActor->AddTargetLayer(TargetLayer->GetLayerName(), FLandscapeTargetLayerSettings(TargetLayer->LayerInfoObj.Get()));
+				}
 			}
-		}
-
-		LandscapeEdMode->GetLandscape()->GetLandscapeInfo()->UpdateLayerInfoMap();
-		LandscapeEdMode->UpdateTargetList();
-	}
 
+			LandscapeEdMode->GetLandscape()->GetLandscapeInfo()->UpdateLayerInfoMap();
+			LandscapeEdMode->UpdateTargetList();
+		}
+	}	 
+	
 	return FReply::Handled();
 }
 
@@ -681,34 +686,30 @@ FReply FLandscapeEditorCustomNodeBuilder_TargetLayers::HandleAutoFillTargetLayer
 
 void FLandscapeEditorCustomNodeBuilder_TargetLayers::HandleCreateLayer()
 {
-	FEdModeLandscape* LandscapeEdMode = GetEditorMode();
-	if (LandscapeEdMode == nullptr)
+	if (FEdModeLandscape* LandscapeEdMode = GetEditorMode(); (LandscapeEdMode != nullptr) && LandscapeEdMode->CurrentToolTarget.LandscapeInfo.IsValid())
 	{
-		return; 
-	}
+		TWeakObjectPtr<ALandscape> Landscape = LandscapeEdMode->CurrentToolTarget.LandscapeInfo->LandscapeActor;
 
-	TWeakObjectPtr<ALandscape> Landscape = LandscapeEdMode->CurrentToolTarget.LandscapeInfo->LandscapeActor;
-
-	if (!Landscape.Get())
-	{
-		return; 
-	}
+		if (!Landscape.Get())
+		{
+			return;
+		}
 
-	FScopedTransaction Transaction(LOCTEXT("LandscapeTargetLayer_Create", "Create a Target Layer"));
+		FScopedTransaction Transaction(LOCTEXT("LandscapeTargetLayer_Create", "Create a Target Layer"));
 
-	Landscape->AddTargetLayer();
+		Landscape->AddTargetLayer();
 
-	LandscapeEdMode->CurrentToolTarget.LandscapeInfo->UpdateLayerInfoMap();
-	LandscapeEdMode->UpdateTargetList();
+		LandscapeEdMode->CurrentToolTarget.LandscapeInfo->UpdateLayerInfoMap();
+		LandscapeEdMode->UpdateTargetList();
 
-	// Trigger a rename for the new layer, enters edit mode on next tick
-	LandscapeEdMode->PendingRenameTargetLayerIndex = GetWeightmapTargetLayerCount() - 1;
+		// Trigger a rename for the new layer, enters edit mode on next tick
+		LandscapeEdMode->PendingRenameTargetLayerIndex = GetWeightmapTargetLayerCount() - 1;
 
-	LandscapeEdMode->RefreshDetailPanel();
+		LandscapeEdMode->RefreshDetailPanel();
 
-	// Auto-select new target layer in details panel
-	OnTargetSelectionChanged(LandscapeEdMode->GetTargetList().Last());
-	
+		// Auto-select new target layer in details panel
+		OnTargetSelectionChanged(LandscapeEdMode->GetTargetList().Last());
+	}
 }
 
 TSharedRef<SWidget> FLandscapeEditorCustomNodeBuilder_TargetLayers::GetTargetLayerDisplayOrderButtonMenuContent()
@@ -893,7 +894,7 @@ void FLandscapeEditorCustomNodeBuilder_TargetLayers::OnFilterTextCommitted(const
 
 int32 FLandscapeEditorCustomNodeBuilder_TargetLayers::GetWeightmapTargetLayerCount() const
 {
-	if (const FEdModeLandscape* LandscapeEdMode = GetEditorMode(); LandscapeEdMode->CurrentToolMode)
+	if (const FEdModeLandscape* LandscapeEdMode = GetEditorMode(); (LandscapeEdMode != nullptr) && LandscapeEdMode->CurrentToolMode)
 	{
 		// ELandscapeToolTargetType::Invalid means "weightmap with no valid paint layer" so we still want to display that property if it has been marked to be displayed in Weightmap target type, to be consistent 
 		if ((LandscapeEdMode->CurrentToolTarget.TargetType == ELandscapeToolTargetType::Weightmap)
@@ -1030,7 +1031,7 @@ END_SLATE_FUNCTION_BUILD_OPTIMIZATION
 
 void FLandscapeEditorCustomNodeBuilder_TargetLayers::Tick(float DeltaTime)
 {
-	if (FEdModeLandscape* LandscapeEdMode = GetEditorMode(); LandscapeEdMode->PendingRenameTargetLayerIndex != INDEX_NONE)
+	if (FEdModeLandscape* LandscapeEdMode = GetEditorMode(); (LandscapeEdMode != nullptr) && LandscapeEdMode->PendingRenameTargetLayerIndex != INDEX_NONE)
 	{
 		OnRenameLayer(LandscapeEdMode->PendingRenameTargetLayerIndex);
 		LandscapeEdMode->PendingRenameTargetLayerIndex = INDEX_NONE;
@@ -1707,7 +1708,7 @@ void FLandscapeEditorCustomNodeBuilder_TargetLayers::OnHeightmapLayerContextMenu
 void FLandscapeEditorCustomNodeBuilder_TargetLayers::OnRenameLayer(const int32 InLayerIndex)
 {
 	// On Rename can be called from Tick or Quick Actions menu
-	if (FEdModeLandscape* LandscapeEdMode = GetEditorMode(); InlineTextBlocks.IsValidIndex(InLayerIndex) && InlineTextBlocks[InLayerIndex].IsValid())
+	if (FEdModeLandscape* LandscapeEdMode = GetEditorMode(); (LandscapeEdMode != nullptr) && InlineTextBlocks.IsValidIndex(InLayerIndex) && InlineTextBlocks[InLayerIndex].IsValid())
 	{
 		InlineTextBlocks[InLayerIndex]->EnterEditingMode();
 	}
