Made landscape layer default asset path dialog open at the current path + fixed the path being stuck to the "unset" path when selecting a valid path

diff --git a/Engine/Source/Editor/LandscapeEditor/Private/LandscapeEdMode.cpp b/Engine/Source/Editor/LandscapeEditor/Private/LandscapeEdMode.cpp
--- a/Engine/Source/Editor/LandscapeEditor/Private/LandscapeEdMode.cpp
+++ b/Engine/Source/Editor/LandscapeEditor/Private/LandscapeEdMode.cpp
@@ -2329,7 +2329,7 @@ const FString FEdModeLandscape::GetTargetLayerAssetPackagePath(const bool bIsEmp
 	{
 		// When default path is enabled but no path has been chosen we should use the shared assets path
 		// UI methods can optionally return an empty path to show a HintText
-		if (UISettings->TargetLayerAssetFilePath.DirectoryPath.Path.IsEmpty() && bIsEmptyPathValid)
+		if (!UISettings->TargetLayerAssetFilePath.DirectoryPath.Path.IsEmpty() || bIsEmptyPathValid)
 		{
 			return UISettings->TargetLayerAssetFilePath.DirectoryPath.Path;
 		}
diff --git a/Engine/Source/Editor/LandscapeEditor/Private/LandscapeEditorDetailCustomization_TargetLayers.cpp b/Engine/Source/Editor/LandscapeEditor/Private/LandscapeEditorDetailCustomization_TargetLayers.cpp
--- a/Engine/Source/Editor/LandscapeEditor/Private/LandscapeEditorDetailCustomization_TargetLayers.cpp
+++ b/Engine/Source/Editor/LandscapeEditor/Private/LandscapeEditorDetailCustomization_TargetLayers.cpp
@@ -131,17 +131,22 @@ FText FLandscapeEditorStructCustomization_FTargetLayerAssetPath::GetTargetLayerA
 
 FReply FLandscapeEditorStructCustomization_FTargetLayerAssetPath::OnSetTargetLayerAssetFilePath()
 {
+	const FEdModeLandscape* LandscapeEdMode = GetEditorMode();
+	if ((LandscapeEdMode == nullptr) || (LandscapeEdMode->UISettings == nullptr))
+	{
+		return FReply::Unhandled();
+	}
+
+	FString DialogPath = LandscapeEdMode->GetTargetLayerAssetPackagePath(/*bIsEmptyPathValid = */false);
+
 	const TSharedRef<SDlgPickPath> NewPathDlg =
 		SNew(SDlgPickPath)
-		.Title(LOCTEXT("TargetLayerAssetFilePath_Dlg", "Set Target Layer Asset File Path"));
+		.Title(LOCTEXT("TargetLayerAssetFilePath_Dlg", "Set Target Layer Asset File Path"))
+		.DefaultPath(FText::FromString(DialogPath));
 
 	if (NewPathDlg->ShowModal() != EAppReturnType::Cancel)
 	{
-		const FEdModeLandscape* LandscapeEdMode = GetEditorMode();
-		if (LandscapeEdMode && LandscapeEdMode->UISettings)
-		{
-			LandscapeEdMode->UISettings->TargetLayerAssetFilePath.DirectoryPath.Path = NewPathDlg->GetPath().ToString();
-		}
+		LandscapeEdMode->UISettings->TargetLayerAssetFilePath.DirectoryPath.Path = NewPathDlg->GetPath().ToString();
 	}
 
 	return FReply::Handled();
@@ -1090,8 +1095,6 @@ TSharedPtr<SWidget> FLandscapeEditorCustomNodeBuilder_TargetLayers::GenerateRow(
 	}
 	else
 	{
-		static const FSlateColorBrush SolidWhiteBrush = FSlateColorBrush(FColorList::White);
-
 		InlineTextBlocks.Add(TSharedPtr<SInlineEditableTextBlock>());
 
 		RowWidget = SNew(SLandscapeEditorSelectableBorder)
