Removed deprecated "make landscape layer asset public" functionality : those are always public anyway

#rnx

diff --git a/Engine/Source/Editor/LandscapeEditor/Private/LandscapeEditorDetailCustomization_TargetLayers.cpp b/Engine/Source/Editor/LandscapeEditor/Private/LandscapeEditorDetailCustomization_TargetLayers.cpp
--- a/Engine/Source/Editor/LandscapeEditor/Private/LandscapeEditorDetailCustomization_TargetLayers.cpp
+++ b/Engine/Source/Editor/LandscapeEditor/Private/LandscapeEditorDetailCustomization_TargetLayers.cpp
@@ -1287,23 +1287,6 @@ TSharedPtr<SWidget> FLandscapeEditorCustomNodeBuilder_TargetLayers::GenerateRow(
 									+ SHorizontalBox::Slot()
 									.AutoWidth()
 									.VAlign(VAlign_Center)
-									[
-										SNew(SButton)
-											.ButtonStyle(FAppStyle::Get(), "HoverHintOnly")
-											.ContentPadding(4.0f)
-											.ForegroundColor(FSlateColor::UseForeground())
-											.IsFocusable(false)
-											.ToolTipText(LOCTEXT("Tooltip_MakePublic", "Make Layer Public (move layer info into asset file)"))
-											.Visibility_Static(&FLandscapeEditorCustomNodeBuilder_TargetLayers::GetTargetLayerMakePublicVisibility, Target)
-											.OnClicked_Static(&FLandscapeEditorCustomNodeBuilder_TargetLayers::OnTargetLayerMakePublicClicked, Target)
-											[
-												SNew(SImage)
-													.Image(FAppStyle::GetBrush("LandscapeEditor.Target_MakePublic"))
-											]
-									]
-									+ SHorizontalBox::Slot()
-									.AutoWidth()
-									.VAlign(VAlign_Center)
 									[
 										SNew(SButton)
 											.ButtonStyle(FAppStyle::Get(), "HoverHintOnly")
@@ -1880,16 +1863,6 @@ bool FLandscapeEditorCustomNodeBuilder_TargetLayers::GetTargetLayerCreateEnabled
 	return false;
 }
 
-EVisibility FLandscapeEditorCustomNodeBuilder_TargetLayers::GetTargetLayerMakePublicVisibility(const TSharedRef<FLandscapeTargetListInfo> Target)
-{
-	if (Target->bValid && Target->LayerInfoObj.IsValid() && Target->LayerInfoObj->GetOutermost()->ContainsMap())
-	{
-		return EVisibility::Visible;
-	}
-
-	return EVisibility::Collapsed;
-}
-
 EVisibility FLandscapeEditorCustomNodeBuilder_TargetLayers::GetTargetLayerDeleteVisibility(const TSharedRef<FLandscapeTargetListInfo> Target)
 {
 	if (!Target->bValid)
@@ -1995,30 +1968,6 @@ void FLandscapeEditorCustomNodeBuilder_TargetLayers::CreateTargetLayerInfoAsset(
 	}
 }
 
-FReply FLandscapeEditorCustomNodeBuilder_TargetLayers::OnTargetLayerMakePublicClicked(const TSharedRef<FLandscapeTargetListInfo> Target)
-{
-	FScopedTransaction Transaction(LOCTEXT("Undo_MakePublic", "Make Layer Public"));
-	TArray<UObject*> Objects;
-	Objects.Add(Target->LayerInfoObj.Get());
-
-	FString Path = Target->Owner->GetOutermost()->GetName() + TEXT("_sharedassets");
-	bool bSucceed = ObjectTools::RenameObjects(Objects, false, TEXT(""), Path);
-	if (bSucceed)
-	{
-		FEdModeLandscape* Mode = (FEdModeLandscape*)GLevelEditorModeTools().GetActiveMode(FBuiltinEditorModes::EM_Landscape);
-		if (Mode)
-		{
-			Mode->UpdateTargetList();
-		}
-	}
-	else
-	{
-		Transaction.Cancel();
-	}
-
-	return FReply::Handled();
-}
-
 FReply FLandscapeEditorCustomNodeBuilder_TargetLayers::OnTargetLayerDeleteClicked(const TSharedRef<FLandscapeTargetListInfo> Target)
 {
 	check(Target->LandscapeInfo.IsValid());
diff --git a/Engine/Source/Editor/LandscapeEditor/Private/LandscapeEditorDetailCustomization_TargetLayers.h b/Engine/Source/Editor/LandscapeEditor/Private/LandscapeEditorDetailCustomization_TargetLayers.h
--- a/Engine/Source/Editor/LandscapeEditor/Private/LandscapeEditorDetailCustomization_TargetLayers.h
+++ b/Engine/Source/Editor/LandscapeEditor/Private/LandscapeEditorDetailCustomization_TargetLayers.h
@@ -105,12 +105,10 @@ protected:
 	static void OnTargetLayerSetObject(const FAssetData& AssetData, const TSharedRef<FLandscapeTargetListInfo> Target); 
 	static EVisibility GetTargetLayerInfoSelectorVisibility(const TSharedRef<FLandscapeTargetListInfo> Target);
 	static bool GetTargetLayerCreateEnabled(const TSharedRef<FLandscapeTargetListInfo> Target);
-	static EVisibility GetTargetLayerMakePublicVisibility(const TSharedRef<FLandscapeTargetListInfo> Target);
 	static EVisibility GetTargetLayerDeleteVisibility(const TSharedRef<FLandscapeTargetListInfo> Target);
 	static TSharedRef<SWidget> OnGetTargetLayerCreateMenu(const TSharedRef<FLandscapeTargetListInfo> Target);
 	static void OnTargetLayerCreateClicked(const TSharedRef<FLandscapeTargetListInfo> Target, bool bNoWeightBlend);
 	static void CreateTargetLayerInfoAsset(const TSharedRef<FLandscapeTargetListInfo> Target, const FString& PackageName, const FName& LayerObjectName, const bool bNoWeightBlend);
-	static FReply OnTargetLayerMakePublicClicked(const TSharedRef<FLandscapeTargetListInfo> Target);
 	static FReply OnTargetLayerDeleteClicked(const TSharedRef<FLandscapeTargetListInfo> Target);
 	static FSlateColor GetLayerUsageDebugColor(const TSharedRef<FLandscapeTargetListInfo> Target);
 	static EVisibility GetDebugModeLayerUsageVisibility(const TSharedRef<FLandscapeTargetListInfo> Target);
