Removed auto-filling of empty landscape target layers : this is purely arbitrary as it will almost never be the user's intention to fill a layer with a fully painted layer. It's almost a complete revert of 3249481 (the public FillEmptyLayers function on FLandscapeEditDataInterface has been kept for backwards compatibility reasons)

diff --git a/Engine/Source/Editor/LandscapeEditor/Private/LandscapeEdMode.cpp b/Engine/Source/Editor/LandscapeEditor/Private/LandscapeEdMode.cpp
--- a/Engine/Source/Editor/LandscapeEditor/Private/LandscapeEdMode.cpp
+++ b/Engine/Source/Editor/LandscapeEditor/Private/LandscapeEdMode.cpp
@@ -4341,47 +4341,6 @@ FGuid FEdModeLandscape::GetCurrentLayerGuid() const
 	return CurrentLayer ? CurrentLayer->Guid : FGuid();
 }
 
-bool FEdModeLandscape::NeedToFillEmptyMaterialLayers() const
-{
-	if (!CurrentToolTarget.LandscapeInfo.IsValid() || !CurrentToolTarget.LandscapeInfo->LandscapeActor.IsValid())
-	{
-		return false;
-	}
-
-	bool bCanFill = true;
-
-	CurrentToolTarget.LandscapeInfo->ForEachLandscapeProxy([&](ALandscapeProxy* Proxy)
-	{
-		ALandscape* Landscape = Proxy->GetLandscapeActor();
-
-		if (Landscape != nullptr)
-		{
-			for (const FLandscapeLayer& Layer : Landscape->GetLayers())
-			{
-				for (ULandscapeComponent* Component : Proxy->LandscapeComponents)
-				{
-					const FLandscapeLayerComponentData* LayerComponentData = Component->GetLayerData(Layer.Guid);
-
-					if (LayerComponentData != nullptr)
-					{
-						for (const FWeightmapLayerAllocationInfo& Alloc : LayerComponentData->WeightmapData.LayerAllocations)
-						{
-							if (Alloc.LayerInfo != nullptr)
-							{
-								bCanFill = false;
-								return false;
-							}
-						}
-					}
-				}
-			}
-		}
-		return true;
-	});	
-
-	return bCanFill;
-}
-
 void FEdModeLandscape::UpdateBrushList()
 {
 	BrushList.Empty();
diff --git a/Engine/Source/Editor/LandscapeEditor/Private/LandscapeEdMode.h b/Engine/Source/Editor/LandscapeEditor/Private/LandscapeEdMode.h
--- a/Engine/Source/Editor/LandscapeEditor/Private/LandscapeEdMode.h
+++ b/Engine/Source/Editor/LandscapeEditor/Private/LandscapeEdMode.h
@@ -618,7 +618,6 @@ public:
 	
 	void DuplicateBrush(class ALandscapeBlueprintBrushBase* InBrush);
 
-	bool NeedToFillEmptyMaterialLayers() const;
 	void RequestLayersContentUpdate(ELandscapeLayerUpdateMode InUpdateMode);
 	void RequestLayersContentUpdateForceAll(ELandscapeLayerUpdateMode InUpdateMode = ELandscapeLayerUpdateMode::Update_All);
 
diff --git a/Engine/Source/Editor/LandscapeEditor/Private/LandscapeEdModeComponentTools.cpp b/Engine/Source/Editor/LandscapeEditor/Private/LandscapeEdModeComponentTools.cpp
--- a/Engine/Source/Editor/LandscapeEditor/Private/LandscapeEdModeComponentTools.cpp
+++ b/Engine/Source/Editor/LandscapeEditor/Private/LandscapeEdModeComponentTools.cpp
@@ -788,73 +788,6 @@ public:
 						CollisionComp->RecreateCollision();
 					}
 				}
-
-				TMap<ULandscapeLayerInfoObject*, int32> NeighbourLayerInfoObjectCount;
-
-				{
-					const FLandscapeLayer* LandscapeLayer = Landscape ? Landscape->GetLayerConst(0) : nullptr;
-					FScopedSetLandscapeEditingLayer Scope(Landscape, LandscapeLayer ? LandscapeLayer->Guid : FGuid(), [=] { });
-
-					// Cover 9 tiles around us to determine which object should we use by default
-					for (int32 ComponentIndexX = ComponentIndexX1 - 1; ComponentIndexX <= ComponentIndexX2 + 1; ++ComponentIndexX)
-					{
-						for (int32 ComponentIndexY = ComponentIndexY1 - 1; ComponentIndexY <= ComponentIndexY2 + 1; ++ComponentIndexY)
-						{
-							ULandscapeComponent* NeighbourComponent = LandscapeInfo->XYtoComponentMap.FindRef(FIntPoint(ComponentIndexX, ComponentIndexY));
-
-							if (NeighbourComponent != nullptr && NeighbourComponent != NewComponent)
-							{
-								ULandscapeInfo* NeighbourLandscapeInfo = NeighbourComponent->GetLandscapeInfo();
-
-								for (int32 i = 0; i < NeighbourLandscapeInfo->Layers.Num(); ++i)
-								{
-									ULandscapeLayerInfoObject* NeighbourLayerInfo = NeighbourLandscapeInfo->Layers[i].LayerInfoObj;
-
-									if (NeighbourLayerInfo != nullptr)
-									{
-										TArray<uint8> WeightmapTextureData;
-
-										FLandscapeComponentDataInterface DataInterface(NeighbourComponent);
-										DataInterface.GetWeightmapTextureData(NeighbourLayerInfo, WeightmapTextureData, true);
-
-										if (WeightmapTextureData.Num() > 0)
-										{
-											int32* Count = NeighbourLayerInfoObjectCount.Find(NeighbourLayerInfo);
-
-											if (Count == nullptr)
-											{
-												Count = &NeighbourLayerInfoObjectCount.Add(NeighbourLayerInfo, 1);
-											}
-
-											for (uint8 Value : WeightmapTextureData)
-											{
-												(*Count) += Value;
-											}
-										}
-									}
-								}
-							}
-						}
-					}
-
-					int32 BestLayerInfoObjectCount = 0;
-					ULandscapeLayerInfoObject* BestLayerInfoObject = nullptr;
-
-					for (auto& LayerInfoObjectCount : NeighbourLayerInfoObjectCount)
-					{
-						if (LayerInfoObjectCount.Value > BestLayerInfoObjectCount)
-						{
-							BestLayerInfoObjectCount = LayerInfoObjectCount.Value;
-							BestLayerInfoObject = LayerInfoObjectCount.Key;
-						}
-					}
-
-					if (BestLayerInfoObject != nullptr)
-					{
-						FLandscapeEditDataInterface LandscapeEdit(LandscapeInfo);
-						NewComponent->FillLayer(BestLayerInfoObject, LandscapeEdit);
-					}
-				}				
 			}
 
 			// Add/update "add collision" around the newly added components
diff --git a/Engine/Source/Editor/LandscapeEditor/Private/LandscapeEditorDetailCustomization_TargetLayers.cpp b/Engine/Source/Editor/LandscapeEditor/Private/LandscapeEditorDetailCustomization_TargetLayers.cpp
--- a/Engine/Source/Editor/LandscapeEditor/Private/LandscapeEditorDetailCustomization_TargetLayers.cpp
+++ b/Engine/Source/Editor/LandscapeEditor/Private/LandscapeEditorDetailCustomization_TargetLayers.cpp
@@ -1716,30 +1716,6 @@ void FLandscapeEditorCustomNodeBuilder_TargetLayers::OnFillLayer(const TSharedRe
 	}
 }
 
-void FLandscapeEditorCustomNodeBuilder_TargetLayers::FillEmptyLayers(ULandscapeInfo* LandscapeInfo, ULandscapeLayerInfoObject* LandscapeInfoObject)
-{
-	FEdModeLandscape* LandscapeEdMode = GetEditorMode();
-	if (LandscapeEdMode)
-	{
-		FLandscapeEditDataInterface LandscapeEdit(LandscapeInfo);
-
-		if (LandscapeEdMode->CanHaveLandscapeLayersContent())
-		{
-			if (LandscapeEdMode->NeedToFillEmptyMaterialLayers())
-			{
-				FScopedSetLandscapeEditingLayer Scope(LandscapeEdMode->GetLandscape(), LandscapeEdMode->GetCurrentLayerGuid());
-				LandscapeEdit.FillEmptyLayers(LandscapeInfoObject);
-			}
-		}
-		else
-		{
-			LandscapeEdit.FillEmptyLayers(LandscapeInfoObject);
-		}
-	}
-
-}
-
-
 void FLandscapeEditorCustomNodeBuilder_TargetLayers::OnClearLayer(const TSharedRef<FLandscapeTargetListInfo> Target)
 {
 	FScopedTransaction Transaction(LOCTEXT("Undo_ClearLayer", "Clearing Landscape Layer"));
@@ -1833,8 +1809,6 @@ void FLandscapeEditorCustomNodeBuilder_TargetLayers::OnTargetLayerSetObject(cons
 				LandscapeEdMode->SetCurrentTargetLayer(Target->LayerName, SelectedLayerInfo);
 				LandscapeEdMode->UpdateTargetList();
 			}
-
-			FillEmptyLayers(LandscapeInfo, SelectedLayerInfo);
 		}
 		else
 		{
@@ -1962,8 +1936,6 @@ void FLandscapeEditorCustomNodeBuilder_TargetLayers::CreateTargetLayerInfoAsset(
 		ALandscape* LandscapeActor = Target->LandscapeInfo->LandscapeActor.Get();
 		LandscapeActor->UpdateTargetLayer(FName(LayerName), FLandscapeTargetLayerSettings(LayerInfo));
 
-		FillEmptyLayers(LandscapeInfo, LayerInfo);
-
 		LandscapeEdMode->UpdateTargetList();
 	}
 }
diff --git a/Engine/Source/Editor/LandscapeEditor/Private/LandscapeEditorDetailCustomization_TargetLayers.h b/Engine/Source/Editor/LandscapeEditor/Private/LandscapeEditorDetailCustomization_TargetLayers.h
--- a/Engine/Source/Editor/LandscapeEditor/Private/LandscapeEditorDetailCustomization_TargetLayers.h
+++ b/Engine/Source/Editor/LandscapeEditor/Private/LandscapeEditorDetailCustomization_TargetLayers.h
@@ -98,7 +98,6 @@ protected:
 	static void OnReimportLayer(const TSharedRef<FLandscapeTargetListInfo> Target);
 	static void OnHeightmapLayerContextMenu(const TSharedRef<FLandscapeTargetListInfo> Target);
 	static void OnFillLayer(const TSharedRef<FLandscapeTargetListInfo> Target);
-	static void FillEmptyLayers(ULandscapeInfo* LandscapeInfo, ULandscapeLayerInfoObject* LandscapeInfoObject);
 	static void OnClearLayer(const TSharedRef<FLandscapeTargetListInfo> Target);
 	static void OnRebuildMICs(const TSharedRef<FLandscapeTargetListInfo> Target);
 	static bool ShouldFilterLayerInfo(const struct FAssetData& AssetData, FName LayerName);
