Add option for landscape sculpt and erase tools to apply only when moving the brush.

diff --git a/Engine/Source/Editor/LandscapeEditor/Private/LandscapeEdModePaintTools.cpp b/Engine/Source/Editor/LandscapeEditor/Private/LandscapeEdModePaintTools.cpp
--- a/Engine/Source/Editor/LandscapeEditor/Private/LandscapeEdModePaintTools.cpp
+++ b/Engine/Source/Editor/LandscapeEditor/Private/LandscapeEdModePaintTools.cpp
@@ -103,8 +103,6 @@ class FLandscapeToolStrokePaint : public FLandscapeToolStrokePaintBase<FWeightma
 	bool bIsAllowListMode;
 	bool bAddToAllowList;
 public:
-	// Heightmap sculpt tool will continuously sculpt in the same location, weightmap paint tool doesn't
-	enum { UseContinuousApply = false };
 
 	FLandscapeToolStrokePaint(FEdModeLandscape* InEdMode, FEditorViewportClient* InViewportClient, const FLandscapeToolTarget& InTarget)
 		: Super(InEdMode, InViewportClient, InTarget)
@@ -112,6 +110,8 @@ public:
 		                   (InViewportClient->Viewport->KeyState(EKeys::Equals) || InViewportClient->Viewport->KeyState(EKeys::Hyphen)))
 		, bAddToAllowList(bIsAllowListMode && InViewportClient->Viewport->KeyState(EKeys::Equals))
 	{
+		// bApplyWithoutMovingPaint option only applies to the component brush
+		bUseContinuousApply = InEdMode->UISettings->bApplyWithoutMovingPaint && InEdMode->CurrentBrush->GetBrushType() == ELandscapeBrushType::Component;
 	}
 
 	void Apply(FEditorViewportClient* ViewportClient, FLandscapeBrush* Brush, const ULandscapeEditorObject* UISettings, const TArray<FLandscapeToolInteractorPosition>& InteractorPositions)
@@ -320,13 +320,12 @@ class FLandscapeToolStrokeErase : public FLandscapeToolStrokePaintBase<FHeightma
 	const ValueType FlattenHeight;
 
 public:
-	// Heightmap sculpt tool will continuously sculpt in the same location, weightmap paint tool doesn't
-	enum { UseContinuousApply = true };
 
 	FLandscapeToolStrokeErase(FEdModeLandscape* InEdMode, FEditorViewportClient* InViewportClient, const FLandscapeToolTarget& InTarget)
 		: Super(InEdMode, InViewportClient, InTarget)
 		, FlattenHeight(LandscapeDataAccess::GetTexHeight(0.f))
 	{
+		bUseContinuousApply = InEdMode->UISettings->bApplyWithoutMovingHeight;
 	}
 
 	void Apply(FEditorViewportClient* ViewportClient, FLandscapeBrush* Brush, const ULandscapeEditorObject* UISettings, const TArray<FLandscapeToolInteractorPosition>& InteractorPositions)
@@ -402,12 +401,11 @@ class FLandscapeToolStrokeSculpt : public FLandscapeToolStrokePaintBase<FHeightm
 	using ValueType = ToolTarget::CacheClass::DataType;
 
 public:
-	// Heightmap sculpt tool will continuously sculpt in the same location, weightmap paint tool doesn't
-	enum { UseContinuousApply = true };
 
 	FLandscapeToolStrokeSculpt(FEdModeLandscape* InEdMode, FEditorViewportClient* InViewportClient, const FLandscapeToolTarget& InTarget)
 		: Super(InEdMode, InViewportClient, InTarget)
 	{
+		bUseContinuousApply = InEdMode->UISettings->bApplyWithoutMovingHeight;
 	}
 
 	virtual void SetEditLayer(const FGuid& EditLayerGUID) override
diff --git a/Engine/Source/Editor/LandscapeEditor/Private/LandscapeEdModeTools.h b/Engine/Source/Editor/LandscapeEditor/Private/LandscapeEdModeTools.h
--- a/Engine/Source/Editor/LandscapeEditor/Private/LandscapeEdModeTools.h
+++ b/Engine/Source/Editor/LandscapeEditor/Private/LandscapeEdModeTools.h
@@ -1353,8 +1353,8 @@ struct FWeightmapToolTarget
 class FLandscapeToolStrokeBase : protected FGCObject
 {
 public:
-	// Whether to call Apply() every frame even if the mouse hasn't moved
-	enum { UseContinuousApply = false };
+	// Whether to call Apply() every frame even if the mouse hasn't moved.  May be overwritten by child classes.
+	bool bUseContinuousApply = false;
 
 	// This is also the expected signature of derived class constructor used by FLandscapeToolBase
 	FLandscapeToolStrokeBase(FEdModeLandscape* InEdMode, FEditorViewportClient* InViewportClient, const FLandscapeToolTarget& InTarget)
@@ -1480,7 +1480,7 @@ public:
 				ViewportClient->Invalidate(false, false);
 				InteractorPositions.Empty(1);
 			}
-			else if (TStrokeClass::UseContinuousApply && TimeSinceLastInteractorMove >= 0.25f)
+			else if (ToolStroke->bUseContinuousApply && TimeSinceLastInteractorMove >= 0.25f)
 			{
 				InteractorPositions.Emplace(LastInteractorPosition, IsModifierPressed(ViewportClient));
 				ToolStroke->Apply(ViewportClient, EdMode->CurrentBrush, EdMode->UISettings, InteractorPositions);
diff --git a/Engine/Source/Editor/LandscapeEditor/Private/LandscapeEditorObject.cpp b/Engine/Source/Editor/LandscapeEditor/Private/LandscapeEditorObject.cpp
--- a/Engine/Source/Editor/LandscapeEditor/Private/LandscapeEditorObject.cpp
+++ b/Engine/Source/Editor/LandscapeEditor/Private/LandscapeEditorObject.cpp
@@ -109,6 +109,8 @@ ULandscapeEditorObject::ULandscapeEditorObject(const FObjectInitializer& ObjectI
 	, BrushFalloff(0.5f)
 	, PaintBrushFalloff(0.5f)
 	, bUseClayBrush(false)
+	, bApplyWithoutMovingHeight(true)
+	, bApplyWithoutMovingPaint(true)
 
 	, AlphaBrushScale(0.5f)
 	, bAlphaBrushAutoRotate(true)
@@ -272,6 +274,8 @@ void ULandscapeEditorObject::Load()
 	GConfig->GetFloat(TEXT("LandscapeEdit"), TEXT("RampSideFalloff"), RampSideFalloff, GEditorPerProjectIni);
 
 	GConfig->GetBool(TEXT("LandscapeEdit"), TEXT("bCombinedLayersOperation"), bCombinedLayersOperation, GEditorPerProjectIni);
+	GConfig->GetBool(TEXT("LandscapeEdit"), TEXT("bApplyWithoutMovingHeight"), bApplyWithoutMovingHeight, GEditorPerProjectIni);
+	GConfig->GetBool(TEXT("LandscapeEdit"), TEXT("bApplyWithoutMovingPaint"), bApplyWithoutMovingPaint, GEditorPerProjectIni);
 
 	GConfig->GetInt(TEXT("LandscapeEdit"), TEXT("ErodeThresh"), ErodeThresh, GEditorPerProjectIni);
 	GConfig->GetInt(TEXT("LandscapeEdit"), TEXT("ErodeIterationNum"), ErodeIterationNum, GEditorPerProjectIni);
@@ -427,7 +431,9 @@ void ULandscapeEditorObject::Save()
 	GConfig->SetFloat(TEXT("LandscapeEdit"), TEXT("RampSideFalloff"), RampSideFalloff, GEditorPerProjectIni);
 
 	GConfig->SetBool(TEXT("LandscapeEdit"), TEXT("bCombinedLayersOperation"), bCombinedLayersOperation, GEditorPerProjectIni);
-
+	GConfig->SetBool(TEXT("LandscapeEdit"), TEXT("bApplyWithoutMovingHeight"), bApplyWithoutMovingHeight, GEditorPerProjectIni);
+	GConfig->SetBool(TEXT("LandscapeEdit"), TEXT("bApplyWithoutMovingPaint"), bApplyWithoutMovingPaint, GEditorPerProjectIni);
+	
 	GConfig->SetInt(TEXT("LandscapeEdit"), TEXT("ErodeThresh"), ErodeThresh, GEditorPerProjectIni);
 	GConfig->SetInt(TEXT("LandscapeEdit"), TEXT("ErodeIterationNum"), ErodeIterationNum, GEditorPerProjectIni);
 	GConfig->SetInt(TEXT("LandscapeEdit"), TEXT("ErodeSurfaceThickness"), ErodeSurfaceThickness, GEditorPerProjectIni);
diff --git a/Engine/Source/Editor/LandscapeEditor/Public/LandscapeEditorObject.h b/Engine/Source/Editor/LandscapeEditor/Public/LandscapeEditorObject.h
--- a/Engine/Source/Editor/LandscapeEditor/Public/LandscapeEditorObject.h
+++ b/Engine/Source/Editor/LandscapeEditor/Public/LandscapeEditorObject.h
@@ -664,6 +664,14 @@ public:
 	UPROPERTY(Category="Brush Settings", EditAnywhere, NonTransactional, meta=(ShowForTools="Sculpt", ShowForBrushes="BrushSet_Circle,BrushSet_Alpha,BrushSet_Pattern"))
 	bool bUseClayBrush;
 
+	// Apply brush continuously while held stationary, after a small delay
+	UPROPERTY(Category = "Brush Settings", EditAnywhere, NonTransactional, meta=(DisplayName = "Apply without Moving", ShowForTools = "Sculpt,Erase", ShowForBrushes = "BrushSet_Circle,BrushSet_Gizmo,BrushSet_Pattern,BrushSet_Component"))
+	bool bApplyWithoutMovingHeight;
+
+	// Apply brush continuously while held stationary, after a small delay
+	UPROPERTY(Category = "Brush Settings", EditAnywhere, NonTransactional, meta = (DisplayName = "Apply without Moving", ShowForTools = "Paint", ShowForBrushes = "BrushSet_Component"))
+	bool bApplyWithoutMovingPaint;   // Internally, store a separate "Apply without Moving" value for paint mode
+
 	// Alpha/Pattern Brush:
 
 	// Scale of the brush texture. A scale of 1.000 maps the brush texture to the landscape at a 1 pixel = 1 vertex size
