Expose ApplyWithoutMoving option for all landscape paint mode brush types.  Change the default to false for paint mode.

diff --git a/Engine/Source/Editor/LandscapeEditor/Private/LandscapeEdModePaintTools.cpp b/Engine/Source/Editor/LandscapeEditor/Private/LandscapeEdModePaintTools.cpp
--- a/Engine/Source/Editor/LandscapeEditor/Private/LandscapeEdModePaintTools.cpp
+++ b/Engine/Source/Editor/LandscapeEditor/Private/LandscapeEdModePaintTools.cpp
@@ -110,8 +110,7 @@ public:
 		                   (InViewportClient->Viewport->KeyState(EKeys::Equals) || InViewportClient->Viewport->KeyState(EKeys::Hyphen)))
 		, bAddToAllowList(bIsAllowListMode && InViewportClient->Viewport->KeyState(EKeys::Equals))
 	{
-		// bApplyWithoutMovingPaint option only applies to the component brush
-		bUseContinuousApply = InEdMode->UISettings->bApplyWithoutMovingPaint && InEdMode->CurrentBrush->GetBrushType() == ELandscapeBrushType::Component;
+		bUseContinuousApply = InEdMode->UISettings->bApplyWithoutMovingPaint;
 	}
 
 	void Apply(FEditorViewportClient* ViewportClient, FLandscapeBrush* Brush, const ULandscapeEditorObject* UISettings, const TArray<FLandscapeToolInteractorPosition>& InteractorPositions)
diff --git a/Engine/Source/Editor/LandscapeEditor/Private/LandscapeEditorObject.cpp b/Engine/Source/Editor/LandscapeEditor/Private/LandscapeEditorObject.cpp
--- a/Engine/Source/Editor/LandscapeEditor/Private/LandscapeEditorObject.cpp
+++ b/Engine/Source/Editor/LandscapeEditor/Private/LandscapeEditorObject.cpp
@@ -110,7 +110,7 @@ ULandscapeEditorObject::ULandscapeEditorObject(const FObjectInitializer& ObjectI
 	, PaintBrushFalloff(0.5f)
 	, bUseClayBrush(false)
 	, bApplyWithoutMovingHeight(true)
-	, bApplyWithoutMovingPaint(true)
+	, bApplyWithoutMovingPaint(false)
 
 	, AlphaBrushScale(0.5f)
 	, bAlphaBrushAutoRotate(true)
diff --git a/Engine/Source/Editor/LandscapeEditor/Public/LandscapeEditorObject.h b/Engine/Source/Editor/LandscapeEditor/Public/LandscapeEditorObject.h
--- a/Engine/Source/Editor/LandscapeEditor/Public/LandscapeEditorObject.h
+++ b/Engine/Source/Editor/LandscapeEditor/Public/LandscapeEditorObject.h
@@ -669,7 +669,7 @@ public:
 	bool bApplyWithoutMovingHeight;
 
 	// Apply brush continuously while held stationary, after a small delay
-	UPROPERTY(Category = "Brush Settings", EditAnywhere, NonTransactional, meta = (DisplayName = "Apply without Moving", ShowForTools = "Paint", ShowForBrushes = "BrushSet_Component"))
+	UPROPERTY(Category = "Brush Settings", EditAnywhere, NonTransactional, meta=(DisplayName = "Apply without Moving", ShowForTools =" Paint", ShowForBrushes = "BrushSet_Circle,BrushSet_Alpha,BrushSet_Pattern,BrushSet_Component"))
 	bool bApplyWithoutMovingPaint;   // Internally, store a separate "Apply without Moving" value for paint mode
 
 	// Alpha/Pattern Brush:
