Prevent crash when landscape details tab is not authorized

diff --git a/Engine/Source/Editor/LandscapeEditor/Private/SLandscapeEditor.cpp b/Engine/Source/Editor/LandscapeEditor/Private/SLandscapeEditor.cpp
--- a/Engine/Source/Editor/LandscapeEditor/Private/SLandscapeEditor.cpp
+++ b/Engine/Source/Editor/LandscapeEditor/Private/SLandscapeEditor.cpp
@@ -346,8 +346,14 @@ void FLandscapeToolKit::RefreshInspectedObjectsDetailPanel()
 		}
 		InspectedObjectsTab = TabManagerPtr->TryInvokeTab(UAssetEditorUISubsystem::BottomLeftTabID);
 
-		// Refresh the selected objects
-		InspectedObjectsDetailsView->RefreshDetailPanel();
+		// It's possible we fail to create a tab (e.g. permission issue)
+		if (InspectedObjectsTab.IsValid())
+		{
+			// If we managed to create a tab, the callback should have assigned a InspectedObjectsDetailsView, though 
+			check(InspectedObjectsDetailsView != nullptr);
+			// Refresh the selected objects
+			InspectedObjectsDetailsView->RefreshDetailPanel();
+		}
 	}
 }
 
