Cosmetic change regarding display of "landscape spline control points", as suggested by design

#rnx

diff --git a/Engine/Source/Editor/LandscapeEditor/Private/LandscapeEditorDetailCustomization_MiscTools.cpp b/Engine/Source/Editor/LandscapeEditor/Private/LandscapeEditorDetailCustomization_MiscTools.cpp
--- a/Engine/Source/Editor/LandscapeEditor/Private/LandscapeEditorDetailCustomization_MiscTools.cpp
+++ b/Engine/Source/Editor/LandscapeEditor/Private/LandscapeEditorDetailCustomization_MiscTools.cpp
@@ -144,7 +144,7 @@ void FLandscapeEditorDetailCustomization_MiscTools::CustomizeDetails(IDetailLayo
 			+ SHorizontalBox::Slot()
 			[
 				SNew(SButton)
-				.ToolTipText(LOCTEXT("Spline.ApplySplines.All.Tooltip", "Deforms and paints the landscape to fit all the landscape spline segments and controlpoints."))
+				.ToolTipText(LOCTEXT("Spline.ApplySplines.All.Tooltip", "Deforms and paints the landscape to fit all the landscape spline segments and points."))
 				.Text(LOCTEXT("Spline.ApplySplines.All", "All Splines"))
 				.HAlign(HAlign_Center)
 				.OnClicked_Static(&FLandscapeEditorDetailCustomization_MiscTools::OnApplyAllSplinesButtonClicked)
@@ -152,7 +152,7 @@ void FLandscapeEditorDetailCustomization_MiscTools::CustomizeDetails(IDetailLayo
 			+ SHorizontalBox::Slot()
 			[
 				SNew(SButton)
-				.ToolTipText(LOCTEXT("Spline.ApplySplines.Tooltip", "Deforms and paints the landscape to fit only the selected landscape spline segments and controlpoints."))
+				.ToolTipText(LOCTEXT("Spline.ApplySplines.Tooltip", "Deforms and paints the landscape to fit only the selected landscape spline segments and points."))
 				.Text(LOCTEXT("Spline.ApplySplines.Selected", "Only Selected"))
 				.HAlign(HAlign_Center)
 				.OnClicked_Static(&FLandscapeEditorDetailCustomization_MiscTools::OnApplySelectedSplinesButtonClicked)
@@ -177,8 +177,8 @@ void FLandscapeEditorDetailCustomization_MiscTools::CustomizeDetails(IDetailLayo
 			+ SHorizontalBox::Slot()
 			[
 				SNew(SButton)
-				.ToolTipText(LOCTEXT("Spline.ControlPoints.All.Tooltip", "Selects all landscape spline control points in the map."))
-				.Text(LOCTEXT("Spline.ControlPoints", "Control Points"))
+				.ToolTipText(LOCTEXT("Spline.ControlPoints.All.Tooltip", "Selects all landscape spline points in the map."))
+				.Text(LOCTEXT("Spline.ControlPoints", "Points"))
 				.HAlign(HAlign_Center)
 				.OnClicked_Static(&FLandscapeEditorDetailCustomization_MiscTools::OnSelectAllControlPointsButtonClicked)
 			]
@@ -191,7 +191,7 @@ void FLandscapeEditorDetailCustomization_MiscTools::CustomizeDetails(IDetailLayo
 				.OnClicked_Static(&FLandscapeEditorDetailCustomization_MiscTools::OnSelectAllSegmentsButtonClicked)
 			]
 		];
-		ToolsCategory.AddCustomRow(LOCTEXT("Spline.bUseAutoRotateControlPoint.Selected", "Use Auto Rotate Control Point"))
+		ToolsCategory.AddCustomRow(LOCTEXT("Spline.bUseAutoRotateControlPoint.Selected", "Use Auto Rotate Point"))
 		[
 			SNew(SHorizontalBox)
 			+ SHorizontalBox::Slot()
@@ -202,7 +202,7 @@ void FLandscapeEditorDetailCustomization_MiscTools::CustomizeDetails(IDetailLayo
 				.IsChecked(this, &FLandscapeEditorDetailCustomization_MiscTools::GetbUseAutoRotateControlPoint)
 				.Content()
 				[
-					SNew(STextBlock).Text(LOCTEXT("Spline.bUseAutoRotateControlPoint.Selected", "Use Auto Rotate Control Point"))
+					SNew(STextBlock).Text(LOCTEXT("Spline.bUseAutoRotateControlPoint.Selected", "Use Auto Rotate Point"))
 				]
 			]
 		];
diff --git a/Engine/Source/Editor/LandscapeEditor/Private/LandscapeSplineDetails.cpp b/Engine/Source/Editor/LandscapeEditor/Private/LandscapeSplineDetails.cpp
--- a/Engine/Source/Editor/LandscapeEditor/Private/LandscapeSplineDetails.cpp
+++ b/Engine/Source/Editor/LandscapeEditor/Private/LandscapeSplineDetails.cpp
@@ -48,7 +48,7 @@ void FLandscapeSplineDetails::CustomizeDetails(IDetailLayoutBuilder& DetailBuild
 	[
 		SNew(STextBlock)
 			.Font(IDetailLayoutBuilder::GetDetailFont())
-			.Text(IsSegmentSelectModeEnabled() ? LOCTEXT("SelectSplineSegments", "Select Spline Segments") : LOCTEXT("SelectSplineControlPoints", "Select Spline Control Points"))
+			.Text(IsSegmentSelectModeEnabled() ? LOCTEXT("SelectSplineSegments", "Select Spline Segments") : LOCTEXT("SelectSplineControlPoints", "Select Spline Points"))
 	]
 	.ValueContent()
 	.VAlign(VAlign_Fill)
@@ -66,7 +66,7 @@ void FLandscapeSplineDetails::CustomizeDetails(IDetailLayoutBuilder& DetailBuild
 			SNew(SButton)
 				.ButtonStyle(FAppStyle::Get(), "SplineComponentDetails.SelectFirst")
 				.ContentPadding(2.0f)
-				.ToolTipText(IsSegmentSelectModeEnabled() ? LOCTEXT("SelectFirstSplineSegmentToolTip", "Select first spline segment.") : LOCTEXT("SelectFirstSplinePointToolTip", "Select first spline control point."))
+				.ToolTipText(IsSegmentSelectModeEnabled() ? LOCTEXT("SelectFirstSplineSegmentToolTip", "Select first spline segment.") : LOCTEXT("SelectFirstSplinePointToolTip", "Select first spline point."))
 				.OnClicked(this, &FLandscapeSplineDetails::OnSelectEndLinearSplineElementButtonClicked, ESplineNavigationFlags::DirectionBackward)
 				.IsEnabled(this, &FLandscapeSplineDetails::HasEndLinearSplineElement, ESplineNavigationFlags::DirectionBackward)
 		]
@@ -80,7 +80,7 @@ void FLandscapeSplineDetails::CustomizeDetails(IDetailLayoutBuilder& DetailBuild
 				.ContentPadding(2.f)
 				.VAlign(VAlign_Center)
 				.HAlign(HAlign_Center)
-				.ToolTipText(IsSegmentSelectModeEnabled() ? LOCTEXT("SelectAddPrevSplineSegmentToolTip", "Add previous segment to current selection.") : LOCTEXT("SelectAddPrevSplinePointToolTip", "Add previous control point to current selection."))
+				.ToolTipText(IsSegmentSelectModeEnabled() ? LOCTEXT("SelectAddPrevSplineSegmentToolTip", "Add previous segment to current selection.") : LOCTEXT("SelectAddPrevSplinePointToolTip", "Add previous point to current selection."))
 				.OnClicked(this, &FLandscapeSplineDetails::OnSelectAdjacentLinearSplineElementButtonClicked, (ESplineNavigationFlags::DirectionBackward | ESplineNavigationFlags::AddToSelection))
 				.IsEnabled(this, &FLandscapeSplineDetails::HasAdjacentLinearSplineElement, (ESplineNavigationFlags::DirectionBackward | ESplineNavigationFlags::AddToSelection))
 		]
@@ -93,7 +93,7 @@ void FLandscapeSplineDetails::CustomizeDetails(IDetailLayoutBuilder& DetailBuild
 				.ContentPadding(2.f)
 				.VAlign(VAlign_Center)
 				.HAlign(HAlign_Center)
-				.ToolTipText(IsSegmentSelectModeEnabled() ? LOCTEXT("SelectPrevSplineSegmentToolTip", "Select previous segment.") : LOCTEXT("SelectPrevPointToolTip", "Select previous control point."))
+				.ToolTipText(IsSegmentSelectModeEnabled() ? LOCTEXT("SelectPrevSplineSegmentToolTip", "Select previous segment.") : LOCTEXT("SelectPrevPointToolTip", "Select previous point."))
 				.OnClicked(this, &FLandscapeSplineDetails::OnSelectAdjacentLinearSplineElementButtonClicked, ESplineNavigationFlags::DirectionBackward)
 				.IsEnabled(this, &FLandscapeSplineDetails::HasAdjacentLinearSplineElement, ESplineNavigationFlags::DirectionBackward)
 
@@ -108,7 +108,7 @@ void FLandscapeSplineDetails::CustomizeDetails(IDetailLayoutBuilder& DetailBuild
 				.ContentPadding(2.f)
 				.VAlign(VAlign_Center)
 				.HAlign(HAlign_Center)
-				.ToolTipText(IsSegmentSelectModeEnabled() ? LOCTEXT("SelectAllSplineSegmentToolTip", "Select all segments.") : LOCTEXT("SelectAllSplinePointsToolTip", "Select all control points."))
+				.ToolTipText(IsSegmentSelectModeEnabled() ? LOCTEXT("SelectAllSplineSegmentToolTip", "Select all segments.") : LOCTEXT("SelectAllSplinePointsToolTip", "Select all points."))
 				.OnClicked(this, &FLandscapeSplineDetails::OnSelectAllConnectedSplineElementsButtonClicked)
 		]
 		+ SHorizontalBox::Slot()
@@ -121,7 +121,7 @@ void FLandscapeSplineDetails::CustomizeDetails(IDetailLayoutBuilder& DetailBuild
 				.ContentPadding(2.f)
 				.VAlign(VAlign_Center)
 				.HAlign(HAlign_Center)
-				.ToolTipText(IsSegmentSelectModeEnabled() ? LOCTEXT("SelectNextSplineSegmentToolTip", "Select next segment.") : LOCTEXT("SelectNextSplinePointToolTip", "Select next control point."))
+				.ToolTipText(IsSegmentSelectModeEnabled() ? LOCTEXT("SelectNextSplineSegmentToolTip", "Select next segment.") : LOCTEXT("SelectNextSplinePointToolTip", "Select next point."))
 				.OnClicked(this, &FLandscapeSplineDetails::OnSelectAdjacentLinearSplineElementButtonClicked, ESplineNavigationFlags::DirectionForward)
 				.IsEnabled(this, &FLandscapeSplineDetails::HasAdjacentLinearSplineElement, ESplineNavigationFlags::DirectionForward)
 		]
@@ -135,7 +135,7 @@ void FLandscapeSplineDetails::CustomizeDetails(IDetailLayoutBuilder& DetailBuild
 				.ContentPadding(2.f)
 				.VAlign(VAlign_Center)
 				.HAlign(HAlign_Center)
-				.ToolTipText(IsSegmentSelectModeEnabled() ? LOCTEXT("SelectAddNextSplineSegmentToolTip", "Add next segment to current selection.") : LOCTEXT("SelectAddNextSplinePointToolTip", "Add next control point to current selection."))
+				.ToolTipText(IsSegmentSelectModeEnabled() ? LOCTEXT("SelectAddNextSplineSegmentToolTip", "Add next segment to current selection.") : LOCTEXT("SelectAddNextSplinePointToolTip", "Add next point to current selection."))
 				.OnClicked(this, &FLandscapeSplineDetails::OnSelectAdjacentLinearSplineElementButtonClicked, (ESplineNavigationFlags::DirectionForward | ESplineNavigationFlags::AddToSelection))
 				.IsEnabled(this, &FLandscapeSplineDetails::HasAdjacentLinearSplineElement, (ESplineNavigationFlags::DirectionForward | ESplineNavigationFlags::AddToSelection))
 		]
@@ -149,7 +149,7 @@ void FLandscapeSplineDetails::CustomizeDetails(IDetailLayoutBuilder& DetailBuild
 				.ContentPadding(2.0f)
 				.VAlign(VAlign_Center)
 				.HAlign(HAlign_Center)
-				.ToolTipText(IsSegmentSelectModeEnabled() ? LOCTEXT("SelectLastSplineSegmentToolTip", "Select last spline segment.") : LOCTEXT("SelectLastSplinePointToolTip", "Select last spline control point."))
+				.ToolTipText(IsSegmentSelectModeEnabled() ? LOCTEXT("SelectLastSplineSegmentToolTip", "Select last spline segment.") : LOCTEXT("SelectLastSplinePointToolTip", "Select last spline point."))
 				.OnClicked(this, &FLandscapeSplineDetails::OnSelectEndLinearSplineElementButtonClicked, ESplineNavigationFlags::DirectionForward)
 				.IsEnabled(this, &FLandscapeSplineDetails::HasEndLinearSplineElement, ESplineNavigationFlags::DirectionForward)
 		]
@@ -173,7 +173,7 @@ void FLandscapeSplineDetails::CustomizeDetails(IDetailLayoutBuilder& DetailBuild
 				.ContentPadding(ButtonPadding)
 				.VAlign(VAlign_Center)
 				.HAlign(HAlign_Center)
-				.ToolTipText(IsSegmentSelectModeEnabled() ? LOCTEXT("SelectConvertSplineSegmentsToPoints", "Convert selected spline segments into control points.") : LOCTEXT("SelectConvertSplinePointsToSegments", "Convert selected spline control points into segments."))
+				.ToolTipText(IsSegmentSelectModeEnabled() ? LOCTEXT("SelectConvertSplineSegmentsToPoints", "Switch selected segments to points.") : LOCTEXT("SelectConvertSplinePointsToSegments", "Switch selected points to segments."))
 				.OnClicked(this, &FLandscapeSplineDetails::OnToggleSplineSelectionTypeButtonClicked)
 		]
 	];
