Landscape - Blueprint Brush tool custom details view and new tool location
* Disabling tool and display UI warning message in Blueprint Brush Stack for edit layers that don't support BP brushes
* Moved BP Brush from the Paint/Sculpt panel into the Manage panel

diff --git a/Engine/Source/Editor/LandscapeEditor/Private/LandscapeEdMode.cpp b/Engine/Source/Editor/LandscapeEditor/Private/LandscapeEdMode.cpp
--- a/Engine/Source/Editor/LandscapeEditor/Private/LandscapeEdMode.cpp
+++ b/Engine/Source/Editor/LandscapeEditor/Private/LandscapeEdMode.cpp
@@ -218,6 +218,11 @@ void FEdModeLandscape::UpdateToolModes()
 	ToolMode_Manage->ValidTools.Add(TEXT("ResizeLandscape"));
 	ToolMode_Manage->ValidTools.Add(TEXT("Splines"));
 	ToolMode_Manage->ValidTools.Add(TEXT("ImportExport"));
+	if (CanHaveLandscapeLayersContent())
+	{
+		ToolMode_Manage->ValidTools.Add(TEXT("BlueprintBrush"));
+	}
+
 	
 	// Restore
 	FName* PreviousToolName = PreviousTools.Find(ToolMode_Manage->ToolModeName);
@@ -242,11 +247,6 @@ void FEdModeLandscape::UpdateToolModes()
 	ToolMode_Sculpt->ValidTools.Add(TEXT("Retopologize"));
 	ToolMode_Sculpt->ValidTools.Add(TEXT("Visibility"));
 
-	if (CanHaveLandscapeLayersContent())
-	{
-		ToolMode_Sculpt->ValidTools.Add(TEXT("BlueprintBrush"));
-	}
-
 	ToolMode_Sculpt->ValidTools.Add(TEXT("Mask"));
 	ToolMode_Sculpt->ValidTools.Add(TEXT("CopyPaste"));
 	ToolMode_Sculpt->ValidTools.Add(TEXT("Mirror"));
@@ -266,11 +266,6 @@ void FEdModeLandscape::UpdateToolModes()
 	ToolMode_Paint->ValidTools.Add(TEXT("Visibility"));
 	ToolMode_Paint->ValidTools.Add(TEXT("Mask"));
 
-	if (CanHaveLandscapeLayersContent())
-	{
-		ToolMode_Paint->ValidTools.Add(TEXT("BlueprintBrush"));
-	}
-
 	PreviousToolName = PreviousTools.Find(ToolMode_Paint->ToolModeName);
 	ToolMode_Paint->CurrentToolName = PreviousToolName ? *PreviousToolName : TEXT("Paint");
 
diff --git a/Engine/Source/Editor/LandscapeEditor/Private/LandscapeEdModeBPCustomTools.cpp b/Engine/Source/Editor/LandscapeEditor/Private/LandscapeEdModeBPCustomTools.cpp
--- a/Engine/Source/Editor/LandscapeEditor/Private/LandscapeEdModeBPCustomTools.cpp
+++ b/Engine/Source/Editor/LandscapeEditor/Private/LandscapeEdModeBPCustomTools.cpp
@@ -27,7 +27,6 @@
 
 #define LOCTEXT_NAMESPACE "Landscape"
 
-template<class ToolTarget>
 class FLandscapeToolBlueprintBrush : public FLandscapeTool
 {
 protected:
@@ -54,11 +53,6 @@ public:
 	virtual void SetEditRenderType() override { GLandscapeEditRenderMode = ELandscapeEditRenderMode::None | (GLandscapeEditRenderMode & ELandscapeEditRenderMode::BitMaskForMask); }
 	virtual bool SupportsMask() override { return false; }
 
-	virtual ELandscapeToolTargetTypeMask::Type GetSupportedTargetTypes() override
-	{
-		return ELandscapeToolTargetTypeMask::FromType(ToolTarget::TargetType);
-	}
-
 	virtual void EnterTool() override
 	{
 	}
@@ -87,10 +81,10 @@ public:
 			return false;
 		}
 
-		// Only allow placing brushes that would affect our target type
-		if ((DefaultObject->CanAffectHeightmap() && Target.TargetType == ELandscapeToolTargetType::Heightmap) 
-			|| (DefaultObject->CanAffectWeightmap() && Target.TargetType == ELandscapeToolTargetType::Weightmap)
-			|| (DefaultObject->CanAffectVisibilityLayer() && Target.TargetType == ELandscapeToolTargetType::Visibility))
+		// Only allow placing brushes if our landscape can be edited and the current edit layer supports BP brushes
+		const FLandscapeLayer* Layer = EdMode->GetCurrentLayerConst();
+
+		if ((Layer && Layer->EditLayer->SupportsBlueprintBrushes()) && (DefaultObject->CanAffectHeightmap() || DefaultObject->CanAffectWeightmap() || DefaultObject->CanAffectVisibilityLayer()))
 		{
 			ULandscapeInfo* Info = EdMode->CurrentToolTarget.LandscapeInfo.Get();
 			check(Info);
@@ -174,13 +168,9 @@ public:
 //
 void FEdModeLandscape::InitializeTool_BlueprintBrush()
 {
-	auto Sculpt_Tool_BlueprintBrush = MakeUnique<FLandscapeToolBlueprintBrush<FHeightmapToolTarget>>(this);
-	Sculpt_Tool_BlueprintBrush->ValidBrushes.Add("BrushSet_Dummy");
-	LandscapeTools.Add(MoveTemp(Sculpt_Tool_BlueprintBrush));
-
-	auto Paint_Tool_BlueprintBrush = MakeUnique<FLandscapeToolBlueprintBrush<FWeightmapToolTarget>>(this);
-	Paint_Tool_BlueprintBrush->ValidBrushes.Add("BrushSet_Dummy");
-	LandscapeTools.Add(MoveTemp(Paint_Tool_BlueprintBrush));
+	auto Tool_BlueprintBrush = MakeUnique<FLandscapeToolBlueprintBrush>(this);
+	Tool_BlueprintBrush->ValidBrushes.Add("BrushSet_Dummy");
+	LandscapeTools.Add(MoveTemp(Tool_BlueprintBrush));
 }
 
 #undef LOCTEXT_NAMESPACE
diff --git a/Engine/Source/Editor/LandscapeEditor/Private/LandscapeEditorDetailCustomization_Blueprint.cpp b/Engine/Source/Editor/LandscapeEditor/Private/LandscapeEditorDetailCustomization_Blueprint.cpp
new file mode 100644
--- /dev/null
+++ b/Engine/Source/Editor/LandscapeEditor/Private/LandscapeEditorDetailCustomization_Blueprint.cpp
@@ -0,0 +1,44 @@
+// Copyright Epic Games, Inc. All Rights Reserved.
+
+#include "LandscapeEditorDetailCustomization_Blueprint.h"
+
+#include "DetailLayoutBuilder.h"
+#include "DetailCategoryBuilder.h"
+#include "IDetailCustomization.h"
+#include "PropertyHandle.h"
+#include "LandscapeEditorObject.h"
+#include "LandscapeEdMode.h"
+#include "LandscapeEditLayer.h"
+#include "SlateOptMacros.h"
+
+#define LOCTEXT_NAMESPACE "LandscapeEditor.Blueprint"
+
+TSharedRef<IDetailCustomization> FLandscapeEditorDetailCustomization_Blueprint::MakeInstance()
+{
+	return MakeShareable(new FLandscapeEditorDetailCustomization_Blueprint);
+}
+
+BEGIN_SLATE_FUNCTION_BUILD_OPTIMIZATION
+void FLandscapeEditorDetailCustomization_Blueprint::CustomizeDetails(IDetailLayoutBuilder& DetailBuilder)
+{
+	IDetailCategoryBuilder& ToolsCategory = DetailBuilder.EditCategory("Tool Settings");
+
+	if (const FEdModeLandscape* LandscapeEdMode = GetEditorMode(); LandscapeEdMode && LandscapeEdMode->GetLandscape() && LandscapeEdMode->CurrentToolMode)
+	{
+		const FLandscapeLayer* Layer = LandscapeEdMode->GetCurrentLayerConst(); 
+		check(Layer != nullptr);
+		const bool bSupportsBPBrushes = Layer->EditLayer->SupportsBlueprintBrushes();
+
+		const FText EnabledToolTip = LOCTEXT("LandscapeBlueprintBrush_EnabledToolTip", "Selects the blueprint brush to apply to the current edit layer. Click on the landscape to apply it.");
+		const FText DisabledToolTip = FText::Format(LOCTEXT("LandscapeBlueprintBrush_DisabledToolTip", "Cannot add blueprint brush : the type of layer {0} ({1}) doesn't support blueprint brushes."),
+			FText::FromName(Layer->EditLayer->GetName()), Layer->EditLayer.GetClass()->GetDisplayNameText());
+
+		TSharedRef<IPropertyHandle> PropertyHandle_Blueprint = DetailBuilder.GetProperty(GET_MEMBER_NAME_CHECKED(ULandscapeEditorObject, BlueprintBrush));
+		ToolsCategory.AddProperty(PropertyHandle_Blueprint)
+			.IsEnabled(bSupportsBPBrushes)
+			.ToolTip(bSupportsBPBrushes ? EnabledToolTip : DisabledToolTip);
+	}
+}
+END_SLATE_FUNCTION_BUILD_OPTIMIZATION
+
+#undef LOCTEXT_NAMESPACE
\ No newline at end of file
diff --git a/Engine/Source/Editor/LandscapeEditor/Private/LandscapeEditorDetailCustomization_Blueprint.h b/Engine/Source/Editor/LandscapeEditor/Private/LandscapeEditorDetailCustomization_Blueprint.h
new file mode 100644
--- /dev/null
+++ b/Engine/Source/Editor/LandscapeEditor/Private/LandscapeEditorDetailCustomization_Blueprint.h
@@ -0,0 +1,23 @@
+// Copyright Epic Games, Inc. All Rights Reserved.
+
+#pragma once
+
+#include "CoreMinimal.h"
+#include "LandscapeEditorDetailCustomization_Base.h"
+
+class IDetailLayoutBuilder;
+class IDetailCustomization;
+
+/**
+ * Slate widgets customizer for the "Blueprint Brush" Landscape tool
+ */
+class FLandscapeEditorDetailCustomization_Blueprint : public FLandscapeEditorDetailCustomization_Base
+{
+public:
+	/** Makes a new instance of this detail layout class for a specific detail view requesting it */
+	static TSharedRef<IDetailCustomization> MakeInstance();
+
+	/** IDetailCustomization interface */
+	virtual void CustomizeDetails(IDetailLayoutBuilder& DetailBuilder) override;
+};
+
diff --git a/Engine/Source/Editor/LandscapeEditor/Private/LandscapeEditorDetailCustomization_LayersBrushStack.cpp b/Engine/Source/Editor/LandscapeEditor/Private/LandscapeEditorDetailCustomization_LayersBrushStack.cpp
--- a/Engine/Source/Editor/LandscapeEditor/Private/LandscapeEditorDetailCustomization_LayersBrushStack.cpp
+++ b/Engine/Source/Editor/LandscapeEditor/Private/LandscapeEditorDetailCustomization_LayersBrushStack.cpp
@@ -15,6 +15,7 @@
 #include "Widgets/Images/SImage.h"
 #include "Widgets/Input/SButton.h"
 #include "Widgets/Input/SCheckBox.h"
+#include "Widgets/Input/SMultiLineEditableTextBox.h"
 #include "EditorModeManager.h"
 #include "EditorModes.h"
 #include "DetailLayoutBuilder.h"
@@ -23,6 +24,7 @@
 #include "PropertyCustomizationHelpers.h"
 #include "LandscapeEditorDetailCustomization_Layers.h"
 #include "LandscapeEditTypes.h"
+#include "LandscapeEditLayer.h"
 #include "Styling/StyleColors.h"
 
 #include "ScopedTransaction.h"
@@ -82,8 +84,12 @@ void FLandscapeEditorCustomNodeBuilder_LayersBrushStack::GenerateHeaderRowConten
 BEGIN_SLATE_FUNCTION_BUILD_OPTIMIZATION
 void FLandscapeEditorCustomNodeBuilder_LayersBrushStack::GenerateChildContent(IDetailChildrenBuilder& ChildrenBuilder)
 {
-	if (FEdModeLandscape* LandscapeEdMode = GetEditorMode())
+	if (FEdModeLandscape* LandscapeEdMode = GetEditorMode(); LandscapeEdMode && LandscapeEdMode->CurrentToolMode)
 	{
+		const int32 BrushCount = LandscapeEdMode->GetBrushesForCurrentLayer().Num();
+		const FLandscapeLayer* Layer = LandscapeEdMode->GetCurrentLayerConst();
+		check(Layer != nullptr && Layer->EditLayer != nullptr);
+
 		TSharedPtr<SDragAndDropVerticalBox> BrushesList = SNew(SDragAndDropVerticalBox)
 			.OnCanAcceptDropAdvanced(this, &FLandscapeEditorCustomNodeBuilder_LayersBrushStack::HandleCanAcceptDrop)
 			.OnAcceptDrop(this, &FLandscapeEditorCustomNodeBuilder_LayersBrushStack::HandleAcceptDrop)
@@ -92,33 +98,42 @@ void FLandscapeEditorCustomNodeBuilder_LayersBrushStack::GenerateChildContent(ID
 		BrushesList->SetDropIndicator_Above(*FAppStyle::GetBrush("LandscapeEditor.TargetList.DropZone.Above"));
 		BrushesList->SetDropIndicator_Below(*FAppStyle::GetBrush("LandscapeEditor.TargetList.DropZone.Below"));
 
+		const TSharedPtr<SWidget> BrushesListWidget = SNew(SVerticalBox)
+			+ SVerticalBox::Slot()
+			.AutoHeight()
+			.VAlign(VAlign_Center)
+			.Padding(0, 2)
+			[
+				BrushesList.ToSharedRef()
+			];
+
+		// Display a disabled message when the current edit layer does not support BP brushes
+		const FText DisabledToolTip = FText::Format(LOCTEXT("LandscapeEditLayerBlueprintBrushes_Disabled", "The type of layer {0} ({1}) doesn't support blueprint brushes."),
+			FText::FromName(Layer->EditLayer->GetName()), Layer->EditLayer->GetClass()->GetDisplayNameText());
+
+		const TSharedPtr<SWidget> ErrorMessage = SNew(SMultiLineEditableTextBox)
+			.IsReadOnly(true)
+			.Font(IDetailLayoutBuilder::GetDetailFontBold())
+			.BackgroundColor(FAppStyle::GetColor("ErrorReporting.WarningBackgroundColor"))
+			.Text(DisabledToolTip)
+			.AutoWrapText(true);
+
 		ChildrenBuilder.AddCustomRow(LOCTEXT("LandscapeEditLayerBlueprintBrushes", "Edit Layer Blueprint Brushes"))
-			.Visibility(EVisibility::Visible)
+			.WholeRowContent()
 			[
-				SNew(SVerticalBox)
-				+ SVerticalBox::Slot()
-				.AutoHeight()
-				.VAlign(VAlign_Center)
-				.Padding(0, 2)
-				[
-					BrushesList.ToSharedRef()
-				]
+				Layer->EditLayer->SupportsBlueprintBrushes() ? BrushesListWidget.ToSharedRef() : ErrorMessage.ToSharedRef()
 			];
 
-		if (LandscapeEdMode->CurrentToolMode)
+		for (int32 i = 0; i < BrushCount; ++i)
 		{
-			int32 BrushCount = LandscapeEdMode->GetBrushesForCurrentLayer().Num();
-			for (int32 i = 0; i < BrushCount; ++i)
+			TSharedPtr<SWidget> GeneratedRowWidget = GenerateRow(i);
+			if (GeneratedRowWidget.IsValid())
 			{
-				TSharedPtr<SWidget> GeneratedRowWidget = GenerateRow(i);
-				if (GeneratedRowWidget.IsValid())
-				{
-					BrushesList->AddSlot()
-						.AutoHeight()
-						[
-							GeneratedRowWidget.ToSharedRef()
-						];
-				}
+				BrushesList->AddSlot()
+					.AutoHeight()
+					[
+						GeneratedRowWidget.ToSharedRef()
+					];
 			}
 		}
 	}
diff --git a/Engine/Source/Editor/LandscapeEditor/Private/LandscapeEditorDetails.cpp b/Engine/Source/Editor/LandscapeEditor/Private/LandscapeEditorDetails.cpp
--- a/Engine/Source/Editor/LandscapeEditor/Private/LandscapeEditorDetails.cpp
+++ b/Engine/Source/Editor/LandscapeEditor/Private/LandscapeEditorDetails.cpp
@@ -10,6 +10,7 @@
 #include "Widgets/Input/SSpinBox.h"
 #include "Widgets/Input/SSegmentedControl.h"
 #include "LandscapeEdMode.h"
+#include "LandscapeEditorDetailCustomization_Blueprint.h"
 #include "LandscapeEditorDetailCustomization_NewLandscape.h"
 #include "LandscapeEditorDetailCustomization_ResizeLandscape.h"
 #include "LandscapeEditorDetailCustomization_CopyPaste.h"
@@ -196,6 +197,8 @@ void FLandscapeEditorDetails::CustomizeDetails(IDetailLayoutBuilder& DetailBuild
 	Customization_CopyPaste->CustomizeDetails(DetailBuilder);
 	Customization_MiscTools = MakeShareable(new FLandscapeEditorDetailCustomization_MiscTools);
 	Customization_MiscTools->CustomizeDetails(DetailBuilder);
+	Customization_Blueprint = MakeShareable(new FLandscapeEditorDetailCustomization_Blueprint);
+	Customization_Blueprint->CustomizeDetails(DetailBuilder);
 
 	// Brushes:
 	Customization_AlphaBrush = MakeShareable(new FLandscapeEditorDetailCustomization_AlphaBrush);
diff --git a/Engine/Source/Editor/LandscapeEditor/Private/LandscapeEditorDetails.h b/Engine/Source/Editor/LandscapeEditor/Private/LandscapeEditorDetails.h
--- a/Engine/Source/Editor/LandscapeEditor/Private/LandscapeEditorDetails.h
+++ b/Engine/Source/Editor/LandscapeEditor/Private/LandscapeEditorDetails.h
@@ -13,6 +13,7 @@
 #include "LandscapeEditorDetailCustomization_LayersBrushStack.h"
 
 class FLandscapeEditorDetailCustomization_AlphaBrush;
+class FLandscapeEditorDetailCustomization_Blueprint;
 class FLandscapeEditorDetailCustomization_CopyPaste;
 class FLandscapeEditorDetailCustomization_MiscTools;
 class FLandscapeEditorDetailCustomization_NewLandscape;
@@ -80,6 +81,7 @@ protected:
 	TSharedPtr<FLandscapeEditorDetailCustomization_CopyPaste> Customization_CopyPaste;
 	TSharedPtr<FLandscapeEditorDetailCustomization_MiscTools> Customization_MiscTools;
 	TSharedPtr<FLandscapeEditorDetailCustomization_AlphaBrush> Customization_AlphaBrush;
+	TSharedPtr<FLandscapeEditorDetailCustomization_Blueprint> Customization_Blueprint;
 	TSharedPtr<FLandscapeEditorDetailCustomization_TargetLayers> Customization_TargetLayers;
 	TSharedPtr<FLandscapeEditorDetailCustomization_Layers> Customization_Layers;
 	TSharedPtr<FLandscapeEditorDetailCustomization_LayersBrushStack> Customization_LayersBrushStack;
diff --git a/Engine/Source/Editor/LandscapeEditor/Private/SLandscapeEditor.cpp b/Engine/Source/Editor/LandscapeEditor/Private/SLandscapeEditor.cpp
--- a/Engine/Source/Editor/LandscapeEditor/Private/SLandscapeEditor.cpp
+++ b/Engine/Source/Editor/LandscapeEditor/Private/SLandscapeEditor.cpp
@@ -236,6 +236,10 @@ void FLandscapeToolKit::BuildToolPalette(FName PaletteName, class FToolBarBuilde
 		ToolBarBuilder.AddToolBarButton(Commands.MoveToLevelTool);
 		ToolBarBuilder.AddToolBarButton(Commands.ResizeLandscape);
 		ToolBarBuilder.AddToolBarButton(Commands.SplineTool);
+		if (Settings->AreBlueprintToolsAllowed())
+		{
+			ToolBarBuilder.AddToolBarButton(Commands.BlueprintBrushTool);
+		}
 	}
 
 	else if (PaletteName == LandscapeEditorNames::Sculpt)
@@ -255,11 +259,6 @@ void FLandscapeToolKit::BuildToolPalette(FName PaletteName, class FToolBarBuilde
 		}
 		ToolBarBuilder.AddToolBarButton(Commands.VisibilityTool);
 
-		if (Settings->AreBlueprintToolsAllowed())
-		{
-			ToolBarBuilder.AddToolBarButton(Commands.BlueprintBrushTool);
-		}
-
 		ToolBarBuilder.AddToolBarButton(Commands.MirrorTool);
 		ToolBarBuilder.AddToolBarButton(Commands.RegionCopyPasteTool);
 
@@ -272,12 +271,6 @@ void FLandscapeToolKit::BuildToolPalette(FName PaletteName, class FToolBarBuilde
 		ToolBarBuilder.AddToolBarButton(Commands.SmoothTool);
 		ToolBarBuilder.AddToolBarButton(Commands.FlattenTool);
 		ToolBarBuilder.AddToolBarButton(Commands.NoiseTool);
-
-		if (LandscapeEdMode->CanHaveLandscapeLayersContent() && Settings->AreBlueprintToolsAllowed())
-		{
-			ToolBarBuilder.AddToolBarButton(Commands.BlueprintBrushTool);
-		}
-
 		ToolBarBuilder.AddToolBarButton(Commands.RegionSelectTool);
 	}
 
