Fixed selection getting lost on landscape undo

diff --git a/Engine/Source/Editor/LandscapeEditor/Private/LandscapeEdModeSplineTools.cpp b/Engine/Source/Editor/LandscapeEditor/Private/LandscapeEdModeSplineTools.cpp
--- a/Engine/Source/Editor/LandscapeEditor/Private/LandscapeEdModeSplineTools.cpp
+++ b/Engine/Source/Editor/LandscapeEditor/Private/LandscapeEdModeSplineTools.cpp
@@ -32,6 +32,7 @@
 #include "LandscapeSplineSegment.h"
 #include "LandscapeSplineControlPoint.h"
 #include "ControlPointMeshComponent.h"
+#include "Algo/AnyOf.h"
 #include "Algo/Copy.h"
 #include "HAL/PlatformApplicationMisc.h"
 #include "UnrealExporter.h"
@@ -86,6 +87,11 @@ public:
 		Collector.AddReferencedObject(LandscapeInfo);
 		Collector.AddReferencedObjects(SelectedSplineControlPoints);
 		Collector.AddReferencedObjects(SelectedSplineSegments);
+		// TODO [jared.ritchie] : this needs to be done, but then also they need to be cleared when exiting the tool. And we need to make sure they're kept in sync with SelectedSplineControlPoints/SelectedSplineSegments when these change ...
+		//Collector.AddReferencedObject(LastSelectedControlPoint);
+		//Collector.AddReferencedObject(LastSelectedSegment);
+		//Collector.AddReferencedObjects(LinearControlPoints);
+		//Collector.AddReferencedObjects(LinearSegments);
 		Collector.AddReferencedObject(DraggingTangent_Segment);
 	}
 
@@ -2079,6 +2085,8 @@ public:
 
 	void FixSelection()
 	{
+		TSet<ULandscapeSplineControlPoint*> SelectedSplineControlPointsBefore(ObjectPtrDecay(SelectedSplineControlPoints));
+		TSet<ULandscapeSplineSegment*> SelectedSplineSegmentsBefore(ObjectPtrDecay(SelectedSplineSegments));
 		SelectedSplineControlPoints.Empty();
 		SelectedSplineSegments.Empty();
 
@@ -2095,6 +2103,13 @@ public:
 					}
 				});
 			}
+
+			// If the selection has changed, update the properties panel
+			if (Algo::AnyOf(SelectedSplineControlPoints, [&SelectedSplineControlPointsBefore](TObjectPtr<ULandscapeSplineControlPoint> InControlPoint) { return !SelectedSplineControlPointsBefore.Contains(InControlPoint); })
+				|| Algo::AnyOf(SelectedSplineSegments, [&SelectedSplineSegmentsBefore](TObjectPtr<ULandscapeSplineSegment> InSegment) { return !SelectedSplineSegmentsBefore.Contains(InSegment); }))
+			{
+				UpdatePropertiesWindows();
+			}
 		}
 		else
 		{
@@ -2122,7 +2137,6 @@ public:
 	void OnUndo()
 	{
 		FixSelection();
-		UpdatePropertiesWindows();
 	}
 
 	virtual void EnterTool() override
@@ -2585,12 +2599,12 @@ protected:
 	TSet<TObjectPtr<ULandscapeSplineSegment>> SelectedSplineSegments;
 
 	// Last Selected
-	ULandscapeSplineControlPoint* LastSelectedControlPoint;
-	ULandscapeSplineSegment* LastSelectedSegment;
+	TObjectPtr<ULandscapeSplineControlPoint> LastSelectedControlPoint;
+	TObjectPtr<ULandscapeSplineSegment> LastSelectedSegment;
 
 	// Control Points and Segments stored in sequential order
-	TArray<ULandscapeSplineControlPoint*> LinearControlPoints;
-	TArray<ULandscapeSplineSegment*> LinearSegments;
+	TArray<TObjectPtr<ULandscapeSplineControlPoint>> LinearControlPoints;
+	TArray<TObjectPtr<ULandscapeSplineSegment>> LinearSegments;
 
 	// Persistent control point mode
 	TOptional<UE::Widget::EWidgetMode> ControlPointWidgetMode;
