Landscape Editor - Inspected Object docking panel UX updates
* panel is hidden by default when entering edit mode
* clear selected object when entering and exiting edit mode
* added vertical scroll bar to details view

diff --git a/Engine/Source/Editor/LandscapeEditor/Private/LandscapeEdMode.cpp b/Engine/Source/Editor/LandscapeEditor/Private/LandscapeEdMode.cpp
--- a/Engine/Source/Editor/LandscapeEditor/Private/LandscapeEdMode.cpp
+++ b/Engine/Source/Editor/LandscapeEditor/Private/LandscapeEdMode.cpp
@@ -223,7 +223,6 @@ void FEdModeLandscape::UpdateToolModes()
 	{
 		ToolMode_Manage->ValidTools.Add(TEXT("BlueprintBrush"));
 	}
-
 	
 	// Restore
 	FName* PreviousToolName = PreviousTools.Find(ToolMode_Manage->ToolModeName);
@@ -641,6 +640,9 @@ void FEdModeLandscape::Enter()
 	{
 		GLandscapeViewMode = CurrentLandscapeViewMode;
 	}
+
+	// Set the inspected object to empty when entering edit mode
+	SetInspectedObjects({});
 }
 
 /** FEdMode: Called when the mode is exited */
@@ -652,6 +654,9 @@ void FEdModeLandscape::Exit()
 	GIsGizmoDragging = false;
 	GIsAdjustingBrush  = false;
 
+	// Set the inspected object to empty when exiting edit mode
+	SetInspectedObjects({});
+
 	UWorld* World = GetWorld();
 
 	if (World != nullptr)
diff --git a/Engine/Source/Editor/LandscapeEditor/Private/SLandscapeEditor.cpp b/Engine/Source/Editor/LandscapeEditor/Private/SLandscapeEditor.cpp
--- a/Engine/Source/Editor/LandscapeEditor/Private/SLandscapeEditor.cpp
+++ b/Engine/Source/Editor/LandscapeEditor/Private/SLandscapeEditor.cpp
@@ -534,14 +534,10 @@ void FLandscapeToolKit::RequestModeUITabs()
 void FLandscapeToolKit::InvokeUI()
 {
 	FModeToolkit::InvokeUI();
-
-	RefreshInspectedObjectsDetailPanel();
 }
 
 TSharedRef<SDockTab> FLandscapeToolKit::CreateInspectedObjectsDetailsViewTab(const FSpawnTabArgs& Args)
 {
-	FMargin MorePadding = FMargin(10.0f, 2.0f);
-
 	TSharedPtr<SDockTab> CreatedTab = SNew(SDockTab)
 		.Label(LOCTEXT("LandscapeDetailsTab", "Landscape Details"))
 		.ToolTipText(LOCTEXT("LandscapeDetailsTabToolTip", "Landscape Details"))
@@ -813,6 +809,7 @@ void SLandscapeEditorInspectedDetailsView::Construct(const FArguments& InArgs)
 	FDetailsViewArgs DetailsViewArgs;
 	DetailsViewArgs.NameAreaSettings = FDetailsViewArgs::ObjectsUseNameArea;
 	DetailsViewArgs.bAllowSearch = false;
+	DetailsViewArgs.bShowScrollBar = true;
 	DetailsViewArgs.bCustomNameAreaLocation = true;
 	// Minimum size to allow the ResetToDefault button to be hit testable.
 	DetailsViewArgs.RightColumnMinWidth = 35;
@@ -829,11 +826,15 @@ void SLandscapeEditorInspectedDetailsView::Construct(const FArguments& InArgs)
 					DetailsView->GetNameAreaWidget().ToSharedRef()
 				]
 				+ SVerticalBox::Slot()
-				.AutoHeight()
+				.FillHeight(1.0f)
 				[
 					DetailsView.ToSharedRef()
 				]
 		];
+
+	// Ensure that the details view calls SetObjects on construction
+	// If its not called the panel will not show the "Select an object" message
+	RefreshDetailPanel();
 }
 
 FEdModeLandscape* SLandscapeEditorInspectedDetailsView::GetEditorMode() const
