Addressing crash when running Editor tests in another project due to assets not being found.

#rnx

#virtualized

diff --git a/Engine/Plugins/Tests/EditorTests/Source/EditorTests/Private/UnrealEd/EditorTestActorSelection.cpp b/Engine/Plugins/Tests/EditorTests/Source/EditorTests/Private/UnrealEd/EditorTestActorSelection.cpp
deleted file mode 100644
diff --git a/Engine/Plugins/Tests/EditorTests/Source/EditorTests/Private/UnrealEd/EditorTestHierarchicalISMSelection.cpp b/Engine/Plugins/Tests/EditorTests/Source/EditorTests/Private/UnrealEd/EditorTestHierarchicalISMSelection.cpp
deleted file mode 100644
diff --git a/Engine/Source/Editor/LandscapeEditor/Private/Tests/LandscapeAutomationTests.cpp b/Engine/Source/Editor/LandscapeEditor/Private/Tests/LandscapeAutomationTests.cpp
--- a/Engine/Source/Editor/LandscapeEditor/Private/Tests/LandscapeAutomationTests.cpp
+++ b/Engine/Source/Editor/LandscapeEditor/Private/Tests/LandscapeAutomationTests.cpp
@@ -174,7 +174,7 @@ bool FBeginModifyLandscapeCommand::Update()
 }
 
 /**
-*  Latent command stop using the sculpting tool
+*  Latent command to stop using the sculpting tool
 */
 DEFINE_LATENT_AUTOMATION_COMMAND(FEndModifyLandscapeCommand);
 bool FEndModifyLandscapeCommand::Update()
@@ -185,6 +185,7 @@ bool FEndModifyLandscapeCommand::Update()
 	//End using the sculpting tool
 	FLevelEditorViewportClient* SelectedViewport = LandscapeTestUtils::FindSelectedViewport();
 	LandscapeEdMode->CurrentTool->EndTool(SelectedViewport);
+
 	return true;
 }
 
@@ -214,6 +215,17 @@ bool FCheckHeight::Update()
 	return true;
 }
 
+/**
+*  Latent command resetting to Editor Mode defaults
+*/
+DEFINE_LATENT_AUTOMATION_COMMAND(FResetToDefaultModeCommand);
+bool FResetToDefaultModeCommand::Update()
+{
+	// Reset to the default editing mode
+	GLevelEditorModeTools().ActivateDefaultMode();
+	return true;
+}
+
 /**
 * Landscape creation / edit test
 */
@@ -230,6 +242,7 @@ bool FLandscapeEditorTest::RunTest(const FString& Parameters)
 	ADD_LATENT_AUTOMATION_COMMAND(FWaitLatentCommand(1.0f));
 	ADD_LATENT_AUTOMATION_COMMAND(FBeginModifyLandscapeCommand());
 	ADD_LATENT_AUTOMATION_COMMAND(FEndModifyLandscapeCommand());
+	ADD_LATENT_AUTOMATION_COMMAND(FResetToDefaultModeCommand());
 
 	return true;
 }
@@ -256,6 +269,9 @@ bool FLandscapeEditorImportTest::RunTest(const FString& Parameters)
 
 	ADD_LATENT_AUTOMATION_COMMAND(FCheckHeight());
 
+	// Importing the landscape switches to the landscape mode
+	ADD_LATENT_AUTOMATION_COMMAND(FResetToDefaultModeCommand());
+
 	return true;
 }
 
@@ -284,6 +300,8 @@ bool FLandscapeEditorCreateNewLayer::RunTest(const FString& Parameters)
 	TestEqual("The landscape should now have a new edit layer", LandscapeEditMode->GetLayerCount(), 2);
 	TestNotNull("There should be a new layer with the specified name", CurrentLandscape->GetLayerConst(LayerName));
 
+	ADD_LATENT_AUTOMATION_COMMAND(FResetToDefaultModeCommand());
+
 	return true;
 }
 
