Landscape Editor - Blueprint brush add validation check for landscapes that disable edit layers

diff --git a/Engine/Source/Editor/LandscapeEditor/Private/LandscapeEditorDetailCustomization_Blueprint.cpp b/Engine/Source/Editor/LandscapeEditor/Private/LandscapeEditorDetailCustomization_Blueprint.cpp
--- a/Engine/Source/Editor/LandscapeEditor/Private/LandscapeEditorDetailCustomization_Blueprint.cpp
+++ b/Engine/Source/Editor/LandscapeEditor/Private/LandscapeEditorDetailCustomization_Blueprint.cpp
@@ -25,18 +25,23 @@ void FLandscapeEditorDetailCustomization_Blueprint::CustomizeDetails(IDetailLayo
 
 	if (const FEdModeLandscape* LandscapeEdMode = GetEditorMode(); LandscapeEdMode && LandscapeEdMode->GetLandscape() && LandscapeEdMode->CurrentToolMode)
 	{
-		const FLandscapeLayer* Layer = LandscapeEdMode->GetCurrentLayerConst(); 
-		check(Layer != nullptr);
-		const bool bSupportsBPBrushes = Layer->EditLayer->SupportsBlueprintBrushes();
-
-		const FText EnabledToolTip = LOCTEXT("LandscapeBlueprintBrush_EnabledToolTip", "Selects the blueprint brush to apply to the current edit layer. Click on the landscape to apply it.");
-		const FText DisabledToolTip = FText::Format(LOCTEXT("LandscapeBlueprintBrush_DisabledToolTip", "Cannot add blueprint brush : the type of layer {0} ({1}) doesn't support blueprint brushes."),
-			FText::FromName(Layer->EditLayer->GetName()), Layer->EditLayer.GetClass()->GetDisplayNameText());
-
-		TSharedRef<IPropertyHandle> PropertyHandle_Blueprint = DetailBuilder.GetProperty(GET_MEMBER_NAME_CHECKED(ULandscapeEditorObject, BlueprintBrush));
-		ToolsCategory.AddProperty(PropertyHandle_Blueprint)
-			.IsEnabled(bSupportsBPBrushes)
-			.ToolTip(bSupportsBPBrushes ? EnabledToolTip : DisabledToolTip);
+		// Tool is disabled for landscapes with edit layers disabled
+		if (LandscapeEdMode->CanHaveLandscapeLayersContent())
+		{
+			const FLandscapeLayer* Layer = LandscapeEdMode->GetCurrentLayerConst();
+			check(Layer != nullptr);
+
+			const bool bSupportsBPBrushes = Layer->EditLayer->SupportsBlueprintBrushes();
+
+			const FText EnabledToolTip = LOCTEXT("LandscapeBlueprintBrush_EnabledToolTip", "Selects the blueprint brush to apply to the current edit layer. Click on the landscape to apply it.");
+			const FText DisabledToolTip = FText::Format(LOCTEXT("LandscapeBlueprintBrush_DisabledToolTip", "Cannot add blueprint brush : the type of layer {0} ({1}) doesn't support blueprint brushes."),
+				FText::FromName(Layer->EditLayer->GetName()), Layer->EditLayer.GetClass()->GetDisplayNameText());
+
+			TSharedRef<IPropertyHandle> PropertyHandle_Blueprint = DetailBuilder.GetProperty(GET_MEMBER_NAME_CHECKED(ULandscapeEditorObject, BlueprintBrush));
+			ToolsCategory.AddProperty(PropertyHandle_Blueprint)
+				.IsEnabled(bSupportsBPBrushes)
+				.ToolTip(bSupportsBPBrushes ? EnabledToolTip : DisabledToolTip);
+		}
 	}
 }
 END_SLATE_FUNCTION_BUILD_OPTIMIZATION
