Enabled "Apply Without Moving" option on almost all landscape tools (for consistency)

diff --git a/Engine/Source/Editor/LandscapeEditor/Private/LandscapeEdModePaintTools.cpp b/Engine/Source/Editor/LandscapeEditor/Private/LandscapeEdModePaintTools.cpp
--- a/Engine/Source/Editor/LandscapeEditor/Private/LandscapeEdModePaintTools.cpp
+++ b/Engine/Source/Editor/LandscapeEditor/Private/LandscapeEdModePaintTools.cpp
@@ -109,9 +109,7 @@ public:
 		, bIsAllowListMode(EdMode->UISettings->PaintingRestriction == ELandscapeLayerPaintingRestriction::UseComponentAllowList &&
 		                   (InViewportClient->Viewport->KeyState(EKeys::Equals) || InViewportClient->Viewport->KeyState(EKeys::Hyphen)))
 		, bAddToAllowList(bIsAllowListMode && InViewportClient->Viewport->KeyState(EKeys::Equals))
-	{
-		bUseContinuousApply = InEdMode->UISettings->bApplyWithoutMovingPaint;
-	}
+	{}
 
 	void Apply(FEditorViewportClient* ViewportClient, FLandscapeBrush* Brush, const ULandscapeEditorObject* UISettings, const TArray<FLandscapeToolInteractorPosition>& InteractorPositions)
 	{
@@ -323,9 +321,7 @@ public:
 	FLandscapeToolStrokeErase(FEdModeLandscape* InEdMode, FEditorViewportClient* InViewportClient, const FLandscapeToolTarget& InTarget)
 		: Super(InEdMode, InViewportClient, InTarget)
 		, FlattenHeight(LandscapeDataAccess::GetTexHeight(0.f))
-	{
-		bUseContinuousApply = InEdMode->UISettings->bApplyWithoutMovingHeight;
-	}
+	{}
 
 	void Apply(FEditorViewportClient* ViewportClient, FLandscapeBrush* Brush, const ULandscapeEditorObject* UISettings, const TArray<FLandscapeToolInteractorPosition>& InteractorPositions)
 	{
@@ -403,9 +399,7 @@ public:
 
 	FLandscapeToolStrokeSculpt(FEdModeLandscape* InEdMode, FEditorViewportClient* InViewportClient, const FLandscapeToolTarget& InTarget)
 		: Super(InEdMode, InViewportClient, InTarget)
-	{
-		bUseContinuousApply = InEdMode->UISettings->bApplyWithoutMovingHeight;
-	}
+	{}
 
 	virtual void SetEditLayer(const FGuid& EditLayerGUID) override
 	{
diff --git a/Engine/Source/Editor/LandscapeEditor/Private/LandscapeEdModeTools.h b/Engine/Source/Editor/LandscapeEditor/Private/LandscapeEdModeTools.h
--- a/Engine/Source/Editor/LandscapeEditor/Private/LandscapeEdModeTools.h
+++ b/Engine/Source/Editor/LandscapeEditor/Private/LandscapeEdModeTools.h
@@ -12,6 +12,7 @@
 #include "LandscapeEdit.h"
 #include "LandscapeComponent.h"
 #include "LandscapeDataAccess.h"
+#include "LandscapeEditorObject.h"
 #include "LandscapeHeightfieldCollisionComponent.h"
 #include "InstancedFoliageActor.h"
 #include "AI/NavigationSystemBase.h"
@@ -1362,6 +1363,7 @@ public:
 		, Target(InTarget)
 		, LandscapeInfo(InTarget.LandscapeInfo.Get())
 	{
+		bUseContinuousApply = (InTarget.TargetType == ELandscapeToolTargetType::Weightmap) ? InEdMode->UISettings->bApplyWithoutMovingPaint : InEdMode->UISettings->bApplyWithoutMovingSculpt;
 	}
 
 	virtual void SetEditLayer(const FGuid& EditLayerGUID)
diff --git a/Engine/Source/Editor/LandscapeEditor/Private/LandscapeEditorObject.cpp b/Engine/Source/Editor/LandscapeEditor/Private/LandscapeEditorObject.cpp
--- a/Engine/Source/Editor/LandscapeEditor/Private/LandscapeEditorObject.cpp
+++ b/Engine/Source/Editor/LandscapeEditor/Private/LandscapeEditorObject.cpp
@@ -109,7 +109,7 @@ ULandscapeEditorObject::ULandscapeEditorObject(const FObjectInitializer& ObjectI
 	, BrushFalloff(0.5f)
 	, PaintBrushFalloff(0.5f)
 	, bUseClayBrush(false)
-	, bApplyWithoutMovingHeight(true)
+	, bApplyWithoutMovingSculpt(true)
 	, bApplyWithoutMovingPaint(false)
 
 	, AlphaBrushScale(0.5f)
@@ -274,7 +274,7 @@ void ULandscapeEditorObject::Load()
 	GConfig->GetFloat(TEXT("LandscapeEdit"), TEXT("RampSideFalloff"), RampSideFalloff, GEditorPerProjectIni);
 
 	GConfig->GetBool(TEXT("LandscapeEdit"), TEXT("bCombinedLayersOperation"), bCombinedLayersOperation, GEditorPerProjectIni);
-	GConfig->GetBool(TEXT("LandscapeEdit"), TEXT("bApplyWithoutMovingHeight"), bApplyWithoutMovingHeight, GEditorPerProjectIni);
+	GConfig->GetBool(TEXT("LandscapeEdit"), TEXT("bApplyWithoutMovingSculpt"), bApplyWithoutMovingSculpt, GEditorPerProjectIni);
 	GConfig->GetBool(TEXT("LandscapeEdit"), TEXT("bApplyWithoutMovingPaint"), bApplyWithoutMovingPaint, GEditorPerProjectIni);
 
 	GConfig->GetInt(TEXT("LandscapeEdit"), TEXT("ErodeThresh"), ErodeThresh, GEditorPerProjectIni);
@@ -431,7 +431,7 @@ void ULandscapeEditorObject::Save()
 	GConfig->SetFloat(TEXT("LandscapeEdit"), TEXT("RampSideFalloff"), RampSideFalloff, GEditorPerProjectIni);
 
 	GConfig->SetBool(TEXT("LandscapeEdit"), TEXT("bCombinedLayersOperation"), bCombinedLayersOperation, GEditorPerProjectIni);
-	GConfig->SetBool(TEXT("LandscapeEdit"), TEXT("bApplyWithoutMovingHeight"), bApplyWithoutMovingHeight, GEditorPerProjectIni);
+	GConfig->SetBool(TEXT("LandscapeEdit"), TEXT("bApplyWithoutMovingSculpt"), bApplyWithoutMovingSculpt, GEditorPerProjectIni);
 	GConfig->SetBool(TEXT("LandscapeEdit"), TEXT("bApplyWithoutMovingPaint"), bApplyWithoutMovingPaint, GEditorPerProjectIni);
 	
 	GConfig->SetInt(TEXT("LandscapeEdit"), TEXT("ErodeThresh"), ErodeThresh, GEditorPerProjectIni);
diff --git a/Engine/Source/Editor/LandscapeEditor/Public/LandscapeEditorObject.h b/Engine/Source/Editor/LandscapeEditor/Public/LandscapeEditorObject.h
--- a/Engine/Source/Editor/LandscapeEditor/Public/LandscapeEditorObject.h
+++ b/Engine/Source/Editor/LandscapeEditor/Public/LandscapeEditorObject.h
@@ -665,11 +665,11 @@ public:
 	bool bUseClayBrush;
 
 	// Apply brush continuously while held stationary, after a small delay
-	UPROPERTY(Category = "Brush Settings", EditAnywhere, NonTransactional, meta=(DisplayName = "Apply without Moving", ShowForTools = "Sculpt,Erase", ShowForBrushes = "BrushSet_Circle,BrushSet_Gizmo,BrushSet_Pattern,BrushSet_Component"))
-	bool bApplyWithoutMovingHeight;
+	UPROPERTY(Category = "Brush Settings", EditAnywhere, NonTransactional, meta=(DisplayName = "Apply Without Moving", ShowForTargetTypes = "Heightmap", ShowForTools = "Sculpt,Erase,Smooth,Flatten,Erosion,HydraErosion,Noise,Mask,CopyPaste", ShowForBrushes = "BrushSet_Circle,BrushSet_Alpha,BrushSet_Gizmo,BrushSet_Pattern,BrushSet_Component"))
+	bool bApplyWithoutMovingSculpt;
 
 	// Apply brush continuously while held stationary, after a small delay
-	UPROPERTY(Category = "Brush Settings", EditAnywhere, NonTransactional, meta=(DisplayName = "Apply without Moving", ShowForTools =" Paint", ShowForBrushes = "BrushSet_Circle,BrushSet_Alpha,BrushSet_Pattern,BrushSet_Component"))
+	UPROPERTY(Category = "Brush Settings", EditAnywhere, NonTransactional, meta=(DisplayName = "Apply Without Moving", ShowForTargetTypes = "Weightmap", ShowForTools = "Paint,Smooth,Flatten,Noise,Mask", ShowForBrushes = "BrushSet_Circle,BrushSet_Alpha,BrushSet_Pattern,BrushSet_Component"))
 	bool bApplyWithoutMovingPaint;   // Internally, store a separate "Apply without Moving" value for paint mode
 
 	// Alpha/Pattern Brush:
