Landscape Editor - Fixed a problem in where landscape editor UISettings weren't refreshing correctly when switching between different landscapes or maps
* Adds OnMapChanged delegate to reset UISettings when changing maps (UWorld changes)
* Ensures edit layer index is always valid by updating index when setting landscape info
* Fixes crash when creating new landscapes

diff --git a/Engine/Source/Editor/LandscapeEditor/Private/LandscapeEdMode.cpp b/Engine/Source/Editor/LandscapeEditor/Private/LandscapeEdMode.cpp
--- a/Engine/Source/Editor/LandscapeEditor/Private/LandscapeEdMode.cpp
+++ b/Engine/Source/Editor/LandscapeEditor/Private/LandscapeEdMode.cpp
@@ -62,6 +62,7 @@
 #include "Interfaces/IMainFrameModule.h"
 #include "LandscapeBlueprintBrushBase.h"
 #include "Engine/TextureRenderTarget2D.h"
+#include "Engine/World.h"
 #include "Settings/EditorExperimentalSettings.h"
 #include "ComponentRecreateRenderStateContext.h"
 #include "VisualLogger/VisualLogger.h"
@@ -118,6 +119,7 @@ FEdModeLandscape::FEdModeLandscape()
 	, bNeedsUpdateLayerUsageInformation(false)
 	, bUpdatingLandscapeInfo(false)
 	, bNeedsUpdateTargetLayerList(false)
+	, bHasMapChanged(false)
 {
 	using namespace UE::Landscape::Editor;
 	
@@ -340,12 +342,22 @@ void FEdModeLandscape::SetLandscapeInfo(ULandscapeInfo* InLandscapeInfo)
 	check(!InLandscapeInfo || InLandscapeInfo->SupportsLandscapeEditing());
 	if (CurrentToolTarget.LandscapeInfo != InLandscapeInfo)
 	{
+		TGuardValue<bool> GuardFlag(bUpdatingLandscapeInfo, true);
+		CurrentToolTarget.LandscapeInfo = InLandscapeInfo;
+
+		// There should always be a valid selected layer before setting up tools and UI 
+		// This may happen when switching between landscapes in the same level, opening new maps, etc.
+		if (CanHaveLandscapeLayersContent())
 		{
-			TGuardValue<bool> GuardFlag(bUpdatingLandscapeInfo, true);
-			CurrentToolTarget.LandscapeInfo = InLandscapeInfo;
-			UpdateTargetList();
-			UpdateToolModes();
+			const ULandscapeEditLayerBase* EditLayer = GetEditLayer(GetCurrentLayerIndex());
+			if (GetCurrentLayerIndex() == INDEX_NONE || EditLayer == nullptr)
+			{
+				SetCurrentLayer(0);
+			}
 		}
+		UpdateTargetList();
+		UpdateToolModes();
+
 		RefreshDetailPanel();
 	}
 }
@@ -483,6 +495,7 @@ void FEdModeLandscape::Enter()
 	OnWorldChangeDelegateHandle                 = FEditorSupportDelegates::WorldChange.AddRaw(this, &FEdModeLandscape::HandleLevelsChanged);
 	OnLevelsChangedDelegateHandle				= GetWorld()->OnLevelsChanged().AddRaw(this, &FEdModeLandscape::HandleLevelsChanged);
 	OnMaterialCompilationFinishedDelegateHandle = UMaterial::OnMaterialCompilationFinished().AddRaw(this, &FEdModeLandscape::OnMaterialCompilationFinished);
+	OnMapChangedDelegateHandle = LevelEditorModule.OnMapChanged().AddSP(this, &FEdModeLandscape::HandleOnMapChanged);
 
 	if (CurrentToolTarget.LandscapeInfo.IsValid() && World != nullptr)
 	{
@@ -643,6 +656,9 @@ void FEdModeLandscape::Enter()
 
 	// Set the inspected object to empty when entering edit mode
 	SetInspectedObjects({});
+
+	// Make sure HasMapChanged is reset when entering edit mode
+	bHasMapChanged = false;
 }
 
 /** FEdMode: Called when the mode is exited */
@@ -674,6 +690,7 @@ void FEdModeLandscape::Exit()
 	}
 
 	FLevelEditorModule& LevelEditorModule = FModuleManager::GetModuleChecked<FLevelEditorModule>("LevelEditor");
+	LevelEditorModule.OnMapChanged().Remove(OnMapChangedDelegateHandle);
 	if (TSharedPtr<ILevelEditor> FirstLevelEditor = LevelEditorModule.GetFirstLevelEditor())
 	{
 		FirstLevelEditor->GetEditorModeManager().OnIsEditingDisallowedChanged().Remove(OnIsEditingDisallowedChangedHandle);
@@ -727,8 +744,13 @@ void FEdModeLandscape::Exit()
 	LandscapeTargetList.Empty();
 	TargetLayerStartingIndex = 0;
 
-	// Save UI settings to config file
-	UISettings->Save();
+	// When switching maps we've already saved and reset the UI settings in HandleOnMapChanged()
+	if (!bHasMapChanged)
+	{
+		// Save UI settings to config file
+		UISettings->Save();
+	}
+	bHasMapChanged = false;
 
 	// Save the last used view mode 
 	CurrentLandscapeViewMode = static_cast<ELandscapeViewMode::Type>(GLandscapeViewMode);
@@ -2532,6 +2554,11 @@ void FEdModeLandscape::SetTargetLandscape(const TWeakObjectPtr<ULandscapeInfo>&
 		}
 	}
 
+	// Force UISettings to be cleared to reset transient properties
+	SaveAndResetUISettings();
+	// Load back the saved properties
+	UISettings->Load();
+
 	SetLandscapeInfo(InLandscapeInfo.Get());
 	// force a Leave and Enter the current tool, in case it has something about the current landscape cached
 	SetCurrentTool(CurrentToolIndex);
@@ -2975,6 +3002,13 @@ void FEdModeLandscape::HandleLevelsChanged()
 	}
 }
 
+void FEdModeLandscape::HandleOnMapChanged(UWorld* World, EMapChangeType ChangeType)
+{
+	// Mark the flag so we don't overwrite saved data during Exit()
+	bHasMapChanged = true;
+	SaveAndResetUISettings();
+}
+
 void FEdModeLandscape::OnMaterialCompilationFinished(UMaterialInterface* MaterialInterface)
 {
 	// TODO [jonathan.bard] : we should remove this now that the target layer list is not automatically filled anymore
@@ -4296,6 +4330,18 @@ void FEdModeLandscape::OnLevelActorRemoved(AActor* InActor)
 	}
 }
 
+void FEdModeLandscape::SaveAndResetUISettings()
+{
+	if (UISettings != nullptr)
+	{
+		UISettings->Save();
+		
+		// Recreate object to reset all values
+		UISettings = NewObject<ULandscapeEditorObject>(GetTransientPackage(), TEXT("UISettings"), RF_Transactional);
+		UISettings->SetParent(this);
+	}
+}
+
 bool LandscapeEditorUtils::SetHeightmapData(ALandscapeProxy* Landscape, const TArray<uint16>& Data)
 {
 	FIntRect ComponentsRect = Landscape->GetBoundingRect() + Landscape->LandscapeSectionOffset;
diff --git a/Engine/Source/Editor/LandscapeEditor/Private/LandscapeEdMode.h b/Engine/Source/Editor/LandscapeEditor/Private/LandscapeEdMode.h
--- a/Engine/Source/Editor/LandscapeEditor/Private/LandscapeEdMode.h
+++ b/Engine/Source/Editor/LandscapeEditor/Private/LandscapeEdMode.h
@@ -33,6 +33,8 @@ struct FHeightmapToolTarget;
 struct FViewportActionKeyInput;
 struct FViewportClick;
 template<class ToolTarget> class FLandscapeToolCopyPaste;
+class UWorld;
+enum class EMapChangeType : uint8;
 
 DECLARE_LOG_CATEGORY_EXTERN(LogLandscapeEdMode, Log, All);
 
@@ -622,6 +624,10 @@ public:
 	/** Handle notification that visible levels may have changed and we should update the editable landscapes list */
 	void HandleLevelsChanged();
 
+	/** Handle when map changes and UWorld is updated */
+	void HandleOnMapChanged(UWorld* World, EMapChangeType ChangeType);
+	void SaveAndResetUISettings();
+
 	void OnMaterialCompilationFinished(UMaterialInterface* MaterialInterface);
 
 	void ReimportData(const FLandscapeTargetListInfo& TargetInfo);
@@ -686,6 +692,7 @@ private:
 
 	FDelegateHandle OnWorldChangeDelegateHandle;
 	FDelegateHandle OnLevelsChangedDelegateHandle;
+	FDelegateHandle OnMapChangedDelegateHandle;
 	FDelegateHandle OnMaterialCompilationFinishedDelegateHandle;
 
 	FDelegateHandle OnLevelActorDeletedDelegateHandle;
@@ -704,6 +711,9 @@ private:
 	bool bUpdatingLandscapeInfo;
 	bool bNeedsUpdateTargetLayerList;
 
+	/** When the map is changed, use this flag to make sure Exit() does not overwrite saved UISettings with default data */
+	bool bHasMapChanged;
+
 	/** The landscape mode panel has the ability to display the properties of certain UObjects if needed (to bypass the fact that the standard details 
 	 panel is not able to display anything else than actors and components) */
 	TArray<TWeakObjectPtr<UObject>> InspectedObjects;
diff --git a/Engine/Source/Editor/LandscapeEditor/Private/LandscapeEditorDetailCustomization_NewLandscape.cpp b/Engine/Source/Editor/LandscapeEditor/Private/LandscapeEditorDetailCustomization_NewLandscape.cpp
--- a/Engine/Source/Editor/LandscapeEditor/Private/LandscapeEditorDetailCustomization_NewLandscape.cpp
+++ b/Engine/Source/Editor/LandscapeEditor/Private/LandscapeEditorDetailCustomization_NewLandscape.cpp
@@ -1300,7 +1300,6 @@ FReply FLandscapeEditorDetailCustomization_NewLandscape::OnCreateButtonClicked()
 	LandscapeEdMode->SetCurrentTargetLayer(NAME_None, nullptr);
 	LandscapeEdMode->SetCurrentTool("Select"); // change tool so switching back to the manage mode doesn't give "New Landscape" again
 	LandscapeEdMode->SetCurrentTool("Sculpt"); // change to sculpting mode and tool
-	LandscapeEdMode->SetCurrentLayer(0);
 
 	World->GetSubsystem<ULandscapeSubsystem>()->ChangeGridSize(LandscapeInfo, UISettings->WorldPartitionGridSize);
 
diff --git a/Engine/Source/Editor/LandscapeEditor/Private/LandscapeEditorDetails.cpp b/Engine/Source/Editor/LandscapeEditor/Private/LandscapeEditorDetails.cpp
--- a/Engine/Source/Editor/LandscapeEditor/Private/LandscapeEditorDetails.cpp
+++ b/Engine/Source/Editor/LandscapeEditor/Private/LandscapeEditorDetails.cpp
@@ -197,8 +197,6 @@ void FLandscapeEditorDetails::CustomizeDetails(IDetailLayoutBuilder& DetailBuild
 	Customization_CopyPaste->CustomizeDetails(DetailBuilder);
 	Customization_MiscTools = MakeShareable(new FLandscapeEditorDetailCustomization_MiscTools);
 	Customization_MiscTools->CustomizeDetails(DetailBuilder);
-	Customization_Blueprint = MakeShareable(new FLandscapeEditorDetailCustomization_Blueprint);
-	Customization_Blueprint->CustomizeDetails(DetailBuilder);
 
 	// Brushes:
 	Customization_AlphaBrush = MakeShareable(new FLandscapeEditorDetailCustomization_AlphaBrush);
@@ -207,6 +205,10 @@ void FLandscapeEditorDetails::CustomizeDetails(IDetailLayoutBuilder& DetailBuild
 	// Hide Layers/LayersBrushStack when New Landscape tab is active
 	if (LandscapeEdMode->CanHaveLandscapeLayersContent() && LandscapeEdMode->CurrentTool->GetToolName() != FName("NewLandscape"))
 	{
+		// Blueprint tool
+		Customization_Blueprint = MakeShareable(new FLandscapeEditorDetailCustomization_Blueprint);
+		Customization_Blueprint->CustomizeDetails(DetailBuilder);
+
 		// Layers
 		Customization_Layers = MakeShareable(new FLandscapeEditorDetailCustomization_Layers);
 		Customization_Layers->CustomizeDetails(DetailBuilder);
