Landscape - Ensure DefaultLayerInfoObject is used anytime a new target layer info asset is created

* Refactored all ULandscapeLayerInfo creations to call util CreateTargetLayerInfo instead of NewObject<>
* On landscape creation, empty target layers do not automatically create layer info assets
* Empty target layer UI message added to warn users when a layer is invalid

diff --git a/Engine/Source/Editor/LandscapeEditor/Private/LandscapeEditorDetailCustomization_ImportLayers.cpp b/Engine/Source/Editor/LandscapeEditor/Private/LandscapeEditorDetailCustomization_ImportLayers.cpp
--- a/Engine/Source/Editor/LandscapeEditor/Private/LandscapeEditorDetailCustomization_ImportLayers.cpp
+++ b/Engine/Source/Editor/LandscapeEditor/Private/LandscapeEditorDetailCustomization_ImportLayers.cpp
@@ -313,8 +313,8 @@ void FLandscapeEditorStructCustomization_FLandscapeImportLayer::OnImportLayerCre
 	}
 	
 	// Build default layer object name and package name
-	FName LayerObjectName;
-	FString PackageName = UE::Landscape::GetLayerInfoObjectPackageName(LayerName, LandscapeEdMode->GetTargetLayerAssetPackagePath(), LayerObjectName);
+	FName FileName;
+	FString PackageName = UE::Landscape::GetLayerInfoObjectPackageName(LayerName, LandscapeEdMode->GetTargetLayerAssetPackagePath(), FileName);
 
 	TSharedRef<SDlgPickAssetPath> NewLayerDlg =
 		SNew(SDlgPickAssetPath)
@@ -323,23 +323,16 @@ void FLandscapeEditorStructCustomization_FLandscapeImportLayer::OnImportLayerCre
 
 	if (NewLayerDlg->ShowModal() != EAppReturnType::Cancel)
 	{
-		PackageName = NewLayerDlg->GetFullAssetPath().ToString();
-		LayerObjectName = FName(*NewLayerDlg->GetAssetName().ToString());
+		PackageName = NewLayerDlg->GetAssetPath().ToString();
+		FileName = FName(*NewLayerDlg->GetAssetName().ToString());
 
-		UPackage* Package = CreatePackage(*PackageName);
-		ULandscapeLayerInfoObject* LayerInfo = NewObject<ULandscapeLayerInfoObject>(Package, LayerObjectName, RF_Public | RF_Standalone | RF_Transactional);
-		LayerInfo->LayerName = LayerName;
+		// Creates a new layer info object asset, using the default if available, or a new empty one
+		ULandscapeLayerInfoObject* LayerInfo = UE::Landscape::CreateTargetLayerInfo(LayerName, PackageName, *FileName.ToString());
 		LayerInfo->bNoWeightBlend = bNoWeightBlend;
 
 		const UObject* LayerInfoAsUObject = LayerInfo; // HACK: If SetValue took a reference to a const ptr (T* const &) or a non-reference (T*) then this cast wouldn't be necessary
 		ensure(PropertyHandle_LayerInfo->SetValue(LayerInfoAsUObject) == FPropertyAccess::Success);
 
-		// Notify the asset registry
-		FAssetRegistryModule::AssetCreated(LayerInfo);
-
-		// Mark the package dirty...
-		Package->MarkPackageDirty();
-
 		// Show in the content browser
 		TArray<UObject*> Objects;
 		Objects.Add(LayerInfo);
diff --git a/Engine/Source/Editor/LandscapeEditor/Private/LandscapeEditorDetailCustomization_NewLandscape.cpp b/Engine/Source/Editor/LandscapeEditor/Private/LandscapeEditorDetailCustomization_NewLandscape.cpp
--- a/Engine/Source/Editor/LandscapeEditor/Private/LandscapeEditorDetailCustomization_NewLandscape.cpp
+++ b/Engine/Source/Editor/LandscapeEditor/Private/LandscapeEditorDetailCustomization_NewLandscape.cpp
@@ -1252,28 +1252,13 @@ FReply FLandscapeEditorDetailCustomization_NewLandscape::OnCreateButtonClicked()
 
 	LandscapeInfo->UpdateLayerInfoMap(Landscape);
 
-	// Import doesn't fill in the LayerInfo for layers with no data, do that now
+	// Create a target layer for each import layer. Layers without an asset set will be created, but without layer info
 	const TArray<FLandscapeImportLayer>& ImportLandscapeLayersList = UISettings->NewLandscape_Layers;
-	const ULandscapeSettings* Settings = GetDefault<ULandscapeSettings>();
-	TSoftObjectPtr<ULandscapeLayerInfoObject> DefaultLayerInfoObject = Settings->GetDefaultLayerInfoObject().LoadSynchronous();
-
 	for (int32 i = 0; i < ImportLandscapeLayersList.Num(); i++)
 	{
 		ULandscapeLayerInfoObject* LayerInfo = ImportLandscapeLayersList[i].LayerInfo;
 		FName LayerName = ImportLandscapeLayersList[i].LayerName;
 
-		// If DefaultLayerInfoObject is set and LayerInfo does not exist, we will try to create the new LayerInfo by cloning DefaultLayerInfoObject. Except for VisibilityLayer which doesn't require an asset.
-		if (DefaultLayerInfoObject.IsValid() && (LayerInfo == nullptr) && (LayerName != ALandscapeProxy::VisibilityLayer->LayerName))
-		{
-			LayerInfo = Landscape->CreateLayerInfo(*LayerName.ToString(), DefaultLayerInfoObject.Get());
-
-			if (LayerInfo != nullptr)
-			{
-				LayerInfo->LayerUsageDebugColor = LayerInfo->GenerateLayerUsageDebugColor();
-				LayerInfo->MarkPackageDirty();
-			}
-		}
-
 		if (LayerInfo != nullptr)
 		{
 			if (LandscapeEdMode->NewLandscapePreviewMode == ENewLandscapePreviewMode::ImportLandscape)
@@ -1292,6 +1277,11 @@ FReply FLandscapeEditorDetailCustomization_NewLandscape::OnCreateButtonClicked()
 				LayerSettings.LayerInfoObj = LayerInfo;
 			}
 		}
+		else
+		{
+			// Add a target layer with no layer info asset
+			Landscape->AddTargetLayer(LayerName, FLandscapeTargetLayerSettings());
+		}
 	}
 
 	LandscapeEdMode->UpdateLandscapeList();
diff --git a/Engine/Source/Editor/LandscapeEditor/Private/LandscapeEditorDetailCustomization_TargetLayers.cpp b/Engine/Source/Editor/LandscapeEditor/Private/LandscapeEditorDetailCustomization_TargetLayers.cpp
--- a/Engine/Source/Editor/LandscapeEditor/Private/LandscapeEditorDetailCustomization_TargetLayers.cpp
+++ b/Engine/Source/Editor/LandscapeEditor/Private/LandscapeEditorDetailCustomization_TargetLayers.cpp
@@ -240,6 +240,17 @@ void FLandscapeEditorDetailCustomization_TargetLayers::CustomizeDetails(IDetailL
 				.AutoWrapText(true)
 		];
 
+	TargetsCategory.AddCustomRow(FText())
+		.Visibility(TAttribute<EVisibility>::Create(TAttribute<EVisibility>::FGetter::CreateStatic(&FLandscapeEditorDetailCustomization_TargetLayers::GetTargetLayersInvalidInfoAssetTipVisibility)))
+		[
+			SNew(SMultiLineEditableTextBox)
+				.IsReadOnly(true)
+				.Font(DetailBuilder.GetDetailFontBold())
+				.BackgroundColor(FAppStyle::GetColor("InfoReporting.BackgroundColor"))
+				.Text(LOCTEXT("InvalidTargetLayers_Tip", "There are target layers with no layer info asset assigned. Create a new asset, select an existing, or use the Auto-Fill button above to quickly set assets for all layers"))
+				.AutoWrapText(true)
+		];
+
 	TargetsCategory.AddCustomRow(FText())
 		.Visibility(TAttribute<EVisibility>::Create(TAttribute<EVisibility>::FGetter::CreateStatic(&FLandscapeEditorDetailCustomization_TargetLayers::GetFilteredTargetLayersListInfoTipVisibility)))
 		[
@@ -325,7 +336,41 @@ EVisibility FLandscapeEditorDetailCustomization_TargetLayers::GetPopulateTargetL
 			|| (LandscapeEdMode->CurrentToolTarget.TargetType == ELandscapeToolTargetType::Invalid)) // ELandscapeToolTargetType::Invalid means "weightmap with no valid paint layer" 
 		{
 			ULandscapeInfo* LandscapeInfo = LandscapeEdMode->CurrentToolTarget.LandscapeInfo.Get();
-			return LandscapeInfo->Layers.IsEmpty() ? EVisibility::Visible : EVisibility::Collapsed;
+			// Visibility layer is added by default behind the scenes, tooltip should be shown until there is a valid weightmap layer in the list
+			const bool bIsVisibilityOnlyLayer = LandscapeInfo->Layers.Num() == 1 && LandscapeInfo->Layers[0].LayerInfoObj == ALandscapeProxy::VisibilityLayer;
+
+			return LandscapeInfo->Layers.IsEmpty() || bIsVisibilityOnlyLayer ? EVisibility::Visible : EVisibility::Collapsed;
+		}
+	}
+
+	return EVisibility::Collapsed;
+}
+
+EVisibility FLandscapeEditorDetailCustomization_TargetLayers::GetTargetLayersInvalidInfoAssetTipVisibility()
+{
+	if (FEdModeLandscape* LandscapeEdMode = GetEditorMode(); (LandscapeEdMode != nullptr) && LandscapeEdMode->CurrentToolTarget.LandscapeInfo.IsValid())
+	{
+		if ((LandscapeEdMode->CurrentToolTarget.TargetType == ELandscapeToolTargetType::Weightmap)
+			|| (LandscapeEdMode->CurrentToolTarget.TargetType == ELandscapeToolTargetType::Invalid)) // ELandscapeToolTargetType::Invalid means "weightmap with no valid paint layer" 
+		{
+			ULandscapeInfo* LandscapeInfo = LandscapeEdMode->CurrentToolTarget.LandscapeInfo.Get();
+			// Visibility layer is added by default behind the scenes, tooltip should be shown until there is a valid weightmap layer in the list
+			const bool bIsVisibilityOnlyLayer = LandscapeInfo->Layers.Num() == 1 && LandscapeInfo->Layers[0].LayerInfoObj == ALandscapeProxy::VisibilityLayer;
+
+			// If we have no layers we cannot have missing layer info assets
+			if (LandscapeInfo->Layers.IsEmpty() || bIsVisibilityOnlyLayer)
+			{
+				return EVisibility::Collapsed;
+			}
+
+			// Show the message if any layer is missing an asset
+			for (FLandscapeInfoLayerSettings& Layer : LandscapeInfo->Layers)
+			{
+				if (Layer.LayerInfoObj == nullptr)
+				{
+					return EVisibility::Visible;
+				}
+			}
 		}
 	}
 
@@ -678,10 +723,10 @@ FReply FLandscapeEditorCustomNodeBuilder_TargetLayers::HandleAutoFillTargetLayer
 			}
 			else if (bCreateNewTargetLayers)
 			{
-				FName LayerObjectName;
-				const FString PackageName = UE::Landscape::GetLayerInfoObjectPackageName(TargetInfo->LayerName, TargetLayerAssetFilePath, LayerObjectName);
+				FName FileName;
+				const FString PackageName = UE::Landscape::GetLayerInfoObjectPackageName(TargetInfo->LayerName, TargetLayerAssetFilePath, FileName);
 
-				CreateTargetLayerInfoAsset(TargetInfo, PackageName, LayerObjectName, /* bNoWeightBlend */ true);
+				CreateTargetLayerInfoAsset(TargetInfo, TargetLayerAssetFilePath, FileName.ToString(), /* bNoWeightBlend */ true);
 			}
 		}
 	}
@@ -1264,7 +1309,6 @@ TSharedPtr<SWidget> FLandscapeEditorCustomNodeBuilder_TargetLayers::GenerateRow(
 											.OnObjectChanged_Static(&FLandscapeEditorCustomNodeBuilder_TargetLayers::OnTargetLayerSetObject, Target)
 											.OnShouldFilterAsset_Static(&FLandscapeEditorCustomNodeBuilder_TargetLayers::ShouldFilterLayerInfo, Target->LayerName)
 											.AllowCreate(false)
-											.AllowClear(false)
 									]
 									+ SHorizontalBox::Slot()
 									.AutoWidth()
@@ -1768,14 +1812,32 @@ bool FLandscapeEditorCustomNodeBuilder_TargetLayers::ShouldFilterLayerInfo(const
 
 void FLandscapeEditorCustomNodeBuilder_TargetLayers::OnTargetLayerSetObject(const FAssetData& AssetData, const TSharedRef<FLandscapeTargetListInfo> Target)
 {
-	// Can't assign null to a layer
+	FScopedTransaction Transaction(LOCTEXT("Undo_AssignTargetLayer", "Assigning Asset to Target Layer"));
+
 	UObject* Object = AssetData.GetAsset();
+
+	// Allow user to clear a layer. UI will display a warning message after to let them know a target layer is missing a layer info asset
 	if (Object == nullptr)
 	{
-		return;
-	}
+		ULandscapeInfo* LandscapeInfo = Target->LandscapeInfo.Get();
+		ALandscape* LandscapeActor = LandscapeInfo->LandscapeActor.Get();
 
-	FScopedTransaction Transaction(LOCTEXT("Undo_UseExisting", "Assigning Layer to Landscape"));
+		int32 Index = LandscapeInfo->GetLayerInfoIndex(Target->LayerName, Target->Owner.Get());
+		if (ensure(Index != INDEX_NONE))
+		{
+			FLandscapeInfoLayerSettings& LayerSettings = LandscapeInfo->Layers[Index];
+			LandscapeInfo->ReplaceLayer(LayerSettings.LayerInfoObj, nullptr);
+		}
+
+		if (FEdModeLandscape* LandscapeEdMode = GetEditorMode())
+		{
+			LandscapeEdMode->CurrentToolTarget.TargetType = Target->TargetType;
+			LandscapeEdMode->SetCurrentTargetLayer(Target->LayerName, nullptr);
+			LandscapeInfo->UpdateLayerInfoMap();
+			LandscapeEdMode->UpdateTargetList();
+			return;
+		}
+	}
 
 	ULandscapeLayerInfoObject* SelectedLayerInfo = const_cast<ULandscapeLayerInfoObject*>(CastChecked<ULandscapeLayerInfoObject>(Object));
 
@@ -1815,8 +1877,7 @@ void FLandscapeEditorCustomNodeBuilder_TargetLayers::OnTargetLayerSetObject(cons
 				}
 			}
 						
-			FEdModeLandscape* LandscapeEdMode = GetEditorMode();
-			if (LandscapeEdMode)
+			if (FEdModeLandscape* LandscapeEdMode = GetEditorMode())
 			{
 				LandscapeEdMode->CurrentToolTarget.TargetType = Target->TargetType;
 				LandscapeEdMode->SetCurrentTargetLayer(Target->LayerName, SelectedLayerInfo);
@@ -1872,8 +1933,8 @@ void FLandscapeEditorCustomNodeBuilder_TargetLayers::OnTargetLayerCreateClicked(
 	if (FEdModeLandscape* LandscapeEdMode = GetEditorMode())
 	{
 		FName LayerName = Target->GetLayerName();
-		FName LayerObjectName;
-		FString PackageName = UE::Landscape::GetLayerInfoObjectPackageName(LayerName, LandscapeEdMode->GetTargetLayerAssetPackagePath(), LayerObjectName);
+		FName FileName;
+		FString PackageName = UE::Landscape::GetLayerInfoObjectPackageName(LayerName, LandscapeEdMode->GetTargetLayerAssetPackagePath(), FileName);
 
 		TSharedRef<SDlgPickAssetPath> NewLayerDlg =
 			SNew(SDlgPickAssetPath)
@@ -1882,40 +1943,23 @@ void FLandscapeEditorCustomNodeBuilder_TargetLayers::OnTargetLayerCreateClicked(
 
 		if (NewLayerDlg->ShowModal() != EAppReturnType::Cancel)
 		{
-			PackageName = NewLayerDlg->GetFullAssetPath().ToString();
-			LayerObjectName = FName(*NewLayerDlg->GetAssetName().ToString());
+			PackageName = NewLayerDlg->GetAssetPath().ToString();
+			FileName = FName(*NewLayerDlg->GetAssetName().ToString());
 
-			CreateTargetLayerInfoAsset(Target, PackageName, LayerObjectName, bNoWeightBlend);
+			CreateTargetLayerInfoAsset(Target, PackageName, FileName.ToString(), bNoWeightBlend);
 		}
 	}
 }
 
-void FLandscapeEditorCustomNodeBuilder_TargetLayers::CreateTargetLayerInfoAsset(const TSharedRef<FLandscapeTargetListInfo> Target, const FString& PackageName, const FName& LayerObjectName, const bool bNoWeightBlend)
+void FLandscapeEditorCustomNodeBuilder_TargetLayers::CreateTargetLayerInfoAsset(const TSharedRef<FLandscapeTargetListInfo> Target, const FString& PackageName, const FString& FileName, const bool bNoWeightBlend)
 {
 	if (FEdModeLandscape* LandscapeEdMode = GetEditorMode())
 	{
 		const FName LayerName = Target->GetLayerName();
 
-		UPackage* Package = CreatePackage(*PackageName);
-
-		// Do not pass RF_Transactional to NewObject, or the asset will mark itself as garbage on Undo (which is not a well-supported path, potentially causing crashes)
-		ULandscapeLayerInfoObject* LayerInfo = NewObject<ULandscapeLayerInfoObject>(Package, LayerObjectName, RF_Public | RF_Standalone);
-		LayerInfo->SetFlags(RF_Transactional);	// we add RF_Transactional after creation, so that future edits _are_ recorded in undo
-		LayerInfo->LayerName = LayerName;
+		ULandscapeLayerInfoObject* LayerInfo = UE::Landscape::CreateTargetLayerInfo(LayerName, PackageName, FileName);
 		LayerInfo->bNoWeightBlend = bNoWeightBlend;
 
-		ULandscapeInfo* LandscapeInfo = Target->LandscapeInfo.Get();
-		LandscapeInfo->Modify();
-		int32 Index = LandscapeInfo->GetLayerInfoIndex(LayerName, Target->Owner.Get());
-		if (Index == INDEX_NONE)
-		{
-			LandscapeInfo->Layers.Add(FLandscapeInfoLayerSettings(LayerInfo, Target->Owner.Get()));
-		}
-		else
-		{
-			LandscapeInfo->Layers[Index].LayerInfoObj = LayerInfo;
-		}
-
 		if (LandscapeEdMode->CurrentToolTarget.LayerName == Target->LayerName
 			&& LandscapeEdMode->CurrentToolTarget.LayerInfo == Target->LayerInfoObj)
 		{
@@ -1925,12 +1969,6 @@ void FLandscapeEditorCustomNodeBuilder_TargetLayers::CreateTargetLayerInfoAsset(
 		Target->LayerInfoObj = LayerInfo;
 		Target->LandscapeInfo->CreateTargetLayerSettingsFor(LayerInfo);
 
-		// Notify the asset registry
-		FAssetRegistryModule::AssetCreated(LayerInfo);
-
-		// Mark the package dirty...
-		Package->MarkPackageDirty();
-
 		// Show in the content browser
 		TArray<UObject*> Objects;
 		Objects.Add(LayerInfo);
diff --git a/Engine/Source/Editor/LandscapeEditor/Private/LandscapeEditorDetailCustomization_TargetLayers.h b/Engine/Source/Editor/LandscapeEditor/Private/LandscapeEditorDetailCustomization_TargetLayers.h
--- a/Engine/Source/Editor/LandscapeEditor/Private/LandscapeEditorDetailCustomization_TargetLayers.h
+++ b/Engine/Source/Editor/LandscapeEditor/Private/LandscapeEditorDetailCustomization_TargetLayers.h
@@ -58,6 +58,7 @@ protected:
 	static EVisibility GetPaintingRestrictionVisibility();
 	static EVisibility GetVisibilityMaskTipVisibility();
 	static EVisibility GetPopulateTargetLayersInfoTipVisibility();
+	static EVisibility GetTargetLayersInvalidInfoAssetTipVisibility();
 	static EVisibility GetFilteredTargetLayersListInfoTipVisibility();
 };
 
@@ -107,7 +108,7 @@ protected:
 	static bool GetTargetLayerCreateEnabled(const TSharedRef<FLandscapeTargetListInfo> Target);
 	static TSharedRef<SWidget> OnGetTargetLayerCreateMenu(const TSharedRef<FLandscapeTargetListInfo> Target);
 	static void OnTargetLayerCreateClicked(const TSharedRef<FLandscapeTargetListInfo> Target, bool bNoWeightBlend);
-	static void CreateTargetLayerInfoAsset(const TSharedRef<FLandscapeTargetListInfo> Target, const FString& PackageName, const FName& LayerObjectName, const bool bNoWeightBlend);
+	static void CreateTargetLayerInfoAsset(const TSharedRef<FLandscapeTargetListInfo> Target, const FString& PackageName, const FString& FileName, const bool bNoWeightBlend);
 	static FReply OnTargetLayerDeleteClicked(const TSharedRef<FLandscapeTargetListInfo> Target);
 	static FSlateColor GetLayerUsageDebugColor(const TSharedRef<FLandscapeTargetListInfo> Target);
 	static EVisibility GetDebugModeLayerUsageVisibility(const TSharedRef<FLandscapeTargetListInfo> Target);
diff --git a/Engine/Source/Editor/LandscapeEditor/Private/LandscapeEditorObject.cpp b/Engine/Source/Editor/LandscapeEditor/Private/LandscapeEditorObject.cpp
--- a/Engine/Source/Editor/LandscapeEditor/Private/LandscapeEditorObject.cpp
+++ b/Engine/Source/Editor/LandscapeEditor/Private/LandscapeEditorObject.cpp
@@ -979,23 +979,12 @@ void ULandscapeEditorObject::AutoFillTargetLayers(const FString& DefaultTargetLa
 		}
 		else if (bCreateAssetForEmptyLayers)
 		{
-			const FName LayerName = NewLayer.LayerName;
-			FName LayerObjectName;
-			const FString PackageName = UE::Landscape::GetLayerInfoObjectPackageName(LayerName, DefaultTargetLayerAssetPath, LayerObjectName);
-
-			UPackage* Package = CreatePackage(*PackageName);
-			ULandscapeLayerInfoObject* LayerInfo = NewObject<ULandscapeLayerInfoObject>(Package, LayerObjectName, RF_Public | RF_Standalone | RF_Transactional);
-			LayerInfo->LayerName = LayerName;
+			// Creates a new layer info object, using the default if available, or a new empty one
+			ULandscapeLayerInfoObject* LayerInfo = UE::Landscape::CreateTargetLayerInfo(NewLayer.LayerName, DefaultTargetLayerAssetPath);
 			LayerInfo->bNoWeightBlend = true;
 
 			NewLayer.LayerInfo = LayerInfo;
 
-			// Notify the asset registry
-			FAssetRegistryModule::AssetCreated(LayerInfo);
-
-			// Mark the package dirty...
-			Package->MarkPackageDirty();
-
 			// Show in the content browser
 			TArray<UObject*> Objects;
 			Objects.Add(LayerInfo);
diff --git a/Engine/Source/Editor/LandscapeEditor/Private/LandscapeEditorUtils.cpp b/Engine/Source/Editor/LandscapeEditor/Private/LandscapeEditorUtils.cpp
--- a/Engine/Source/Editor/LandscapeEditor/Private/LandscapeEditorUtils.cpp
+++ b/Engine/Source/Editor/LandscapeEditor/Private/LandscapeEditorUtils.cpp
@@ -124,11 +124,11 @@ namespace LandscapeEditorUtils
 		for (const FAssetData& LayerInfoAsset : LayerInfoAssets)
 		{
 			const ULandscapeLayerInfoObject* LayerInfo = CastChecked<ULandscapeLayerInfoObject>(LayerInfoAsset.GetAsset());
+			FString CurrentPackagePath = LayerInfoAsset.PackagePath.ToString() + "/";
 
 			// Only include assets in the package or its sub folders
-			if (LayerInfo && LayerInfo->LayerName == LayerName && LayerInfoAsset.PackagePath.ToString().Contains(TargetLayerAssetPackagePath))
+			if (LayerInfo && LayerInfo->LayerName == LayerName && CurrentPackagePath.Contains(TargetLayerAssetPackagePath))
 			{
-				const FString CurrentPackagePath = LayerInfoAsset.PackagePath.ToString();
 				float WorstCase = static_cast<float>(TargetLayerAssetPackagePath.Len() + CurrentPackagePath.Len());
 				WorstCase = FMath::Max(WorstCase, 1.f);
 
diff --git a/Engine/Source/Editor/LandscapeEditor/Private/LandscapeLayerInfoObjectFactory.cpp b/Engine/Source/Editor/LandscapeEditor/Private/LandscapeLayerInfoObjectFactory.cpp
deleted file mode 100644
diff --git a/Engine/Source/Editor/LandscapeEditor/Private/LandscapeLayerInfoObjectFactory.h b/Engine/Source/Editor/LandscapeEditor/Private/LandscapeLayerInfoObjectFactory.h
deleted file mode 100644
diff --git a/Engine/Source/Editor/WorldBrowser/Private/Tiles/WorldTileCollectionModel.cpp b/Engine/Source/Editor/WorldBrowser/Private/Tiles/WorldTileCollectionModel.cpp
--- a/Engine/Source/Editor/WorldBrowser/Private/Tiles/WorldTileCollectionModel.cpp
+++ b/Engine/Source/Editor/WorldBrowser/Private/Tiles/WorldTileCollectionModel.cpp
@@ -1559,13 +1559,9 @@ static ULandscapeLayerInfoObject* GetLandscapeLayerInfoObject(FName LayerName, c
 	ULandscapeLayerInfoObject* LayerInfo = FindObject<ULandscapeLayerInfoObject>(Package, *LayerObjectName);
 	if (LayerInfo == nullptr)
 	{
-		LayerInfo = NewObject<ULandscapeLayerInfoObject>(Package, FName(*LayerObjectName), RF_Public | RF_Standalone | RF_Transactional);
-		LayerInfo->LayerName = LayerName;
-		// Notify the asset registry
-		FAssetRegistryModule::AssetCreated(LayerInfo);
-		// Mark the package dirty...
-		Package->MarkPackageDirty();
-		//
+		// Creates a new layer info object, using the default if available, or a new empty one
+		LayerInfo = UE::Landscape::CreateTargetLayerInfo(LayerName, Path, LayerObjectName);
+
 		TArray<UPackage*> PackagesToSave; 
 		PackagesToSave.Add(Package);
 		FEditorFileUtils::PromptForCheckoutAndSave(PackagesToSave, false, false);
diff --git a/Engine/Source/Runtime/Landscape/Classes/LandscapeProxy.h b/Engine/Source/Runtime/Landscape/Classes/LandscapeProxy.h
--- a/Engine/Source/Runtime/Landscape/Classes/LandscapeProxy.h
+++ b/Engine/Source/Runtime/Landscape/Classes/LandscapeProxy.h
@@ -1336,25 +1336,13 @@ PRAGMA_ENABLE_DEPRECATION_WARNINGS
 	// Query all the LandscapeLayerInfo Objects from the weight layers allocated on this proxy.
 	TMap<FName, ULandscapeLayerInfoObject*> RetrieveTargetLayerInfosFromAllocations() const;
 
-	/**
-	* Creates a new LandscapeLayerInfoObject
-	* @param	LayerName	The name of the created Layer
-	* @param	InFilePath	New asset file path, should be Current Level SharedAssetPath or the LandscapeEditorObject->TargetLayerAssetFilePath->DirectoryPath
-	* @param	InTemplate	If set, this object will be cloned as a basis for the new LandscapeLayerInfoObject.
-	* @return a new LandscapeLayerInfoObject if created successfully.
-	*/
+	UE_DEPRECATED(5.6, "This create function is deprecated. Please use new util UE::Landscape::CreateTargetLayerInfo instead.")
 	LANDSCAPE_API static ULandscapeLayerInfoObject* CreateLayerInfo(const TCHAR* InLayerName, const FString& InFilePath, const ULandscapeLayerInfoObject* InTemplate = nullptr);
 
-	/**
-	* Creates a new LandscapeLayerInfoObject
-	* @param	LayerName	The name of the created Layer
-	* @param	Level		The Level tied to LandscapeLayerInfoObject. The new package's location will be relative to this Level.
-	* @param	InTemplate	If set, this object will be cloned as a basis for the new LandscapeLayerInfoObject.
-	* @return a new LandscapeLayerInfoObject if created successfully.
-	*/
-	UE_DEPRECATED(5.6, "This create function is deprecated. Please use new CreateLayerInfo and pass full asset file path.")
+	UE_DEPRECATED(5.6, "This create function is deprecated. Please use new util UE::Landscape::CreateTargetLayerInfo instead.")
 	LANDSCAPE_API static ULandscapeLayerInfoObject* CreateLayerInfo(const TCHAR* InLayerName, const ULevel* InLevel, const ULandscapeLayerInfoObject* InTemplate = nullptr);
 
+	UE_DEPRECATED(5.6, "This create function is deprecated. Please use new util UE::Landscape::CreateTargetLayerInfo instead.")
 	LANDSCAPE_API ULandscapeLayerInfoObject* CreateLayerInfo(const TCHAR* InLayerName, const ULandscapeLayerInfoObject* InTemplate = nullptr);
 
 	/** Fix up component layers, weightmaps */
diff --git a/Engine/Source/Runtime/Landscape/Private/LandscapeEdit.cpp b/Engine/Source/Runtime/Landscape/Private/LandscapeEdit.cpp
--- a/Engine/Source/Runtime/Landscape/Private/LandscapeEdit.cpp
+++ b/Engine/Source/Runtime/Landscape/Private/LandscapeEdit.cpp
@@ -3199,59 +3199,22 @@ TMap<FName, ULandscapeLayerInfoObject*> ALandscapeProxy::RetrieveTargetLayerInfo
 	return InfoObjects;
 }
 
+// Deprecated
 ULandscapeLayerInfoObject* ALandscapeProxy::CreateLayerInfo(const TCHAR* InLayerName, const FString& InFilePath, const ULandscapeLayerInfoObject* InTemplate)
 {
-	FName LayerObjectName;
-	FString PackageName = UE::Landscape::GetLayerInfoObjectPackageName(InLayerName, InFilePath, LayerObjectName);
-	UPackage* Package = CreatePackage(*PackageName);
-	ULandscapeLayerInfoObject* LayerInfo = nullptr;
-	check(Package != nullptr);
-
-	if (InTemplate != nullptr)
-	{
-		LayerInfo = DuplicateObject<ULandscapeLayerInfoObject>(InTemplate, Package, LayerObjectName);
-	}
-	else
-	{
-		LayerInfo = NewObject<ULandscapeLayerInfoObject>(Package, LayerObjectName, RF_Public | RF_Standalone | RF_Transactional);
-	}
-
-	check(LayerInfo != nullptr);
-	LayerInfo->LayerName = InLayerName;
-
-	FAssetRegistryModule::AssetCreated(LayerInfo);
-	LayerInfo->MarkPackageDirty();
-
-	return LayerInfo;
+	return UE::Landscape::CreateTargetLayerInfo(InLayerName, InFilePath);
 }
 
 // Deprecated
 ULandscapeLayerInfoObject* ALandscapeProxy::CreateLayerInfo(const TCHAR* InLayerName, const ULevel* InLevel, const ULandscapeLayerInfoObject* InTemplate)
 {
-	return CreateLayerInfo(InLayerName, UE::Landscape::GetSharedAssetsPath(InLevel), InTemplate);
+	return UE::Landscape::CreateTargetLayerInfo(InLayerName, UE::Landscape::GetSharedAssetsPath(InLevel));
 }
 
+// Deprecated
 ULandscapeLayerInfoObject* ALandscapeProxy::CreateLayerInfo(const TCHAR* InLayerName, const ULandscapeLayerInfoObject* InTemplate)
 {
-	ULandscapeLayerInfoObject* LayerInfo = ALandscapeProxy::CreateLayerInfo(InLayerName, UE::Landscape::GetSharedAssetsPath(GetLevel()), InTemplate);
-
-	check(LayerInfo);
-
-	ULandscapeInfo* LandscapeInfo = GetLandscapeInfo();
-	if (LandscapeInfo)
-	{
-		int32 Index = LandscapeInfo->GetLayerInfoIndex(InLayerName, this);
-		if (Index == INDEX_NONE)
-		{
-			LandscapeInfo->Layers.Add(FLandscapeInfoLayerSettings(LayerInfo, this));
-		}
-		else
-		{
-			LandscapeInfo->Layers[Index].LayerInfoObj = LayerInfo;
-		}
-	}
-
-	return LayerInfo;
+	return UE::Landscape::CreateTargetLayerInfo(InLayerName, UE::Landscape::GetSharedAssetsPath(GetLevel()));
 }
 
 #define HEIGHTDATA(X,Y) (HeightData.Num() == 0 ? LandscapeDataAccess::GetTexHeight(0.0f) : HeightData[ FMath::Clamp<int32>(Y,0,VertsY) * VertsX + FMath::Clamp<int32>(X,0,VertsX) ])
@@ -5196,7 +5159,11 @@ void ULandscapeInfo::ReplaceLayer(ULandscapeLayerInfoObject* FromLayerInfo, ULan
 		}
 
 		LandscapeActor->RemoveTargetLayer(FromLayerInfo->LayerName);
-		LandscapeActor->AddTargetLayer(ToLayerInfo->LayerName, FLandscapeTargetLayerSettings(ToLayerInfo));
+
+		// When a layer is cleared ToLayerInfo will be invalid
+		const FLandscapeTargetLayerSettings NewSettings = ToLayerInfo != nullptr ? FLandscapeTargetLayerSettings(ToLayerInfo) : FLandscapeTargetLayerSettings();
+		const FName NewName = ToLayerInfo != nullptr ? ToLayerInfo->LayerName : FromLayerInfo->LayerName;
+		LandscapeActor->AddTargetLayer(NewName, NewSettings);
 		UpdateLayerInfoMap();
 
 		GWarn->EndSlowTask();
diff --git a/Engine/Source/Runtime/Landscape/Private/LandscapeEditInterface.cpp b/Engine/Source/Runtime/Landscape/Private/LandscapeEditInterface.cpp
--- a/Engine/Source/Runtime/Landscape/Private/LandscapeEditInterface.cpp
+++ b/Engine/Source/Runtime/Landscape/Private/LandscapeEditInterface.cpp
@@ -2124,7 +2124,7 @@ void ULandscapeComponent::ReplaceLayer(ULandscapeLayerInfoObject* FromLayerInfo,
 
 void ULandscapeComponent::ReplaceLayerInternal(ULandscapeLayerInfoObject* FromLayerInfo, ULandscapeLayerInfoObject* ToLayerInfo, FLandscapeEditDataInterface& LandscapeEdit, const FGuid& InEditLayerGUID)
 {
-	check(FromLayerInfo && ToLayerInfo);
+	check(FromLayerInfo); // ToLayerInfo may be null when clearing a layer
 	if (FromLayerInfo == ToLayerInfo)
 	{
 		return;
diff --git a/Engine/Source/Runtime/Landscape/Private/LandscapeUtils.cpp b/Engine/Source/Runtime/Landscape/Private/LandscapeUtils.cpp
--- a/Engine/Source/Runtime/Landscape/Private/LandscapeUtils.cpp
+++ b/Engine/Source/Runtime/Landscape/Private/LandscapeUtils.cpp
@@ -241,6 +241,52 @@ FString GetLayerInfoObjectPackageName(const ULevel* InLevel, const FName& InLaye
 	return GetLayerInfoObjectPackageName(InLayerName, GetSharedAssetsPath(InLevel), OutLayerObjectName);
 }
 
+ULandscapeLayerInfoObject* CreateTargetLayerInfo(const FName& InLayerName, const FString& InFilePath)
+{
+	// Appends %s_LayerInfo_%d to ensure the new asset has a valid filename
+	FName FileName;
+	const FString PackageName = GetLayerInfoObjectPackageName(InLayerName, InFilePath, FileName);
+
+	return CreateTargetLayerInfo(InLayerName, InFilePath, FileName.ToString());
+}
+
+ULandscapeLayerInfoObject* CreateTargetLayerInfo(const FName& InLayerName, const FString& InFilePath, const FString& InFileName)
+{
+	// Get the default asset from the project settings
+	const ULandscapeSettings* Settings = GetDefault<ULandscapeSettings>();
+	TSoftObjectPtr<ULandscapeLayerInfoObject> DefaultLayerInfoObject = Settings->GetDefaultLayerInfoObject().LoadSynchronous();
+
+	// Ensure the package path has a terminating "/"
+	const FString PackagePath = InFilePath.EndsWith("/") ? InFilePath + InFileName : InFilePath + "/" + InFileName;
+
+	UPackage* Package = CreatePackage(*PackagePath);
+	ULandscapeLayerInfoObject* LayerInfo = nullptr;
+	check(Package != nullptr);
+
+	if (DefaultLayerInfoObject.Get() != nullptr)
+	{
+		LayerInfo = DuplicateObject<ULandscapeLayerInfoObject>(DefaultLayerInfoObject.Get(), Package, *InFileName);
+		LayerInfo->SetFlags(RF_Public | RF_Standalone | RF_Transactional);
+	}
+	else
+	{
+		// Do not pass RF_Transactional to NewObject, or the asset will mark itself as garbage on Undo (which is not a well-supported path, potentially causing crashes)
+		LayerInfo = NewObject<ULandscapeLayerInfoObject>(Package, *InFileName, RF_Public | RF_Standalone);
+		LayerInfo->SetFlags(RF_Transactional);	// we add RF_Transactional after creation, so that future edits _are_ recorded in undo
+	}
+
+	check(LayerInfo != nullptr);
+	LayerInfo->LayerName = InLayerName;
+	LayerInfo->LayerUsageDebugColor = LayerInfo->GenerateLayerUsageDebugColor();
+
+	// Notify the asset registry
+	FAssetRegistryModule::AssetCreated(LayerInfo);
+	Package->MarkPackageDirty();
+	LayerInfo->MarkPackageDirty();
+
+	return LayerInfo;
+}
+
 bool IsVisibilityLayer(const ULandscapeLayerInfoObject* InLayerInfoObject)
 {
 	return (ALandscapeProxy::VisibilityLayer != nullptr) && (ALandscapeProxy::VisibilityLayer == InLayerInfoObject);
diff --git a/Engine/Source/Runtime/Landscape/Public/LandscapeUtils.h b/Engine/Source/Runtime/Landscape/Public/LandscapeUtils.h
--- a/Engine/Source/Runtime/Landscape/Public/LandscapeUtils.h
+++ b/Engine/Source/Runtime/Landscape/Public/LandscapeUtils.h
@@ -312,6 +312,23 @@ LANDSCAPE_API FString GetLayerInfoObjectPackageName(const FName& InLayerName, co
 UE_DEPRECATED(5.6, "This Get function is deprecated. Please use new GetLayerInfo and pass in a full asset package path.")
 LANDSCAPE_API FString GetLayerInfoObjectPackageName(const ULevel* InLevel, const FName& InLayerName, FName& OutLayerObjectName);
 
+/**
+* Creates a new layer info object, using the default template if available, or a new empty one. Sets asset file name as LayerName_LayerInfo_%d
+* @param	InLayerName	The layer name of the created asset info
+* @param	InFilePath	New asset file path, typically is CurrentLevel/SharedAssetPath or the LandscapeEditorObject->TargetLayerAssetFilePath->DirectoryPath
+* @return a new LandscapeLayerInfoObject.
+*/
+LANDSCAPE_API ULandscapeLayerInfoObject* CreateTargetLayerInfo(const FName& InLayerName, const FString& InFilePath);
+
+/** 
+* Creates a new layer info object, using the default template if available, or a new empty one at the path: InFilePath/InFilename
+* @param	InLayerName	The layer name of the created asset info
+* @param	InFilePath	New asset file path, typically is CurrentLevel/SharedAssetPath or the LandscapeEditorObject->TargetLayerAssetFilePath->DirectoryPath
+* @param	InFileName 	The unique file name of the created asset
+* @return a new LandscapeLayerInfoObject.
+*/
+LANDSCAPE_API ULandscapeLayerInfoObject* CreateTargetLayerInfo(const FName& InLayerName, const FString& InFilePath, const FString& InFileName);
+
 /** Returns true if the provided layer info object is the current visibility layer. */
 LANDSCAPE_API bool IsVisibilityLayer(const ULandscapeLayerInfoObject* InLayerInfoObject);
 
