Landscape Editor - Moving CurrentLayerIndex from ULandscapeEditorObject into ALandscape, revert OnMapChanged delegate
* Each landscape tracks its own selected layer index
* Revert OnMapChanged event in landscape editor

diff --git a/Engine/Source/Editor/LandscapeEditor/Private/LandscapeEdMode.cpp b/Engine/Source/Editor/LandscapeEditor/Private/LandscapeEdMode.cpp
--- a/Engine/Source/Editor/LandscapeEditor/Private/LandscapeEdMode.cpp
+++ b/Engine/Source/Editor/LandscapeEditor/Private/LandscapeEdMode.cpp
@@ -62,7 +62,6 @@
 #include "Interfaces/IMainFrameModule.h"
 #include "LandscapeBlueprintBrushBase.h"
 #include "Engine/TextureRenderTarget2D.h"
-#include "Engine/World.h"
 #include "Settings/EditorExperimentalSettings.h"
 #include "ComponentRecreateRenderStateContext.h"
 #include "VisualLogger/VisualLogger.h"
@@ -342,21 +341,13 @@ void FEdModeLandscape::SetLandscapeInfo(ULandscapeInfo* InLandscapeInfo)
 	check(!InLandscapeInfo || InLandscapeInfo->SupportsLandscapeEditing());
 	if (CurrentToolTarget.LandscapeInfo != InLandscapeInfo)
 	{
-		TGuardValue<bool> GuardFlag(bUpdatingLandscapeInfo, true);
-		CurrentToolTarget.LandscapeInfo = InLandscapeInfo;
-
-		// There should always be a valid selected layer before setting up tools and UI 
-		// This may happen when switching between landscapes in the same level, opening new maps, etc.
-		if (CanHaveLandscapeLayersContent())
 		{
-			const ULandscapeEditLayerBase* EditLayer = GetEditLayer(GetCurrentLayerIndex());
-			if (GetCurrentLayerIndex() == INDEX_NONE || EditLayer == nullptr)
-			{
-				SetCurrentLayer(0);
-			}
+			TGuardValue<bool> GuardFlag(bUpdatingLandscapeInfo, true);
+			CurrentToolTarget.LandscapeInfo = InLandscapeInfo;
+
+			UpdateTargetList();
+			UpdateToolModes();
 		}
-		UpdateTargetList();
-		UpdateToolModes();
 
 		RefreshDetailPanel();
 	}
@@ -495,7 +486,6 @@ void FEdModeLandscape::Enter()
 	OnWorldChangeDelegateHandle                 = FEditorSupportDelegates::WorldChange.AddRaw(this, &FEdModeLandscape::HandleLevelsChanged);
 	OnLevelsChangedDelegateHandle				= GetWorld()->OnLevelsChanged().AddRaw(this, &FEdModeLandscape::HandleLevelsChanged);
 	OnMaterialCompilationFinishedDelegateHandle = UMaterial::OnMaterialCompilationFinished().AddRaw(this, &FEdModeLandscape::OnMaterialCompilationFinished);
-	OnMapChangedDelegateHandle = LevelEditorModule.OnMapChanged().AddSP(this, &FEdModeLandscape::HandleOnMapChanged);
 
 	if (CurrentToolTarget.LandscapeInfo.IsValid() && World != nullptr)
 	{
@@ -690,7 +680,6 @@ void FEdModeLandscape::Exit()
 	}
 
 	FLevelEditorModule& LevelEditorModule = FModuleManager::GetModuleChecked<FLevelEditorModule>("LevelEditor");
-	LevelEditorModule.OnMapChanged().Remove(OnMapChangedDelegateHandle);
 	if (TSharedPtr<ILevelEditor> FirstLevelEditor = LevelEditorModule.GetFirstLevelEditor())
 	{
 		FirstLevelEditor->GetEditorModeManager().OnIsEditingDisallowedChanged().Remove(OnIsEditingDisallowedChangedHandle);
@@ -744,13 +733,8 @@ void FEdModeLandscape::Exit()
 	LandscapeTargetList.Empty();
 	TargetLayerStartingIndex = 0;
 
-	// When switching maps we've already saved and reset the UI settings in HandleOnMapChanged()
-	if (!bHasMapChanged)
-	{
-		// Save UI settings to config file
-		UISettings->Save();
-	}
-	bHasMapChanged = false;
+	// Save UI settings to config file
+	UISettings->Save();
 
 	// Save the last used view mode 
 	CurrentLandscapeViewMode = static_cast<ELandscapeViewMode::Type>(GLandscapeViewMode);
@@ -847,7 +831,7 @@ void FEdModeLandscape::Tick(FEditorViewportClient* ViewportClient, float DeltaTi
 	if (NewLandscapePreviewMode == ENewLandscapePreviewMode::None)
 	{
 		const bool bStaleTargetLandscapeInfo = CurrentToolTarget.LandscapeInfo.IsStale();
-		const bool bStaleTargetLandscape = CurrentToolTarget.LandscapeInfo.IsValid() && (CurrentToolTarget.LandscapeInfo->GetLandscapeProxy() != nullptr);
+		const bool bStaleTargetLandscape = CurrentToolTarget.LandscapeInfo.IsValid() && (CurrentToolTarget.LandscapeInfo->GetLandscapeProxy() == nullptr);
 
 		if (bStaleTargetLandscapeInfo || bStaleTargetLandscape)
 		{
@@ -2497,7 +2481,10 @@ int32 FEdModeLandscape::UpdateLandscapeList()
 			SetLandscapeInfo(LandscapeList[0].Info);
 			CurrentIndex = 0;
 
-			SetCurrentLayer(0);
+			if (CanHaveLandscapeLayersContent())
+			{
+				SetSelectedEditLayer(0);
+			}
 
 			UpdateLayerUsageInformation();
 						
@@ -2516,7 +2503,7 @@ int32 FEdModeLandscape::UpdateLandscapeList()
 	}
 	else if (!HasValidLandscapeEditLayerSelection())
 	{
-		SetCurrentLayer(0);
+		SetSelectedEditLayer(0);
 	}
 
 	if (!CanEditCurrentTarget())
@@ -2554,11 +2541,6 @@ void FEdModeLandscape::SetTargetLandscape(const TWeakObjectPtr<ULandscapeInfo>&
 		}
 	}
 
-	// Force UISettings to be cleared to reset transient properties
-	SaveAndResetUISettings();
-	// Load back the saved properties
-	UISettings->Load();
-
 	SetLandscapeInfo(InLandscapeInfo.Get());
 	// force a Leave and Enter the current tool, in case it has something about the current landscape cached
 	SetCurrentTool(CurrentToolIndex);
@@ -2630,11 +2612,11 @@ void FEdModeLandscape::UpdateTargetList(bool bRegenerateThumbnails)
 		if (Landscape != nullptr)
 		{
 			CachedLandscapeMaterial = Landscape->GetLandscapeMaterial();
-
+			const int32 SelectedEditLayerIndex = GetSelectedEditLayerIndex();
 			bool bFoundSelected = false;
 
 			// Add heightmap
-			LandscapeTargetList.Add(MakeShareable(new FLandscapeTargetListInfo(LOCTEXT("Heightmap", "Heightmap"), ELandscapeToolTargetType::Heightmap, CurrentToolTarget.LandscapeInfo.Get(), UISettings->CurrentLayerIndex, /* bIsLayerReferencedByMaterial=*/ false)));
+			LandscapeTargetList.Add(MakeShareable(new FLandscapeTargetListInfo(LOCTEXT("Heightmap", "Heightmap"), ELandscapeToolTargetType::Heightmap, CurrentToolTarget.LandscapeInfo.Get(), SelectedEditLayerIndex, /* bIsLayerReferencedByMaterial=*/ false)));
 
 			if (CurrentToolTarget.TargetType == ELandscapeToolTargetType::Heightmap)
 			{
@@ -2643,7 +2625,7 @@ void FEdModeLandscape::UpdateTargetList(bool bRegenerateThumbnails)
 
 			// Add visibility
 			FLandscapeInfoLayerSettings VisibilitySettings(ALandscapeProxy::VisibilityLayer, Landscape);
-			LandscapeTargetList.Add(MakeShareable(new FLandscapeTargetListInfo(LOCTEXT("Visibility", "Visibility"), ELandscapeToolTargetType::Visibility, VisibilitySettings, UISettings->CurrentLayerIndex, /* bIsLayerReferencedByMaterial=*/ false)));
+			LandscapeTargetList.Add(MakeShareable(new FLandscapeTargetListInfo(LOCTEXT("Visibility", "Visibility"), ELandscapeToolTargetType::Visibility, VisibilitySettings, SelectedEditLayerIndex, /* bIsLayerReferencedByMaterial=*/ false)));
 
 			if (CurrentToolTarget.TargetType == ELandscapeToolTargetType::Visibility)
 			{
@@ -2688,7 +2670,7 @@ void FEdModeLandscape::UpdateTargetList(bool bRegenerateThumbnails)
 				}
 
 				// Add the layer
-				LandscapeTargetList.Add(MakeShareable(new FLandscapeTargetListInfo(FText::FromName(LayerName), ELandscapeToolTargetType::Weightmap, LayerSettings, UISettings->CurrentLayerIndex, /* bIsLayerReferencedByMaterial=*/ LayerNames.Contains(LayerName))));
+				LandscapeTargetList.Add(MakeShareable(new FLandscapeTargetListInfo(FText::FromName(LayerName), ELandscapeToolTargetType::Weightmap, LayerSettings, SelectedEditLayerIndex, /* bIsLayerReferencedByMaterial=*/ LayerNames.Contains(LayerName))));
 			}
 
 			if (!bFoundSelected)
@@ -3002,13 +2984,6 @@ void FEdModeLandscape::HandleLevelsChanged()
 	}
 }
 
-void FEdModeLandscape::HandleOnMapChanged(UWorld* World, EMapChangeType ChangeType)
-{
-	// Mark the flag so we don't overwrite saved data during Exit()
-	bHasMapChanged = true;
-	SaveAndResetUISettings();
-}
-
 void FEdModeLandscape::OnMaterialCompilationFinished(UMaterialInterface* MaterialInterface)
 {
 	// TODO [jonathan.bard] : we should remove this now that the target layer list is not automatically filled anymore
@@ -4017,13 +3992,12 @@ int32 FEdModeLandscape::GetLayerCount() const
 	return Landscape ? Landscape->GetEditLayersConst().Num() : 0;
 }
 
-void FEdModeLandscape::SetCurrentLayer(int32 InLayerIndex)
+void FEdModeLandscape::SetSelectedEditLayer(int32 InLayerIndex)
 {
-	UISettings->Modify();
-	UISettings->CurrentLayerIndex = InLayerIndex;
-
 	if (ALandscape* Landscape = GetLandscape())
 	{
+		Landscape->Modify(false);
+		Landscape->SetSelectedEditLayerIndex(InLayerIndex);
 		const ULandscapeEditLayerBase* EditLayer = Landscape->GetEditLayerConst(InLayerIndex);
 
 		if (Landscape->HasLayersContent() 
@@ -4039,9 +4013,14 @@ void FEdModeLandscape::SetCurrentLayer(int32 InLayerIndex)
 	RequestLayersContentUpdateForceAll(ELandscapeLayerUpdateMode::Update_Client_Editing);
 }
 
-int32 FEdModeLandscape::GetCurrentLayerIndex() const
+int32 FEdModeLandscape::GetSelectedEditLayerIndex() const
 {
-	return UISettings ? UISettings->CurrentLayerIndex : INDEX_NONE;
+	if (ALandscape* Landscape = GetLandscape())
+	{
+		return Landscape->GetSelectedEditLayerIndex();
+	}
+
+	return INDEX_NONE;
 }
 
 ALandscape* FEdModeLandscape::GetLandscape() const
@@ -4099,7 +4078,7 @@ void FEdModeLandscape::AddBrushToCurrentLayer(ALandscapeBlueprintBrushBase* InBr
 		return;
 	}
 
-	Landscape->AddBrushToLayer(GetCurrentLayerIndex(), InBrush);
+	Landscape->AddBrushToLayer(GetSelectedEditLayerIndex(), InBrush);
 	RefreshDetailPanel();
 }
 
@@ -4112,7 +4091,7 @@ void FEdModeLandscape::RemoveBrushFromCurrentLayer(int32 InBrushIndex)
 		return;
 	}
 
-	Landscape->RemoveBrushFromLayer(GetCurrentLayerIndex(), InBrushIndex);
+	Landscape->RemoveBrushFromLayer(GetSelectedEditLayerIndex(), InBrushIndex);
 	RefreshDetailPanel();
 }
 
@@ -4120,7 +4099,7 @@ ALandscapeBlueprintBrushBase* FEdModeLandscape::GetBrushForCurrentLayer(int32 In
 {
 	if (ALandscape* Landscape = GetLandscape())
 	{
-		return Landscape->GetBrushForLayer(GetCurrentLayerIndex(), InBrushIndex);
+		return Landscape->GetBrushForLayer(GetSelectedEditLayerIndex(), InBrushIndex);
 	}
 	return nullptr;
 }
@@ -4130,7 +4109,7 @@ TArray<ALandscapeBlueprintBrushBase*> FEdModeLandscape::GetBrushesForCurrentLaye
 	TArray<ALandscapeBlueprintBrushBase*> Brushes;
 	if (ALandscape* Landscape = GetLandscape())
 	{
-		Brushes = Landscape->GetBrushesForLayer(GetCurrentLayerIndex());
+		Brushes = Landscape->GetBrushesForLayer(GetSelectedEditLayerIndex());
 	}
 	return Brushes;
 }
@@ -4158,7 +4137,7 @@ void FEdModeLandscape::DuplicateBrush(class ALandscapeBlueprintBrushBase* InBrus
 
 const ULandscapeEditLayerBase* FEdModeLandscape::GetCurrentEditLayerConst() const
 {
-	return GetEditLayerConst(GetCurrentLayerIndex());
+	return GetEditLayerConst(GetSelectedEditLayerIndex());
 }
 
 void FEdModeLandscape::AutoUpdateDirtyLandscapeSplines()
@@ -4330,18 +4309,6 @@ void FEdModeLandscape::OnLevelActorRemoved(AActor* InActor)
 	}
 }
 
-void FEdModeLandscape::SaveAndResetUISettings()
-{
-	if (UISettings != nullptr)
-	{
-		UISettings->Save();
-		
-		// Recreate object to reset all values
-		UISettings = NewObject<ULandscapeEditorObject>(GetTransientPackage(), TEXT("UISettings"), RF_Transactional);
-		UISettings->SetParent(this);
-	}
-}
-
 bool LandscapeEditorUtils::SetHeightmapData(ALandscapeProxy* Landscape, const TArray<uint16>& Data)
 {
 	FIntRect ComponentsRect = Landscape->GetBoundingRect() + Landscape->LandscapeSectionOffset;
diff --git a/Engine/Source/Editor/LandscapeEditor/Private/LandscapeEdMode.h b/Engine/Source/Editor/LandscapeEditor/Private/LandscapeEdMode.h
--- a/Engine/Source/Editor/LandscapeEditor/Private/LandscapeEdMode.h
+++ b/Engine/Source/Editor/LandscapeEditor/Private/LandscapeEdMode.h
@@ -33,8 +33,6 @@ struct FHeightmapToolTarget;
 struct FViewportActionKeyInput;
 struct FViewportClick;
 template<class ToolTarget> class FLandscapeToolCopyPaste;
-class UWorld;
-enum class EMapChangeType : uint8;
 
 DECLARE_LOG_CATEGORY_EXTERN(LogLandscapeEdMode, Log, All);
 
@@ -588,8 +586,8 @@ public:
 	bool CanHaveLandscapeLayersContent() const;
 	bool HasLandscapeLayersContent() const;
 	int32 GetLayerCount() const;
-	void SetCurrentLayer(int32 InLayerIndex);
-	int32 GetCurrentLayerIndex() const;
+	void SetSelectedEditLayer(int32 InLayerIndex);
+	int32 GetSelectedEditLayerIndex() const;
 	ALandscape* GetLandscape() const;
 	bool CanRenameLayerTo(int32 InLayerIndex, const FName& InNewName) const;
 	ULandscapeEditLayerBase* GetEditLayer(int32 InLayerIndex) const;
@@ -624,10 +622,6 @@ public:
 	/** Handle notification that visible levels may have changed and we should update the editable landscapes list */
 	void HandleLevelsChanged();
 
-	/** Handle when map changes and UWorld is updated */
-	void HandleOnMapChanged(UWorld* World, EMapChangeType ChangeType);
-	void SaveAndResetUISettings();
-
 	void OnMaterialCompilationFinished(UMaterialInterface* MaterialInterface);
 
 	void ReimportData(const FLandscapeTargetListInfo& TargetInfo);
@@ -692,7 +686,6 @@ private:
 
 	FDelegateHandle OnWorldChangeDelegateHandle;
 	FDelegateHandle OnLevelsChangedDelegateHandle;
-	FDelegateHandle OnMapChangedDelegateHandle;
 	FDelegateHandle OnMaterialCompilationFinishedDelegateHandle;
 
 	FDelegateHandle OnLevelActorDeletedDelegateHandle;
diff --git a/Engine/Source/Editor/LandscapeEditor/Private/LandscapeEditorDetailCustomization_Layers.cpp b/Engine/Source/Editor/LandscapeEditor/Private/LandscapeEditorDetailCustomization_Layers.cpp
--- a/Engine/Source/Editor/LandscapeEditor/Private/LandscapeEditorDetailCustomization_Layers.cpp
+++ b/Engine/Source/Editor/LandscapeEditor/Private/LandscapeEditorDetailCustomization_Layers.cpp
@@ -416,7 +416,7 @@ bool FLandscapeEditorCustomNodeBuilder_Layers::IsLayerSelected(int32 InLayerInde
 	FEdModeLandscape* LandscapeEdMode = GetEditorMode();
 	if (LandscapeEdMode)
 	{
-		return LandscapeEdMode->GetCurrentLayerIndex() == InLayerIndex;
+		return LandscapeEdMode->GetSelectedEditLayerIndex() == InLayerIndex;
 	}
 
 	return false;
@@ -1141,10 +1141,10 @@ bool FLandscapeEditorCustomNodeBuilder_Layers::CanCreateLayer(FText& OutReason)
 
 void FLandscapeEditorCustomNodeBuilder_Layers::OnLayerSelectionChanged(int32 InLayerIndex)
 {
-	if (FEdModeLandscape* LandscapeEdMode = GetEditorMode(); LandscapeEdMode && LandscapeEdMode->GetCurrentLayerIndex() != InLayerIndex)
+	if (FEdModeLandscape* LandscapeEdMode = GetEditorMode(); LandscapeEdMode && LandscapeEdMode->GetSelectedEditLayerIndex() != InLayerIndex)
 	{
 		FScopedTransaction Transaction(LOCTEXT("Landscape_Layers_SetCurrentLayer", "Set Current Layer"));
-		LandscapeEdMode->SetCurrentLayer(InLayerIndex);
+		LandscapeEdMode->SetSelectedEditLayer(InLayerIndex);
 		LandscapeEdMode->UpdateTargetList();
 	}
 }
@@ -1412,7 +1412,7 @@ FReply FLandscapeEditorCustomNodeBuilder_Layers::HandleAcceptDrop(FDragDropEvent
 	if (DragDropOperation->IsOfType<FLandscapeBrushDragDropOp>())
 	{
 		int32 StartingBrushIndex = DragDropOperation->SlotIndexBeingDragged;
-		int32 StartingLayerIndex = LandscapeEdMode->GetCurrentLayerIndex();
+		int32 StartingLayerIndex = LandscapeEdMode->GetSelectedEditLayerIndex();
 		int32 DestinationLayerIndex = SlotIndexToLayerIndex(SlotIndex);
 
 		if (StartingLayerIndex == DestinationLayerIndex)
@@ -1431,7 +1431,7 @@ FReply FLandscapeEditorCustomNodeBuilder_Layers::HandleAcceptDrop(FDragDropEvent
 		Landscape->RemoveBrushFromLayer(StartingLayerIndex, StartingBrushIndex);
 		Landscape->AddBrushToLayer(DestinationLayerIndex, Brush);
 
-		LandscapeEdMode->SetCurrentLayer(DestinationLayerIndex);
+		LandscapeEdMode->SetSelectedEditLayer(DestinationLayerIndex);
 
 		// HACK: We don't return FReply::Handled() here because otherwise, SDragAndDropVerticalBox::OnDrop
 		// will apply UI slot reordering after we return. Properly speaking, we should have a way to signal 
@@ -1446,7 +1446,7 @@ FReply FLandscapeEditorCustomNodeBuilder_Layers::HandleAcceptDrop(FDragDropEvent
 	const FScopedTransaction Transaction(LOCTEXT("Landscape_Layers_Reorder", "Reorder Layer"));
 	if (Landscape->ReorderLayer(StartingLayerIndex, DestinationLayerIndex))
 	{
-		LandscapeEdMode->SetCurrentLayer(DestinationLayerIndex);
+		LandscapeEdMode->SetSelectedEditLayer(DestinationLayerIndex);
 		return FReply::Handled();
 	}
 
diff --git a/Engine/Source/Editor/LandscapeEditor/Private/LandscapeEditorDetailCustomization_LayersBrushStack.cpp b/Engine/Source/Editor/LandscapeEditor/Private/LandscapeEditorDetailCustomization_LayersBrushStack.cpp
--- a/Engine/Source/Editor/LandscapeEditor/Private/LandscapeEditorDetailCustomization_LayersBrushStack.cpp
+++ b/Engine/Source/Editor/LandscapeEditor/Private/LandscapeEditorDetailCustomization_LayersBrushStack.cpp
@@ -559,7 +559,7 @@ FReply FLandscapeEditorCustomNodeBuilder_LayersBrushStack::HandleAcceptDrop(FDra
 			int32 StartingBrushIndex = DragDropOperation->SlotIndexBeingDragged;
 			int32 DestinationBrushIndex = SlotIndex;
 			const FScopedTransaction Transaction(LOCTEXT("Landscape_LayerBrushes_Reorder", "Reorder Layer Brush"));
-			if (Landscape->ReorderLayerBrush(LandscapeEdMode->GetCurrentLayerIndex(), StartingBrushIndex, DestinationBrushIndex))
+			if (Landscape->ReorderLayerBrush(LandscapeEdMode->GetSelectedEditLayerIndex(), StartingBrushIndex, DestinationBrushIndex))
 			{
 				LandscapeEdMode->RefreshDetailPanel();
 				return FReply::Handled();
diff --git a/Engine/Source/Editor/LandscapeEditor/Private/LandscapeEditorDetailCustomization_TargetLayers.cpp b/Engine/Source/Editor/LandscapeEditor/Private/LandscapeEditorDetailCustomization_TargetLayers.cpp
--- a/Engine/Source/Editor/LandscapeEditor/Private/LandscapeEditorDetailCustomization_TargetLayers.cpp
+++ b/Engine/Source/Editor/LandscapeEditor/Private/LandscapeEditorDetailCustomization_TargetLayers.cpp
@@ -583,7 +583,7 @@ const TArray<const TSharedRef<FLandscapeTargetListInfo>> FLandscapeEditorCustomN
 			if (!LandscapeActor->GetTargetLayers().Find(LayerName))
 			{
 				const FLandscapeInfoLayerSettings LayerSettings = FLandscapeInfoLayerSettings(LayerName, LandscapeActor.Get());
-				const TSharedRef<FLandscapeTargetListInfo> Target = MakeShareable(new FLandscapeTargetListInfo(FText::FromName(LayerName), ELandscapeToolTargetType::Weightmap, LayerSettings, LandscapeEdMode->UISettings->CurrentLayerIndex, /* bIsLayerReferencedByMaterial=*/ true));
+				const TSharedRef<FLandscapeTargetListInfo> Target = MakeShareable(new FLandscapeTargetListInfo(FText::FromName(LayerName), ELandscapeToolTargetType::Weightmap, LayerSettings, LandscapeEdMode->GetSelectedEditLayerIndex(), /* bIsLayerReferencedByMaterial=*/ true));
 
 				TargetLayerList.Add(Target);
 			}
@@ -2047,7 +2047,7 @@ ECheckBoxState FLandscapeEditorCustomNodeBuilder_TargetLayers::IsLayersSubstract
 {
 	if (const FEdModeLandscape* LandscapeEdMode = GetEditorMode())
 	{
-		ULandscapeEditLayerBase* EditLayer = LandscapeEdMode->GetEditLayer(LandscapeEdMode->UISettings->CurrentLayerIndex);
+		ULandscapeEditLayerBase* EditLayer = LandscapeEdMode->GetEditLayer(LandscapeEdMode->GetSelectedEditLayerIndex());
 		const bool* AllocationBlend = EditLayer->GetWeightmapLayerAllocationBlend().Find(Target.Get().LayerInfoObj.Get());
 
 		if (AllocationBlend != nullptr)
@@ -2064,7 +2064,7 @@ void FLandscapeEditorCustomNodeBuilder_TargetLayers::OnLayersSubstractiveBlendCh
 	if (const FEdModeLandscape* LandscapeEdMode = GetEditorMode())
 	{
 		FScopedTransaction Transaction(LOCTEXT("Undo_SubtractiveBlend", "Set Subtractive Blend Layer"));
-		ULandscapeEditLayerBase* EditLayer = LandscapeEdMode->GetEditLayer(LandscapeEdMode->UISettings->CurrentLayerIndex);
+		ULandscapeEditLayerBase* EditLayer = LandscapeEdMode->GetEditLayer(LandscapeEdMode->GetSelectedEditLayerIndex());
 
 		EditLayer->AddOrUpdateWeightmapAllocationLayerBlend(Target.Get().LayerInfoObj.Get(), NewCheckedState == ECheckBoxState::Checked, /*bInModify =*/true);
 	}
diff --git a/Engine/Source/Editor/LandscapeEditor/Private/LandscapeEditorObject.cpp b/Engine/Source/Editor/LandscapeEditor/Private/LandscapeEditorObject.cpp
--- a/Engine/Source/Editor/LandscapeEditor/Private/LandscapeEditorObject.cpp
+++ b/Engine/Source/Editor/LandscapeEditor/Private/LandscapeEditorObject.cpp
@@ -127,7 +127,6 @@ ULandscapeEditorObject::ULandscapeEditorObject(const FObjectInitializer& ObjectI
 	, BrushComponentSize(1)
 	, TargetDisplayOrder(ELandscapeLayerDisplayMode::Default)
 	, ShowUnusedLayers(true)
-	, CurrentLayerIndex(INDEX_NONE)
 {
 	// Structure to hold one-time initialization
 	struct FConstructorStatics
diff --git a/Engine/Source/Editor/LandscapeEditor/Public/LandscapeEditorObject.h b/Engine/Source/Editor/LandscapeEditor/Public/LandscapeEditorObject.h
--- a/Engine/Source/Editor/LandscapeEditor/Public/LandscapeEditorObject.h
+++ b/Engine/Source/Editor/LandscapeEditor/Public/LandscapeEditorObject.h
@@ -747,9 +747,6 @@ public:
 	UPROPERTY(Category = "Target Layers", EditAnywhere)
 	bool ShowUnusedLayers;	
 
-	UPROPERTY(Transient)
-	int32 CurrentLayerIndex;
-
 #if WITH_EDITOR
 	//~ Begin UObject Interface
 	void PostEditChangeProperty(FPropertyChangedEvent& PropertyChangedEvent) override;
diff --git a/Engine/Source/Runtime/Landscape/Classes/Landscape.h b/Engine/Source/Runtime/Landscape/Classes/Landscape.h
--- a/Engine/Source/Runtime/Landscape/Classes/Landscape.h
+++ b/Engine/Source/Runtime/Landscape/Classes/Landscape.h
@@ -483,6 +483,8 @@ public:
 	LANDSCAPE_API void ShowOnlySelectedLayer(int32 InLayerIndex);
 	LANDSCAPE_API void ShowAllLayers();
 	LANDSCAPE_API void UpdateLandscapeSplines(const FGuid& InLayerGuid = FGuid(), bool bInUpdateOnlySelected = false, bool bInForceUpdateAllCompoments = false);
+	LANDSCAPE_API void SetSelectedEditLayerIndex(const int32 InEditLayerIndex);
+	LANDSCAPE_API const int32 GetSelectedEditLayerIndex() const;
 
 	UE_DEPRECATED(5.5, "Use CreateLayer with a ULandscapeEditLayerSplines instead")
 	LANDSCAPE_API void SetLandscapeSplinesReservedLayer(int32 InLayerIndex);
@@ -683,8 +685,13 @@ public:
 	FGuid LandscapeSplinesTargetLayerGuid_DEPRECATED;
 	
 	/** Current Editing Landscape Layer*/
+	// TODO this is used as shared global state for the landscape editor mode. FLandscapeToolStrokeBase::SetEditLayer should manage the shared editor state instead
 	FGuid EditingLayer;
 
+	/** Current Selected Edit Layer of this landscape. Used by landscape editor mode to track the current selection */
+	UPROPERTY(Transient)
+	int32 SelectedEditLayerIndex = INDEX_NONE;
+
 	/** Used to temporarily disable Grass Update in Editor */
 	bool bGrassUpdateEnabled;
 
diff --git a/Engine/Source/Runtime/Landscape/Private/Landscape.cpp b/Engine/Source/Runtime/Landscape/Private/Landscape.cpp
--- a/Engine/Source/Runtime/Landscape/Private/Landscape.cpp
+++ b/Engine/Source/Runtime/Landscape/Private/Landscape.cpp
@@ -3050,9 +3050,20 @@ PRAGMA_ENABLE_DEPRECATION_WARNINGS
 
 	// In case we're a landscape with edit layers but we actually lack a layer (e.g. it was removed by the test above, because its edit layer class is unknown), let's create one all the same 
 	//  because we're always supposed to have at least 1 : 
-	if (CanHaveLayersContent() && LandscapeEditLayers.IsEmpty())
+	if (CanHaveLayersContent())
 	{
-		ensure(CreateLayer() != INDEX_NONE);
+		if (LandscapeEditLayers.IsEmpty())
+		{
+			// Create a default layer
+			CreateDefaultLayer();
+			ensure(!LandscapeEditLayers.IsEmpty());
+		}
+
+		// Ensure a valid edit layer is selected post load
+		if (!LandscapeEditLayers.IsValidIndex(SelectedEditLayerIndex))
+		{
+			SetSelectedEditLayerIndex(0);
+		}
 	}
 #endif // WITH_EDITOR
 
diff --git a/Engine/Source/Runtime/Landscape/Private/LandscapeEditLayers.cpp b/Engine/Source/Runtime/Landscape/Private/LandscapeEditLayers.cpp
--- a/Engine/Source/Runtime/Landscape/Private/LandscapeEditLayers.cpp
+++ b/Engine/Source/Runtime/Landscape/Private/LandscapeEditLayers.cpp
@@ -2557,12 +2557,14 @@ void ALandscape::ToggleCanHaveLayersContent()
 	{
 		ReleaseLayersRenderingResource();
 		DeleteLayers();
+		SetSelectedEditLayerIndex(INDEX_NONE);
 	}
 	else
 	{
 		check(LandscapeEditLayers.IsEmpty());
 		CreateDefaultLayer();
 		CopyOldDataToDefaultLayer();
+		check(SelectedEditLayerIndex == 0);
 	}
 
 	if (LandscapeEdMode)
@@ -12708,6 +12710,19 @@ bool ALandscape::DeleteLayer(int32 InLayerIndex)
 
 	Modify();
 
+	// If this edit layer is currently selected, select the one below it or the first layer
+	if (GetSelectedEditLayerIndex() == InLayerIndex)
+	{
+		if (LandscapeEditLayers.IsValidIndex(InLayerIndex - 1))
+		{
+			SetSelectedEditLayerIndex(InLayerIndex - 1);
+		}
+		else
+		{
+			SetSelectedEditLayerIndex(0);
+		}
+	}
+
 	// We're about to remove the layer from our list, which will invalidate our LayerStruct pointer.
 	// We'll need to call OnLayerRemoved afterward, though, so keep pointer to the UObject.
 	ULandscapeEditLayerBase* EditLayer = LayerStruct->EditLayer;
@@ -13189,6 +13204,19 @@ const FGuid& ALandscape::GetEditingLayer() const
 	return EditingLayer;
 }
 
+void ALandscape::SetSelectedEditLayerIndex(const int32 InEditLayerIndex)
+{
+	check(CanHaveLayersContent() ? LandscapeEditLayers.IsValidIndex(InEditLayerIndex) : InEditLayerIndex == INDEX_NONE);
+	SelectedEditLayerIndex = InEditLayerIndex;
+}
+
+const int32 ALandscape::GetSelectedEditLayerIndex() const
+{
+	// When edit layers are not supported, index should always be NONE
+	check(CanHaveLayersContent() ? LandscapeEditLayers.IsValidIndex(SelectedEditLayerIndex) : SelectedEditLayerIndex == INDEX_NONE);
+	return SelectedEditLayerIndex;
+}
+
 bool ALandscape::IsMaxLayersReached() const
 {
 	return LandscapeEditLayers.Num() >= GetDefault<ULandscapeSettings>()->MaxNumberOfLayers;
@@ -13204,6 +13232,7 @@ void ALandscape::CreateDefaultLayer()
 	check(LandscapeEditLayers.Num() == 0); // We can only call this function if we have no layers
 
 	CreateLayer(FName(TEXT("Layer")));
+	SetSelectedEditLayerIndex(0);
 }
 
 FLandscapeLayer* ALandscape::DuplicateLayerAndMoveBrushes(const FLandscapeLayer& InOtherLayer)
