Landscape Editor - Making BPBrushBase non placeable to ensure brushes are only added from Landscape BP Brush tool
* Adds NonPlaceable to BP Brush base
* Removes GetOrCreateEditLayer from brush classes since Place actor path is no longer supported
* Removes legacy Add Brush option from EditLayer and BPBrush context menus

diff --git a/Engine/Plugins/Experimental/Landmass/Source/Editor/Private/LandmassErosionBrushBase.cpp b/Engine/Plugins/Experimental/Landmass/Source/Editor/Private/LandmassErosionBrushBase.cpp
--- a/Engine/Plugins/Experimental/Landmass/Source/Editor/Private/LandmassErosionBrushBase.cpp
+++ b/Engine/Plugins/Experimental/Landmass/Source/Editor/Private/LandmassErosionBrushBase.cpp
@@ -20,13 +20,14 @@ ALandmassErosionBrushBase::ALandmassErosionBrushBase()
 		FLevelEditorModule& LevelEditor = FModuleManager::GetModuleChecked<FLevelEditorModule>("LevelEditor");
 		OnActorSelectionChangedHandle = LevelEditor.OnActorSelectionChanged().AddUObject(this, &ALandmassErosionBrushBase::HandleActorSelectionChanged);
 	}
+
+	SetCanAffectHeightmap(true);
 }
 
 void ALandmassErosionBrushBase::OnConstruction(const FTransform& Transform)
 {
 	Super::OnConstruction(Transform);
 	FindAndAssignLandscape();
-
 }
 
 void ALandmassErosionBrushBase::SetTargetLandscape(ALandscape* InTargetLandscape)
@@ -61,25 +62,6 @@ void ALandmassErosionBrushBase::SetTargetLandscape(ALandscape* InTargetLandscape
 			UE_LOG(LandmassErosionBrush, Warning, TEXT("Landscape target for Landmass Erosion Brush did not have edit layers enabled. Unable to attach manager."));
 			SetOwningLandscape(nullptr);
 		}
-		else
-		{
-			static const FName PatchLayerName = FName("Landmass Erosion");
-
-			ILandscapeModule& LandscapeModule = FModuleManager::GetModuleChecked<ILandscapeModule>("Landscape");
-			int32 PatchLayerIndex = LandscapeModule.GetLandscapeEditorServices()->GetOrCreateEditLayer(PatchLayerName, InTargetLandscape);
-
-			// Among other things, this will call SetOwningLandscape on us.
-			InTargetLandscape->AddBrushToLayer(PatchLayerIndex, this);
-
-			// It's not clear whether this is really necessary, but we do it for consistency because Landscape does this in its
-			// PostLoad for all its brushes (through FLandscapeLayerBrush::SetOwner). One would think that it would be done 
-			// in AddBrushToLayer if it were at all important, but it currently isn't...
-			if (this->GetTypedOuter<ULevel>() != InTargetLandscape->GetTypedOuter<ULevel>())
-			{
-				// Change owner to be that level
-				this->Rename(nullptr, InTargetLandscape->GetTypedOuter<ULevel>());
-			}
-		}
 	}
 #endif
 }
diff --git a/Engine/Plugins/Experimental/Landmass/Source/Editor/Private/LandmassManagerBase.cpp b/Engine/Plugins/Experimental/Landmass/Source/Editor/Private/LandmassManagerBase.cpp
--- a/Engine/Plugins/Experimental/Landmass/Source/Editor/Private/LandmassManagerBase.cpp
+++ b/Engine/Plugins/Experimental/Landmass/Source/Editor/Private/LandmassManagerBase.cpp
@@ -21,6 +21,9 @@ ALandmassManagerBase::ALandmassManagerBase()
 		FLevelEditorModule& LevelEditor = FModuleManager::GetModuleChecked<FLevelEditorModule>("LevelEditor");
 		OnActorSelectionChangedHandle = LevelEditor.OnActorSelectionChanged().AddUObject(this, &ALandmassManagerBase::HandleActorSelectionChanged);
 	}
+
+	// Explicitly set flag so this can be placed as a BP Brush from the landscape editor 
+	SetCanAffectHeightmap(true);
 }
 
 void ALandmassManagerBase::SetTargetLandscape(ALandscape* InTargetLandscape)
@@ -49,25 +52,6 @@ void ALandmassManagerBase::SetTargetLandscape(ALandscape* InTargetLandscape)
 			UE_LOG(LandmassManager, Warning, TEXT("Landscape target for Landmass Brush manager did not have edit layers enabled. Unable to attach manager."));
 			SetOwningLandscape(nullptr);
 		}
-		else
-		{
-			static const FName PatchLayerName = FName("Landmass Brushes");
-
-			ILandscapeModule& LandscapeModule = FModuleManager::GetModuleChecked<ILandscapeModule>("Landscape");
-			int32 PatchLayerIndex = LandscapeModule.GetLandscapeEditorServices()->GetOrCreateEditLayer(PatchLayerName, InTargetLandscape);
-
-			// Among other things, this will call SetOwningLandscape on us.
-			InTargetLandscape->AddBrushToLayer(PatchLayerIndex, this);
-
-			// It's not clear whether this is really necessary, but we do it for consistency because Landscape does this in its
-			// PostLoad for all its brushes (through FLandscapeLayerBrush::SetOwner). One would think that it would be done 
-			// in AddBrushToLayer if it were at all important, but it currently isn't...
-			if (this->GetTypedOuter<ULevel>() != InTargetLandscape->GetTypedOuter<ULevel>())
-			{
-				// Change owner to be that level
-				this->Rename(nullptr, InTargetLandscape->GetTypedOuter<ULevel>());
-			}
-		}
 	}
 #endif
 }
diff --git a/Engine/Plugins/Experimental/LandscapePatch/Source/LandscapePatch/Private/LandscapePatchManager.cpp b/Engine/Plugins/Experimental/LandscapePatch/Source/LandscapePatch/Private/LandscapePatchManager.cpp
--- a/Engine/Plugins/Experimental/LandscapePatch/Source/LandscapePatch/Private/LandscapePatchManager.cpp
+++ b/Engine/Plugins/Experimental/LandscapePatch/Source/LandscapePatch/Private/LandscapePatchManager.cpp
@@ -322,30 +322,6 @@ void ALandscapePatchManager::SetTargetLandscape(ALandscape* InTargetLandscape)
 			UE_LOG(LogLandscapePatch, Warning, TEXT("Landscape target for patch manager did not have edit layers enabled. Unable to attach manager."));
 			SetOwningLandscape(nullptr);
 		}
-		else
-		{
-			static const FName PatchLayerName = FName("LandscapePatches");
-
-			ILandscapeModule& LandscapeModule = FModuleManager::GetModuleChecked<ILandscapeModule>("Landscape");
-			int32 PatchLayerIndex = LandscapeModule.GetLandscapeEditorServices()->GetOrCreateEditLayer(PatchLayerName, InTargetLandscape);
-
-			if (!ensure(PatchLayerIndex >= 0))
-			{
-				return;
-			}
-
-			// Among other things, this will call SetOwningLandscape on us.
-			InTargetLandscape->AddBrushToLayer(PatchLayerIndex, this);
-
-			// It's not clear whether this is really necessary, but we do it for consistency because Landscape does this in its
-			// PostLoad for all its brushes (through FLandscapeLayerBrush::SetOwner). One would think that it would be done 
-			// in AddBrushToLayer if it were at all important, but it currently isn't...
-			if (this->GetTypedOuter<ULevel>() != InTargetLandscape->GetTypedOuter<ULevel>())
-			{
-				// Change owner to be that level
-				this->Rename(nullptr, InTargetLandscape->GetTypedOuter<ULevel>());
-			}
-		}
 	}
 #endif
 }
diff --git a/Engine/Plugins/Experimental/LandscapePatch/Source/LandscapePatch/Public/LandscapePatchManager.h b/Engine/Plugins/Experimental/LandscapePatch/Source/LandscapePatch/Public/LandscapePatchManager.h
--- a/Engine/Plugins/Experimental/LandscapePatch/Source/LandscapePatch/Public/LandscapePatchManager.h
+++ b/Engine/Plugins/Experimental/LandscapePatch/Source/LandscapePatch/Public/LandscapePatchManager.h
@@ -49,6 +49,9 @@ public:
 		TArray<TBitArray<>>& OutTargetLayerGroups) const override {}
 	virtual TArray<UE::Landscape::EditLayers::FEditLayerRenderItem> GetRenderItems(const UE::Landscape::EditLayers::FMergeContext* InMergeContext) const override { return {}; }
 	virtual bool RenderLayer(UE::Landscape::EditLayers::FRenderParams& RenderParams, UE::Landscape::FRDGBuilderRecorder& RDGBuilderRecorder) override { return false; };
+
+	// In 5.6 users should not be able to add new Patch Managers anywhere in the editor
+	virtual bool SupportsBlueprintBrushTool() const { return false; }
 #endif // WITH_EDITOR
 
 	// Adds the brush to the given landscape, removing it from any previous one. This differs from SetOwningLandscape
diff --git a/Engine/Plugins/Experimental/Water/Source/Editor/Private/WaterEditorModule.cpp b/Engine/Plugins/Experimental/Water/Source/Editor/Private/WaterEditorModule.cpp
--- a/Engine/Plugins/Experimental/Water/Source/Editor/Private/WaterEditorModule.cpp
+++ b/Engine/Plugins/Experimental/Water/Source/Editor/Private/WaterEditorModule.cpp
@@ -222,14 +222,7 @@ void FWaterEditorModule::OnLevelActorAddedToWorld(AActor* Actor)
 					{
 						UActorFactory* WaterBrushActorFactory = GEditor->FindActorFactoryForActorClass(WaterBrushClassPtr);
 
-						// Attempt to find an existing water layer, else create one before spawning the water brush
-						const FName WaterLayerName = FName("Water");
-						int32 ExistingWaterLayerIndex = FoundLandscape->GetLayerIndex(WaterLayerName);
-						if (ExistingWaterLayerIndex == INDEX_NONE)
-						{
-							ExistingWaterLayerIndex = FoundLandscape->CreateLayer(WaterLayerName);
-						}
-
+						// The manager is added from the BP brush tool only and is being spawned on a valid edit layer
 						FString BrushActorString = FString::Format(TEXT("{0}_{1}"), { FoundLandscape->GetActorLabel(), WaterBrushClassPtr->GetName() });
 						FName BrushActorName = MakeUniqueObjectName(FoundLandscape->GetOuter(), WaterBrushClassPtr, FName(BrushActorString));
 						FActorSpawnParameters SpawnParams;
@@ -249,12 +242,6 @@ void FWaterEditorModule::OnLevelActorAddedToWorld(AActor* Actor)
 							bHasWaterManager = true;
 							NewBrush->SetActorLabel(BrushActorString);
 							NewBrush->SetTargetLandscape(FoundLandscape);
-							const int32 CurrentBrushLayer = FoundLandscape->GetBrushLayer(NewBrush);
-							if (CurrentBrushLayer != ExistingWaterLayerIndex)
-							{
-								FoundLandscape->RemoveBrush(NewBrush);
-								FoundLandscape->AddBrushToLayer(ExistingWaterLayerIndex, NewBrush);
-							}
 						}
 					}
 				}
diff --git a/Engine/Source/Editor/LandscapeEditor/Private/LandscapeEdModeBPCustomTools.cpp b/Engine/Source/Editor/LandscapeEditor/Private/LandscapeEdModeBPCustomTools.cpp
--- a/Engine/Source/Editor/LandscapeEditor/Private/LandscapeEdModeBPCustomTools.cpp
+++ b/Engine/Source/Editor/LandscapeEditor/Private/LandscapeEdModeBPCustomTools.cpp
@@ -81,9 +81,8 @@ public:
 			return false;
 		}
 
-		// Only allow placing brushes if our landscape can be edited and the current edit layer supports BP brushes
 		const ULandscapeEditLayerBase* EditLayer = EdMode->GetCurrentEditLayerConst();
-
+		// Brushes can only be placed if the BP Brush can affect the landscape and the current edit layer supports BP brushes
 		if ((EditLayer && EditLayer->SupportsBlueprintBrushes()) && (DefaultObject->CanAffectHeightmap() || DefaultObject->CanAffectWeightmap() || DefaultObject->CanAffectVisibilityLayer()))
 		{
 			ULandscapeInfo* Info = EdMode->CurrentToolTarget.LandscapeInfo.Get();
@@ -100,6 +99,12 @@ public:
 			SpawnInfo.OverrideLevel = Info->LandscapeActor.Get()->GetLevel(); // always spawn in the same level as the one containing the ALandscape
 			SpawnInfo.Name = BrushActorName;
 
+			// If brush overrides default behavior do not create one from the BP brush tool 
+			if (!DefaultObject->SupportsBlueprintBrushTool())
+			{
+				return false;
+			}
+
 			FScopedTransaction Transaction(LOCTEXT("LandscapeEdModeBlueprintToolSpawn", "Create landscape brush"));
 
 			// Use the class factory if there's one :
diff --git a/Engine/Source/Editor/LandscapeEditor/Private/LandscapeEditorDetailCustomization_Layers.cpp b/Engine/Source/Editor/LandscapeEditor/Private/LandscapeEditorDetailCustomization_Layers.cpp
--- a/Engine/Source/Editor/LandscapeEditor/Private/LandscapeEditorDetailCustomization_Layers.cpp
+++ b/Engine/Source/Editor/LandscapeEditor/Private/LandscapeEditorDetailCustomization_Layers.cpp
@@ -453,14 +453,6 @@ FSlateColor FLandscapeEditorCustomNodeBuilder_Layers::GetLayerTextColor(int32 In
 	return IsLayerSelected(InLayerIndex) ? FStyleColors::ForegroundHover : FSlateColor::UseForeground();
 }
 
-void FLandscapeEditorCustomNodeBuilder_Layers::FillAddBrushMenu(FMenuBuilder& MenuBuilder, TArray<ALandscapeBlueprintBrushBase*> Brushes)
-{
-	for (ALandscapeBlueprintBrushBase* Brush : Brushes)
-	{
-		FUIAction AddAction = FUIAction(FExecuteAction::CreateSP(this, &FLandscapeEditorCustomNodeBuilder_Layers::AddBrushToCurrentLayer, Brush));
-		MenuBuilder.AddMenuEntry(FText::FromString(Brush->GetActorLabel()), FText(), FSlateIcon(), AddAction);
-	}
-}
 
 void FLandscapeEditorCustomNodeBuilder_Layers::FillClearTargetLayerMenu(FMenuBuilder& MenuBuilder, int32 InLayerIndex, TArray<ULandscapeLayerInfoObject*> InUsedLayerInfos)
 {
@@ -525,12 +517,6 @@ void FLandscapeEditorCustomNodeBuilder_Layers::FillClearLayerMenu(FMenuBuilder&
 	}
 }
 
-void FLandscapeEditorCustomNodeBuilder_Layers::AddBrushToCurrentLayer(ALandscapeBlueprintBrushBase* Brush)
-{
-	const FScopedTransaction Transaction(LOCTEXT("LandscapeBrushAddToCurrentLayerTransaction", "Add brush to current layer"));
-	GetEditorMode()->AddBrushToCurrentLayer(Brush);
-}
-
 TSharedPtr<SWidget> FLandscapeEditorCustomNodeBuilder_Layers::OnLayerContextMenuOpening(int32 InLayerIndex)
 {
 	FEdModeLandscape* LandscapeEdMode = GetEditorMode();
@@ -627,23 +613,6 @@ TSharedPtr<SWidget> FLandscapeEditorCustomNodeBuilder_Layers::OnLayerContextMenu
 		}
 		MenuBuilder.EndSection();
 
-		const TArray<ALandscapeBlueprintBrushBase*>& Brushes = LandscapeEdMode->GetBrushList();
-		TArray<ALandscapeBlueprintBrushBase*> FilteredBrushes = Brushes.FilterByPredicate([](ALandscapeBlueprintBrushBase* Brush) { return Brush->GetOwningLandscape() == nullptr; });
-		if (FilteredBrushes.Num())
-		{
-			MenuBuilder.BeginSection("LandscapeEditorBrushActions", LOCTEXT("LandscapeEditorBrushActions.Heading", "Brushes"));
-			{
-				MenuBuilder.AddSubMenu(
-					LOCTEXT("LandscaeEditorBrushAddSubMenu", "Add"),
-					LOCTEXT("LandscaeEditorBrushAddSubMenuToolTip", "Add brush to current layer"),
-					FNewMenuDelegate::CreateSP(this, &FLandscapeEditorCustomNodeBuilder_Layers::FillAddBrushMenu, FilteredBrushes),
-					false,
-					FSlateIcon()
-				);
-			}
-			MenuBuilder.EndSection();
-		}
-
 		return MenuBuilder.MakeWidget();
 	}
 	return nullptr;
diff --git a/Engine/Source/Editor/LandscapeEditor/Private/LandscapeEditorDetailCustomization_Layers.h b/Engine/Source/Editor/LandscapeEditor/Private/LandscapeEditorDetailCustomization_Layers.h
--- a/Engine/Source/Editor/LandscapeEditor/Private/LandscapeEditorDetailCustomization_Layers.h
+++ b/Engine/Source/Editor/LandscapeEditor/Private/LandscapeEditorDetailCustomization_Layers.h
@@ -128,8 +128,6 @@ protected:
 	FReply OnToggleLock(int32 InLayerIndex);
 	const FSlateBrush* GetLockBrushForLayer(int32 InLayerIndex) const;
 
-	void FillAddBrushMenu(FMenuBuilder& MenuBuilder, TArray<ALandscapeBlueprintBrushBase*> Brushes);
-	void AddBrushToCurrentLayer(ALandscapeBlueprintBrushBase* Brush);
 	void FillClearTargetLayerMenu(FMenuBuilder& MenuBuilder, int32 InLayerIndex, TArray<ULandscapeLayerInfoObject*> InUsedLayerInfos);
 	void FillClearLayerMenu(FMenuBuilder& MenuBuilder, int32 InLayerIndex);
 
diff --git a/Engine/Source/Editor/LandscapeEditor/Private/LandscapeEditorDetailCustomization_LayersBrushStack.cpp b/Engine/Source/Editor/LandscapeEditor/Private/LandscapeEditorDetailCustomization_LayersBrushStack.cpp
--- a/Engine/Source/Editor/LandscapeEditor/Private/LandscapeEditorDetailCustomization_LayersBrushStack.cpp
+++ b/Engine/Source/Editor/LandscapeEditor/Private/LandscapeEditorDetailCustomization_LayersBrushStack.cpp
@@ -285,29 +285,15 @@ TSharedPtr<SWidget> FLandscapeEditorCustomNodeBuilder_LayersBrushStack::OnBrushC
 					GetEditorMode()->DuplicateBrush(CurrentBrush);
 				}));
 				MenuBuilder.AddMenuEntry(LOCTEXT("LandscapeBrushDuplicate", "Duplicate"), LOCTEXT("LandscapeBrushDuplicateToolTip", "Duplicate brush"), FSlateIcon(), DuplicateAction);
-			}
 
-			// Add Brush
-			const TArray<ALandscapeBlueprintBrushBase*>& AllBrushes = GetEditorMode()->GetBrushList();
-			TArray<ALandscapeBlueprintBrushBase*> FilteredBrushes = AllBrushes.FilterByPredicate([](ALandscapeBlueprintBrushBase* Brush) { return Brush->GetOwningLandscape() == nullptr; });
-			if (FilteredBrushes.Num())
-			{
-				MenuBuilder.AddSubMenu(
-					LOCTEXT("LandscapeEditorBrushAddSubMenu", "Add"),
-					LOCTEXT("LandscapeEditorBrushAddSubMenuToolTip", "Add brush to selected edit layer"),
-					FNewMenuDelegate::CreateSP(this, &FLandscapeEditorCustomNodeBuilder_LayersBrushStack::FillAddBrushMenu, FilteredBrushes),
-					false,
-					FSlateIcon()
-				);
+				// Remove Brush
+				FUIAction RemoveAction = FUIAction(FExecuteAction::CreateLambda([SharedThis, InBrushIndex]
+					{
+						const FScopedTransaction Transaction(LOCTEXT("LandscapeBrushRemoveTransaction", "Remove Brush"));
+						GetEditorMode()->RemoveBrushFromCurrentLayer(InBrushIndex);
+					}));
+				MenuBuilder.AddMenuEntry(LOCTEXT("LandscapeBrushRemove", "Remove"), LOCTEXT("LandscapeBrushRemoveToolTip", "Remove brush from selected edit layer"), FSlateIcon(), RemoveAction);
 			}
-
-			// Remove Brush
-			FUIAction RemoveAction = FUIAction(FExecuteAction::CreateLambda([SharedThis, InBrushIndex]
-			{ 
-				const FScopedTransaction Transaction(LOCTEXT("LandscapeBrushRemoveTransaction", "Remove Brush"));
-				GetEditorMode()->RemoveBrushFromCurrentLayer(InBrushIndex);
-			}));
-			MenuBuilder.AddMenuEntry(LOCTEXT("LandscapeBrushRemove", "Remove"), LOCTEXT("LandscapeBrushRemoveToolTip", "Remove brush from selected edit layer"), FSlateIcon(), RemoveAction);
 		}
 		MenuBuilder.EndSection();
 		return MenuBuilder.MakeWidget();
@@ -316,20 +302,6 @@ TSharedPtr<SWidget> FLandscapeEditorCustomNodeBuilder_LayersBrushStack::OnBrushC
 	return nullptr;
 }
 
-void FLandscapeEditorCustomNodeBuilder_LayersBrushStack::FillAddBrushMenu(FMenuBuilder& MenuBuilder, TArray<ALandscapeBlueprintBrushBase*> Brushes)
-{
-	TSharedRef<FLandscapeEditorCustomNodeBuilder_LayersBrushStack> SharedThis = AsShared();
-	for (ALandscapeBlueprintBrushBase* Brush : Brushes)
-	{
-		FUIAction AddAction = FUIAction(FExecuteAction::CreateLambda([SharedThis, Brush]()
-		{
-			const FScopedTransaction Transaction(LOCTEXT("LandscapeBrushAddToSelectedLayerTransaction", "Add brush to selected edit layer"));
-			GetEditorMode()->AddBrushToCurrentLayer(Brush);
-		}));
-		MenuBuilder.AddMenuEntry(FText::FromString(Brush->GetActorLabel()), FText(), FSlateIcon(), AddAction);
-	}
-}
-
 void FLandscapeEditorCustomNodeBuilder_LayersBrushStack::OnToggleAffectsHeightmap(ECheckBoxState InCheckBoxState, int32 InBrushIndex)
 {
 	if ((InCheckBoxState != ECheckBoxState::Checked) && (InCheckBoxState != ECheckBoxState::Unchecked))
diff --git a/Engine/Source/Editor/LandscapeEditor/Private/LandscapeEditorDetailCustomization_LayersBrushStack.h b/Engine/Source/Editor/LandscapeEditor/Private/LandscapeEditorDetailCustomization_LayersBrushStack.h
--- a/Engine/Source/Editor/LandscapeEditor/Private/LandscapeEditorDetailCustomization_LayersBrushStack.h
+++ b/Engine/Source/Editor/LandscapeEditor/Private/LandscapeEditorDetailCustomization_LayersBrushStack.h
@@ -61,8 +61,6 @@ protected:
 	TSharedPtr<SWidget> GenerateRow(int32 InBrushIndex);
 	TSharedPtr<SWidget> OnBrushContextMenuOpening(int32 InBrushIndex);
 
-	void FillAddBrushMenu(FMenuBuilder& MenuBuilder, TArray<ALandscapeBlueprintBrushBase*> Brushes);
-
 	FReply OnToggleVisibility(int32 InBrushIndex);
 	void OnToggleAffectsHeightmap(ECheckBoxState InCheckBoxState, int32 InBrushIndex);
 	void OnToggleAffectsWeightmap(ECheckBoxState InCheckBoxState, int32 InBrushIndex);
diff --git a/Engine/Source/Runtime/Landscape/Private/LandscapeEditLayers.cpp b/Engine/Source/Runtime/Landscape/Private/LandscapeEditLayers.cpp
--- a/Engine/Source/Runtime/Landscape/Private/LandscapeEditLayers.cpp
+++ b/Engine/Source/Runtime/Landscape/Private/LandscapeEditLayers.cpp
@@ -13492,6 +13492,8 @@ void ALandscape::AddBrushToLayer(int32 InLayerIndex, ALandscapeBlueprintBrushBas
 	if (FLandscapeLayer* Layer = GetLayerInternal(InLayerIndex))
 	{
 		Modify();
+		// ensure the brush has the correct parent on creation
+		InBrush->SetOwner(this);
 		Layer->Brushes.Add(FLandscapeLayerBrush(InBrush));
 		InBrush->SetOwningLandscape(this);
 		RequestLayersContentUpdateForceAll();
diff --git a/Engine/Source/Runtime/Landscape/Public/LandscapeBlueprintBrushBase.h b/Engine/Source/Runtime/Landscape/Public/LandscapeBlueprintBrushBase.h
--- a/Engine/Source/Runtime/Landscape/Public/LandscapeBlueprintBrushBase.h
+++ b/Engine/Source/Runtime/Landscape/Public/LandscapeBlueprintBrushBase.h
@@ -37,7 +37,7 @@ struct FLandscapeBrushParameters
 };
 
 
-UCLASS(Abstract, NotBlueprintable, MinimalAPI)
+UCLASS(Abstract, NotPlaceable, NotBlueprintable, MinimalAPI)
 class ALandscapeBlueprintBrushBase : public AActor
 #if CPP && WITH_EDITOR // UHT doesn't support inheriting from namespaced class
 	, public ILandscapeEditLayerRenderer
@@ -145,6 +145,7 @@ public:
 	LANDSCAPE_API virtual bool AffectsWeightmap() const;
 	LANDSCAPE_API virtual bool AffectsWeightmapLayer(const FName& InLayerName) const;
 	LANDSCAPE_API virtual bool AffectsVisibilityLayer() const;
+	LANDSCAPE_API virtual bool SupportsBlueprintBrushTool() const { return true; }
 
 	bool IsVisible() const { return bIsVisible; }
 	LANDSCAPE_API bool IsLayerUpdatePending() const;
