Landscape Editor - Simple assert update to fix crash when Copy Tool gizmo does not have a valid landscape

diff --git a/Engine/Source/Editor/LandscapeEditor/Private/LandscapeEditorObject.cpp b/Engine/Source/Editor/LandscapeEditor/Private/LandscapeEditorObject.cpp
--- a/Engine/Source/Editor/LandscapeEditor/Private/LandscapeEditorObject.cpp
+++ b/Engine/Source/Editor/LandscapeEditor/Private/LandscapeEditorObject.cpp
@@ -516,10 +516,9 @@ void ULandscapeEditorObject::SetGizmoSnapMode(ELandscapeGizmoSnapType InSnapMode
 
 	if (SnapMode != ELandscapeGizmoSnapType::None)
 	{
-		if (ParentMode->CurrentGizmoActor.IsValid())
+		// Ignore gizmo updates if landscape info is invalid (Ex. landscape is deleted after using Copy Tool)
+		if (ParentMode->CurrentGizmoActor.IsValid() && ParentMode->CurrentGizmoActor->TargetLandscapeInfo)
 		{
-			check(ParentMode->CurrentGizmoActor->TargetLandscapeInfo);
-
 			const FVector WidgetLocation = ParentMode->CurrentGizmoActor->GetActorLocation();
 			const FRotator WidgetRotation = ParentMode->CurrentGizmoActor->GetActorRotation();
 
