Landscape Editor - Edit layer context menu: Option to Add Unassigned Brushes to edit layer
* Allows users to assign brushes that have an invalid landscape actor
* Adds to layer by index instead of selected edit layer only

diff --git a/Engine/Source/Editor/LandscapeEditor/Private/LandscapeEditorDetailCustomization_Layers.cpp b/Engine/Source/Editor/LandscapeEditor/Private/LandscapeEditorDetailCustomization_Layers.cpp
--- a/Engine/Source/Editor/LandscapeEditor/Private/LandscapeEditorDetailCustomization_Layers.cpp
+++ b/Engine/Source/Editor/LandscapeEditor/Private/LandscapeEditorDetailCustomization_Layers.cpp
@@ -613,6 +613,43 @@ TSharedPtr<SWidget> FLandscapeEditorCustomNodeBuilder_Layers::OnLayerContextMenu
 		}
 		MenuBuilder.EndSection();
 
+		// Add unassigned brush to layer
+		const TArray<ALandscapeBlueprintBrushBase*>& Brushes = LandscapeEdMode->GetBrushList();
+		TArray<ALandscapeBlueprintBrushBase*> FilteredBrushes = Brushes.FilterByPredicate([](ALandscapeBlueprintBrushBase* Brush) { return Brush->GetOwningLandscape() == nullptr; });
+		
+		MenuBuilder.BeginSection("LandscapeEditorBrushActions", LOCTEXT("LandscapeEditorBrushActions.Heading", "Unassigned Brushes"));
+		{
+			if (EditLayer)
+			{
+				// If there are no unassigned brushes or the edit layer does not support brushes, show a disabled state instead of hiding the entire section
+				if (FilteredBrushes.Num() > 0 && EditLayer->SupportsBlueprintBrushes())
+				{
+					MenuBuilder.AddSubMenu(
+						LOCTEXT("LandscapeEditorBrushAddSubMenu", "Assign Existing Brush"),
+						LOCTEXT("LandscapeEditorBrushAddSubMenuToolTip", "To modify the terrain, brushes need to be assigned to a landscape actor. Add the brush to this edit layer"),
+						FNewMenuDelegate::CreateSP(this, &FLandscapeEditorCustomNodeBuilder_Layers::FillUnassignedBrushMenu, FilteredBrushes, InLayerIndex),
+						false,
+						FSlateIcon()
+					);
+				}
+				else
+				{
+					FUIAction NoOpAction = FUIAction(FExecuteAction::CreateLambda([]() {}),
+						FCanExecuteAction::CreateLambda([]() { return false; }));
+					const FText NoOpToolTip = EditLayer->SupportsBlueprintBrushes() ? LOCTEXT("LandscapeEditorBrushNone", "All Blueprint Brush actors are assigned to a landscape edit layer")
+						: FText::Format(LOCTEXT("LandscapeEditorBrushUnsupported", "This layer's type ({0}) doesn't support blueprint brushes."), EditLayer->GetClass()->GetDisplayNameText());
+
+					MenuBuilder.AddMenuEntry(
+						LOCTEXT("LandscapeEditorBrushNone", "None"),
+						NoOpToolTip,
+						FSlateIcon(),
+						NoOpAction
+					);
+				}
+			}
+		}
+		MenuBuilder.EndSection();
+		
 		return MenuBuilder.MakeWidget();
 	}
 	return nullptr;
@@ -1294,6 +1331,27 @@ const FSlateBrush* FLandscapeEditorCustomNodeBuilder_Layers::GetLockBrushForLaye
 	return bIsLocked ? FAppStyle::GetBrush(TEXT("PropertyWindow.Locked")) : FAppStyle::GetBrush(TEXT("PropertyWindow.Unlocked"));
 }
 
+void FLandscapeEditorCustomNodeBuilder_Layers::FillUnassignedBrushMenu(FMenuBuilder& MenuBuilder, TArray<ALandscapeBlueprintBrushBase*> Brushes, int32 InLayerIndex)
+{
+	for (ALandscapeBlueprintBrushBase* Brush : Brushes)
+	{
+		FUIAction AddAction = FUIAction(FExecuteAction::CreateSP(this, &FLandscapeEditorCustomNodeBuilder_Layers::AssignBrushToEditLayer, Brush, InLayerIndex));
+		MenuBuilder.AddMenuEntry(FText::FromString(Brush->GetActorLabel()), FText(), FSlateIcon(), AddAction);
+	}
+}
+
+void FLandscapeEditorCustomNodeBuilder_Layers::AssignBrushToEditLayer(ALandscapeBlueprintBrushBase* Brush, const int32 InLayerIndex)
+{
+	const FScopedTransaction Transaction(LOCTEXT("LandscapeBrushAddToCurrentLayerTransaction", "Add brush to edit layer"));
+	if (FEdModeLandscape* LandscapeEdMode = GetEditorMode())
+	{
+		if (ALandscape* Landscape = LandscapeEdMode->GetLandscape())
+		{
+			Landscape->AddBrushToLayer(InLayerIndex, Brush);
+		}
+	}
+}
+
 const FSlateBrush* FLandscapeEditorCustomNodeBuilder_Layers::GetEditLayerIconBrush(int32 InLayerIndex) const
 {
 	FEdModeLandscape* LandscapeEdMode = GetEditorMode();
diff --git a/Engine/Source/Editor/LandscapeEditor/Private/LandscapeEditorDetailCustomization_Layers.h b/Engine/Source/Editor/LandscapeEditor/Private/LandscapeEditorDetailCustomization_Layers.h
--- a/Engine/Source/Editor/LandscapeEditor/Private/LandscapeEditorDetailCustomization_Layers.h
+++ b/Engine/Source/Editor/LandscapeEditor/Private/LandscapeEditorDetailCustomization_Layers.h
@@ -128,6 +128,10 @@ protected:
 	FReply OnToggleLock(int32 InLayerIndex);
 	const FSlateBrush* GetLockBrushForLayer(int32 InLayerIndex) const;
 
+	// Add Existing BP Brush (Brushes with no landscape actor assigned)
+	void FillUnassignedBrushMenu(FMenuBuilder& MenuBuilder, TArray<ALandscapeBlueprintBrushBase*> Brushes, int32 InLayerIndex);
+	void AssignBrushToEditLayer(ALandscapeBlueprintBrushBase* Brush, int32 InLayerIndex);
+
 	void FillClearTargetLayerMenu(FMenuBuilder& MenuBuilder, int32 InLayerIndex, TArray<ULandscapeLayerInfoObject*> InUsedLayerInfos);
 	void FillClearLayerMenu(FMenuBuilder& MenuBuilder, int32 InLayerIndex);
 
