Fixed bad ensure condition when trying to cache data from an empty area in the landscape smooth/flatten/erosion tools

diff --git a/Engine/Source/Editor/LandscapeEditor/Private/LandscapeEdModeTools.h b/Engine/Source/Editor/LandscapeEditor/Private/LandscapeEdModeTools.h
--- a/Engine/Source/Editor/LandscapeEditor/Private/LandscapeEdModeTools.h
+++ b/Engine/Source/Editor/LandscapeEditor/Private/LandscapeEdModeTools.h
@@ -501,7 +501,7 @@ public:
 
 			CacheData(NewCacheBounds.Min.X, NewCacheBounds.Min.Y, NewCacheBounds.Max.X, NewCacheBounds.Max.Y);
 		}
-		ensure(HasCachedData(X1, Y1, X2, Y2));
+		ensureMsgf(!bIsValid || HasCachedData(X1, Y1, X2, Y2), TEXT("Data failed to cache (X1:%i, Y1:%i, X2:%i, Y2:%i), (CachedX1:%i, CachedY1:%i, CachedX2:%i, CachedY2:%i)"), X1, Y1, X2, Y2, CachedX1, CachedY1, CachedX2, CachedY2);
 		return GetCachedData(X1, Y1, X2, Y2, OutData);
 	}
 
