Landscape Spline - Moving landscape spline selection and navigation into dedicated UClass
* Enables transactional spline selection for points and segments
* Simplifies spline navigation logic by removing LastSelected logic
* Removes fixup selection code on Undo

diff --git a/Engine/Source/Editor/LandscapeEditor/Private/LandscapeEdMode.cpp b/Engine/Source/Editor/LandscapeEditor/Private/LandscapeEdMode.cpp
--- a/Engine/Source/Editor/LandscapeEditor/Private/LandscapeEdMode.cpp
+++ b/Engine/Source/Editor/LandscapeEditor/Private/LandscapeEdMode.cpp
@@ -53,6 +53,7 @@
 // Classes
 #include "LandscapeMaterialInstanceConstant.h"
 #include "LandscapeSplinesComponent.h"
+#include "LandscapeSplineSelection.h"
 #include "ComponentReregisterContext.h"
 #include "EngineUtils.h"
 #include "Misc/ScopedSlowTask.h"
diff --git a/Engine/Source/Editor/LandscapeEditor/Private/LandscapeEdMode.h b/Engine/Source/Editor/LandscapeEditor/Private/LandscapeEdMode.h
--- a/Engine/Source/Editor/LandscapeEditor/Private/LandscapeEdMode.h
+++ b/Engine/Source/Editor/LandscapeEditor/Private/LandscapeEdMode.h
@@ -33,6 +33,7 @@ struct FHeightmapToolTarget;
 struct FViewportActionKeyInput;
 struct FViewportClick;
 template<class ToolTarget> class FLandscapeToolCopyPaste;
+enum class ESplineNavigationFlags : uint8;
 
 DECLARE_LOG_CATEGORY_EXTERN(LogLandscapeEdMode, Log, All);
 
@@ -275,22 +276,6 @@ enum class ELandscapeEditingState : uint8
 	NoLandscape,
 };
 
-enum class ESplineNavigationFlags : uint8
-{
-	None = 0,
-	DirectionForward = 1 << 0,
-	DirectionBackward = 1 << 1,
-
-	SegmentSelectModeEnabled = 1 << 2,
-	ControlPointSelectModeEnabled = 1 << 3,
-
-	AddToSelection = 1 << 4,
-
-	SelectModeMask = SegmentSelectModeEnabled | ControlPointSelectModeEnabled,
-	DirectionMask = DirectionForward | DirectionBackward
-};
-ENUM_CLASS_FLAGS(ESplineNavigationFlags);
-
 /**
  * Landscape editor mode
  */
diff --git a/Engine/Source/Editor/LandscapeEditor/Private/LandscapeEdModeSplineTools.cpp b/Engine/Source/Editor/LandscapeEditor/Private/LandscapeEdModeSplineTools.cpp
--- a/Engine/Source/Editor/LandscapeEditor/Private/LandscapeEdModeSplineTools.cpp
+++ b/Engine/Source/Editor/LandscapeEditor/Private/LandscapeEdModeSplineTools.cpp
@@ -31,6 +31,7 @@
 #include "LandscapeSplinesComponent.h"
 #include "LandscapeSplineSegment.h"
 #include "LandscapeSplineControlPoint.h"
+#include "LandscapeSplineSelection.h"
 #include "ControlPointMeshComponent.h"
 #include "Algo/AnyOf.h"
 #include "Algo/Copy.h"
@@ -50,13 +51,7 @@ public:
 	FLandscapeToolSplines(FEdModeLandscape* InEdMode)
 		: EdMode(InEdMode)
 		, LandscapeInfo(nullptr)
-		, SelectedSplineControlPoints()
-		, SelectedSplineSegments()
-		, LastSelectedControlPoint(nullptr)
-		, LastSelectedSegment(nullptr)
-		, LinearControlPoints()
-		, LinearSegments()
-		, ControlPointWidgetMode()
+		, SplineSelection(nullptr)
 		, DraggingTangent_Segment(nullptr)
 		, DraggingTangent_Length(0.0f)
 		, DraggingTangent_CacheCoordSpace(ECoordSystem::COORD_None)
@@ -85,13 +80,7 @@ public:
 	virtual void AddReferencedObjects(FReferenceCollector& Collector) override
 	{
 		Collector.AddReferencedObject(LandscapeInfo);
-		Collector.AddReferencedObjects(SelectedSplineControlPoints);
-		Collector.AddReferencedObjects(SelectedSplineSegments);
-		// TODO [jared.ritchie] : this needs to be done, but then also they need to be cleared when exiting the tool. And we need to make sure they're kept in sync with SelectedSplineControlPoints/SelectedSplineSegments when these change ...
-		//Collector.AddReferencedObject(LastSelectedControlPoint);
-		//Collector.AddReferencedObject(LastSelectedSegment);
-		//Collector.AddReferencedObjects(LinearControlPoints);
-		//Collector.AddReferencedObjects(LinearSegments);
+		Collector.AddReferencedObject(SplineSelection);
 		Collector.AddReferencedObject(DraggingTangent_Segment);
 	}
 
@@ -112,435 +101,6 @@ public:
 	virtual void SetEditRenderType() override { GLandscapeEditRenderMode = ELandscapeEditRenderMode::None | (GLandscapeEditRenderMode & ELandscapeEditRenderMode::BitMaskForMask); }
 	virtual bool SupportsMask() override { return false; }
 
-	void UpdatePropertiesWindows()
-	{
-		if (GLevelEditorModeTools().IsModeActive(EdMode->GetID()))
-		{
-			TArray<UObject*> Objects;
-			Objects.Reset(SelectedSplineControlPoints.Num() + SelectedSplineSegments.Num());
-			Algo::Copy(SelectedSplineControlPoints, Objects);
-			Algo::Copy(SelectedSplineSegments, Objects);
-
-			FPropertyEditorModule& PropertyModule = FModuleManager::Get().LoadModuleChecked<FPropertyEditorModule>(TEXT("PropertyEditor"));
-			PropertyModule.UpdatePropertyViews(Objects);
-		}
-	}
-
-	void ClearSelectedControlPoints()
-	{
-		for (ULandscapeSplineControlPoint* ControlPoint : SelectedSplineControlPoints)
-		{
-			checkSlow(ControlPoint->IsSplineSelected());
-			ControlPoint->Modify(false);
-			ControlPoint->SetSplineSelected(false);
-		}
-		SelectedSplineControlPoints.Empty();
-	}
-
-	void ClearSelectedSegments()
-	{
-		for (ULandscapeSplineSegment* Segment : SelectedSplineSegments)
-		{
-			checkSlow(Segment->IsSplineSelected());
-			Segment->Modify(false);
-			Segment->SetSplineSelected(false);
-		}
-		SelectedSplineSegments.Empty();
-	}
-
-	void ClearSelection()
-	{
-		ClearSelectedControlPoints();
-		ClearSelectedSegments();
-	}
-
-	void DeselectControlPoint(ULandscapeSplineControlPoint* ControlPoint)
-	{
-		checkSlow(ControlPoint->IsSplineSelected());
-		SelectedSplineControlPoints.Remove(ControlPoint);
-		ControlPoint->Modify(false);
-		ControlPoint->SetSplineSelected(false);
-	}
-
-	void DeSelectSegment(ULandscapeSplineSegment* Segment)
-	{
-		checkSlow(Segment->IsSplineSelected());
-		SelectedSplineSegments.Remove(Segment);
-		Segment->Modify(false);
-		Segment->SetSplineSelected(false);
-	}
-
-	void SelectControlPoint(ULandscapeSplineControlPoint* ControlPoint)
-	{
-		SelectNavigationControlPoint(ControlPoint);
-
-		checkSlow(!ControlPoint->IsSplineSelected());
-
-		// When switching from segment mode, restore the last used control point mode
-		if (SelectedSplineSegments.Num() > 0 && ControlPointWidgetMode.IsSet())
-		{
-			GLevelEditorModeTools().SetWidgetMode(ControlPointWidgetMode.GetValue());
-		}
-		SelectedSplineControlPoints.Add(ControlPoint);
-		ControlPoint->Modify(false);
-		ControlPoint->SetSplineSelected(true);
-	}
-
-	void SelectSegment(ULandscapeSplineSegment* Segment)
-	{
-		SelectNavigationSegment(Segment);
-
-		checkSlow(!Segment->IsSplineSelected());
-
-		// Before entering segment mode, save the control point mode for persistence
-		if (SelectedSplineSegments.Num() == 0)
-		{
-			ControlPointWidgetMode = GLevelEditorModeTools().GetWidgetMode();
-		}
-		SelectedSplineSegments.Add(Segment);
-		Segment->Modify(false);
-		Segment->SetSplineSelected(true);
-
-		GLevelEditorModeTools().SetWidgetMode(UE::Widget::WM_Scale);
-	}
-
-	void SelectConnected()
-	{
-		TArray<ULandscapeSplineControlPoint*> ControlPointsToProcess =
-			ObjectPtrDecay(SelectedSplineControlPoints.Array());
-
-		while (ControlPointsToProcess.Num() > 0)
-		{
-			const ULandscapeSplineControlPoint* ControlPoint = ControlPointsToProcess.Pop();
-
-			for (const FLandscapeSplineConnection& Connection : ControlPoint->ConnectedSegments)
-			{
-				ULandscapeSplineControlPoint* OtherEnd = Connection.GetFarConnection().ControlPoint;
-
-				if (!OtherEnd->IsSplineSelected())
-				{
-					SelectControlPoint(OtherEnd);
-					ControlPointsToProcess.Add(OtherEnd);
-				}
-			}
-		}
-
-		TArray<ULandscapeSplineSegment*> SegmentsToProcess = ObjectPtrDecay(SelectedSplineSegments.Array());
-
-		while (SegmentsToProcess.Num() > 0)
-		{
-			const ULandscapeSplineSegment* Segment = SegmentsToProcess.Pop();
-
-			for (const FLandscapeSplineSegmentConnection& SegmentConnection : Segment->Connections)
-			{
-				for (const FLandscapeSplineConnection& Connection : SegmentConnection.ControlPoint->ConnectedSegments)
-				{
-					if (Connection.Segment != Segment && !Connection.Segment->IsSplineSelected())
-					{
-						SelectSegment(Connection.Segment);
-						SegmentsToProcess.Add(Connection.Segment);
-					}
-				}
-			}
-		}
-	}
-
-	void SelectAllSplineSegments(const ULandscapeInfo& InLandscapeInfo)
-	{
-		TArray<TScriptInterface<ILandscapeSplineInterface>> SplineActors(InLandscapeInfo.GetSplineActors());
-		for (TScriptInterface<ILandscapeSplineInterface> SplineActor : SplineActors)
-		{
-			if (ULandscapeSplinesComponent* SplineComponent = SplineActor->GetSplinesComponent())
-			{
-				SplineComponent->ForEachControlPoint([this](ULandscapeSplineControlPoint* ControlPoint) 
-				{ 
-					for (const FLandscapeSplineConnection& Connection : ControlPoint->ConnectedSegments)
-					{
-						if (!Connection.Segment->IsSplineSelected())
-						{
-							SelectSegment(Connection.Segment);
-						}
-					}
-				});
-			}
-		}
-	}
-
-	void SelectAllControlPoints(const ULandscapeInfo& InLandscapeInfo)
-	{
-		TArray<TScriptInterface<ILandscapeSplineInterface>> SplineActors(InLandscapeInfo.GetSplineActors());
-		for (TScriptInterface<ILandscapeSplineInterface> SplineActor : SplineActors)
-		{
-			if (ULandscapeSplinesComponent* SplineComponent = SplineActor->GetSplinesComponent())
-			{
-				SplineComponent->ForEachControlPoint([this](ULandscapeSplineControlPoint* ControlPoint) { SelectControlPoint(ControlPoint); });
-			}
-		}
-	}
-
-	void SelectAdjacentControlPoints()
-	{
-		for (ULandscapeSplineSegment* Segment : SelectedSplineSegments)
-		{
-			if (!Segment->Connections[0].ControlPoint->IsSplineSelected())
-			{
-				SelectControlPoint(Segment->Connections[0].ControlPoint);
-			}
-			if (!Segment->Connections[1].ControlPoint->IsSplineSelected())
-			{
-				SelectControlPoint(Segment->Connections[1].ControlPoint);
-			}
-		}
-	}
-
-	void SelectAdjacentSegments()
-	{
-		for (ULandscapeSplineControlPoint* ControlPoint : SelectedSplineControlPoints)
-		{
-			for (const FLandscapeSplineConnection& Connection : ControlPoint->ConnectedSegments)
-			{
-				if (!Connection.Segment->IsSplineSelected())
-				{
-					SelectSegment(Connection.Segment);
-				}
-			}
-		}
-	}
-
-	ULandscapeSplineControlPoint* GetLinearEndControlPoint(ULandscapeSplineControlPoint* SelectedControlPoint)
-	{
-		TSet<ULandscapeSplineControlPoint*> ControlPointsVisited;
-		TArray<ULandscapeSplineControlPoint*> ControlPointsToProcess;
-		ControlPointsToProcess.Add(SelectedControlPoint);
-
-		// Reset the linear path before it gets rebuilt below
-		LinearControlPoints.Empty();
-
-		// Control points with zero/multi node connections are not part of a valid linear path
-		if (SelectedControlPoint->ConnectedSegments.IsEmpty() || SelectedControlPoint->ConnectedSegments.Num() > 2)
-		{
-			LinearControlPoints.Add(SelectedControlPoint);
-			return SelectedControlPoint;
-		}
-
-		while (ControlPointsToProcess.Num() > 0)
-		{
-			ULandscapeSplineControlPoint* ControlPoint = ControlPointsToProcess.Pop();
-			ControlPointsVisited.Add(ControlPoint);
-
-			// Store nodes in order of visiting to create the linear path
-			if (!LinearControlPoints.Contains(ControlPoint))
-			{
-				LinearControlPoints.Add(ControlPoint);
-			}
-
-			// Valid end points only have one connection
-			if (ControlPoint->ConnectedSegments.Num() == 1 && ControlPoint != SelectedControlPoint)
-			{
-				return ControlPoint;
-			}
-
-			for (const FLandscapeSplineConnection& Connection : ControlPoint->ConnectedSegments)
-			{
-				ULandscapeSplineControlPoint* OtherEnd = Connection.GetFarConnection().ControlPoint;
-				check(OtherEnd != nullptr);
-
-				if (OtherEnd->ConnectedSegments.IsEmpty() || OtherEnd->ConnectedSegments.Num() > 2)
-				{
-					continue;
-				}
-
-				if (!ControlPointsVisited.Contains(OtherEnd))
-				{
-					ControlPointsToProcess.Add(OtherEnd);
-				}
-			}
-		}
-
-		return SelectedControlPoint;
-	}
-
-	void ResetNavigationPath()
-	{
-		// Clear cached path
-		LinearControlPoints.Empty();
-		LinearSegments.Empty();
-
-		LastSelectedControlPoint = nullptr;
-		LastSelectedSegment = nullptr;
-	}
-
-	void BuildLinearPathFromLastSelectedPoint()
-	{
-		LinearControlPoints.Empty();
-		LinearSegments.Empty();
-
-		ULandscapeSplineControlPoint* StartControlPoint = GetLinearEndControlPoint(LastSelectedControlPoint);
-		ULandscapeSplineControlPoint* EndControlPoint = GetLinearEndControlPoint(StartControlPoint);
-
-		for (const ULandscapeSplineControlPoint* Point : LinearControlPoints)
-		{
-			for (const FLandscapeSplineConnection& Connection : Point->ConnectedSegments)
-			{
-				if (!LinearSegments.Contains(Connection.Segment))
-				{
-					LinearSegments.Add(Connection.Segment);
-				}
-			}
-		}
-	}
-
-	void SelectNavigationSegment(ULandscapeSplineSegment* Segment)
-	{
-		LastSelectedSegment = Segment;
-		LastSelectedControlPoint = Segment->Connections[0].ControlPoint;
-
-		if (!LinearSegments.Contains(Segment))
-		{
-			BuildLinearPathFromLastSelectedPoint();
-		}
-	}
-
-	void SelectNavigationControlPoint(ULandscapeSplineControlPoint* ControlPoint)
-	{
-		LastSelectedControlPoint = ControlPoint;
-
-		if (!LinearControlPoints.Contains(ControlPoint))
-		{
-			BuildLinearPathFromLastSelectedPoint();
-		}
-	}
-
-	bool IsSelectionValidForNavigation() const
-	{
-		// If any elements in selection are part of different linear paths, disable navigation
-		for (ULandscapeSplineControlPoint* ControlPoint : SelectedSplineControlPoints)
-		{
-			if (!LinearControlPoints.Contains(ControlPoint))
-			{
-				return false;
-			}
-		}
-
-		for (ULandscapeSplineSegment* Segment : SelectedSplineSegments)
-		{
-			if (!LinearSegments.Contains(Segment))
-			{
-				return false;
-			}
-		}
-
-		return true;
-	}
-
-	ULandscapeSplineSegment* GetEndSegmentInLinearPath(ESplineNavigationFlags Flags) const
-	{
-		check(EnumHasAllFlags(Flags, ESplineNavigationFlags::SegmentSelectModeEnabled));
-		check(FMath::CountBits(static_cast<uint8>(Flags & ESplineNavigationFlags::DirectionMask)) == 1);
-
-		if (LinearSegments.IsEmpty())
-		{
-			return nullptr;
-		}
-
-		if (EnumHasAllFlags(Flags, ESplineNavigationFlags::DirectionForward))
-		{
-			return LinearSegments.Top();
-		}
-
-		return LinearSegments[0];
-	}
-
-	ULandscapeSplineControlPoint* GetEndControlPointInLinearPath(ESplineNavigationFlags Flags) const
-	{
-		check(EnumHasAllFlags(Flags, ESplineNavigationFlags::ControlPointSelectModeEnabled));
-		check(FMath::CountBits(static_cast<uint8>(Flags & ESplineNavigationFlags::SelectModeMask)) == 1);
-
-		if (LinearControlPoints.IsEmpty())
-		{
-			return nullptr;
-		}
-
-		if (EnumHasAllFlags(Flags, ESplineNavigationFlags::DirectionForward))
-		{
-			return LinearControlPoints.Top();
-		}
-
-		return LinearControlPoints[0];
-	}
-
-	bool HasAdjacentSegmentInLinearPath(ESplineNavigationFlags Flags) const
-	{
-		check(EnumHasAllFlags(Flags, ESplineNavigationFlags::SegmentSelectModeEnabled));
-		check(FMath::CountBits(static_cast<uint8>(Flags & ESplineNavigationFlags::SelectModeMask)) == 1);
-
-		const ULandscapeSplineSegment* AdjacentSegment = GetAdjacentSegmentInLinearPath(Flags);
-
-		if (!AdjacentSegment || !IsSelectionValidForNavigation())
-		{
-			return false;
-		}
-
-		if (EnumHasAllFlags(Flags, ESplineNavigationFlags::AddToSelection))
-		{
-			return !AdjacentSegment->IsSplineSelected();
-		}
-
-		return true;
-	}
-
-	bool HasAdjacentControlPointInLinearPath(ESplineNavigationFlags Flags) const
-	{
-		check(EnumHasAllFlags(Flags, ESplineNavigationFlags::ControlPointSelectModeEnabled));
-		check(FMath::CountBits(static_cast<uint8>(Flags & ESplineNavigationFlags::SelectModeMask)) == 1);
-
-		const ULandscapeSplineControlPoint* AdjacentPoint = GetAdjacentControlPointInPath(Flags);
-
-		if (!AdjacentPoint || !IsSelectionValidForNavigation())
-		{
-			return false;
-		}
-
-		if (EnumHasAllFlags(Flags, ESplineNavigationFlags::AddToSelection))
-		{
-			return !AdjacentPoint->IsSplineSelected();
-		}
-
-		return true;
-	}
-
-	ULandscapeSplineSegment* GetAdjacentSegmentInLinearPath(ESplineNavigationFlags Flags) const
-	{
-		check(EnumHasAllFlags(Flags, ESplineNavigationFlags::SegmentSelectModeEnabled));
-		check(FMath::CountBits(static_cast<uint8>(Flags & ESplineNavigationFlags::SelectModeMask)) == 1);
-
-		const int32 LastSelectedSegmentIndex = LinearSegments.Find(LastSelectedSegment);
-		const int32 NextIndexOffset = EnumHasAllFlags(Flags, ESplineNavigationFlags::DirectionForward) ? 1 : -1;
-		check(LastSelectedSegmentIndex != INDEX_NONE); // last selected should always be part of the linear path
-
-		if (LinearSegments.IsValidIndex(LastSelectedSegmentIndex + NextIndexOffset))
-		{
-			return LinearSegments[LastSelectedSegmentIndex + NextIndexOffset];
-		}
-		return nullptr;
-	}
-
-	ULandscapeSplineControlPoint* GetAdjacentControlPointInPath(ESplineNavigationFlags Flags) const
-	{
-		check(EnumHasAllFlags(Flags, ESplineNavigationFlags::ControlPointSelectModeEnabled));
-		check(FMath::CountBits(static_cast<uint8>(Flags & ESplineNavigationFlags::SelectModeMask)) == 1);
-
-		const int32 LastSelectedControlPointIndex = LinearControlPoints.Find(LastSelectedControlPoint);
-		const int32 NextIndexOffset = EnumHasAllFlags(Flags, ESplineNavigationFlags::DirectionForward) ? 1 : -1;
-		check(LastSelectedControlPointIndex != INDEX_NONE); // last selected should always be part of the linear path
-
-		if (LinearControlPoints.IsValidIndex(LastSelectedControlPointIndex + NextIndexOffset))
-		{
-			return LinearControlPoints[LastSelectedControlPointIndex + NextIndexOffset];
-		}
-		return nullptr;
-	}
-
 	void AddSegment(ULandscapeSplineControlPoint* Start, ULandscapeSplineControlPoint* End, bool bAutoRotateStart, bool bAutoRotateEnd)
 	{
 		FScopedTransaction Transaction(LOCTEXT("LandscapeSpline_AddSegment", "Add Landscape Spline Segment"));
@@ -662,14 +222,13 @@ public:
 		}
 
 		// Adding a segment will change the linear navigation path, reset it
-		ResetNavigationPath();
-		SelectNavigationControlPoint(Start);
+		SplineSelection->ResetNavigationPath();
+		SplineSelection->SelectNavigationControlPoint(Start);
 	}
 
 	void FlipSelectedSplineSegments()
 	{
-		FScopedTransaction Transaction(LOCTEXT("LandscapeSpline_FlipSegments", "Flip Selected Landscape Spline Segments"));
-		for (ULandscapeSplineSegment* Segment : SelectedSplineSegments)
+		for (ULandscapeSplineSegment* Segment : SplineSelection->GetSelectedSplineSegments())
 		{
 			FlipSegment(Segment);
 		}
@@ -742,8 +301,8 @@ public:
 		SplinesComponent->MarkRenderStateDirty();
 
 		// Adding a segment will change the linear navigation path, reset it
-		ResetNavigationPath();
-		SelectNavigationControlPoint(Start);
+		SplineSelection->ResetNavigationPath();
+		SplineSelection->SelectNavigationControlPoint(Start);
 
 		return true;
 	}
@@ -759,6 +318,8 @@ public:
 
 		NewControlPoint->Location = LocalLocation;
 
+		const TArray<ULandscapeSplineControlPoint*> SelectedSplineControlPoints = SplineSelection->GetSelectedSplineControlPoints();
+
 		if (SelectedSplineControlPoints.Num() > 0)
 		{
 			ULandscapeSplineControlPoint* FirstPoint = *SelectedSplineControlPoints.CreateConstIterator();
@@ -809,10 +370,8 @@ public:
 			NewControlPoint->UpdateSplinePoints();
 		}
 
-		ResetNavigationPath();
-		ClearSelection();
-		SelectControlPoint(NewControlPoint);
-		UpdatePropertiesWindows();
+		SplineSelection->ResetNavigationPath();
+		SplineSelection->SelectControlPoint(NewControlPoint, ESplineNavigationFlags::UpdatePropertiesWindows);
 
 		EdMode->AutoUpdateDirtyLandscapeSplines();
 		if (!SplinesComponent->IsRegistered())
@@ -854,7 +413,7 @@ public:
 		}
 
 		SplinesComponent->Segments.Remove(ToDelete);
-		ResetNavigationPath();
+		SplineSelection->ResetNavigationPath();
 
 		ToDelete->Connections[0].ControlPoint->UpdateSplinePoints();
 		ToDelete->Connections[1].ControlPoint->UpdateSplinePoints();
@@ -872,7 +431,7 @@ public:
 
 		ToDelete->Modify();
 		ToDelete->DeleteSplinePoints();
-		ResetNavigationPath();
+		SplineSelection->ResetNavigationPath();
 
 		if (ToDelete->ConnectedSegments.Num() == 2
 			&& ToDelete->ConnectedSegments[0].Segment != ToDelete->ConnectedSegments[1].Segment)
@@ -1028,8 +587,7 @@ public:
 		Segment->UpdateSplinePoints();
 		NewSegment->UpdateSplinePoints();
 
-		ClearSelection();
-		UpdatePropertiesWindows();
+		SplineSelection->ClearSelection();
 
 		SplinesComponent->MarkRenderStateDirty();
 	}
@@ -1149,7 +707,7 @@ public:
 
 	void UpdateSplineMeshLevels()
 	{
-		for (ULandscapeSplineControlPoint* ControlPoint : SelectedSplineControlPoints)
+		for (ULandscapeSplineControlPoint* ControlPoint : SplineSelection->GetSelectedSplineControlPoints())
 		{
 			const bool bUpdateCollision = true;
 			const bool bUpdateSegments = false;
@@ -1157,7 +715,7 @@ public:
 			ControlPoint->UpdateSplinePoints(bUpdateCollision, bUpdateSegments, bUpdateMeshLevel);
 		}
 
-		for (ULandscapeSplineSegment* Segment : SelectedSplineSegments)
+		for (ULandscapeSplineSegment* Segment : SplineSelection->GetSelectedSplineSegments())
 		{
 			const bool bUpdateCollision = true;
 			const bool bUpdateMeshLevel = true;
@@ -1168,7 +726,7 @@ public:
 	bool CanMoveSelectedToLevel() const
 	{
 		// Move to level only supportd on LandscapeProxy Splines
-		for (ULandscapeSplineControlPoint* ControlPoint : SelectedSplineControlPoints)
+		for (ULandscapeSplineControlPoint* ControlPoint : SplineSelection->GetSelectedSplineControlPoints())
 		{
 			ULandscapeSplinesComponent* LandscapeSplinesComp = ControlPoint->GetOuterULandscapeSplinesComponent();
 			ALandscapeProxy* FromProxy = LandscapeSplinesComp ? Cast<ALandscapeProxy>(LandscapeSplinesComp->GetOuter()) : nullptr;
@@ -1186,7 +744,7 @@ public:
 		TSet<ALandscapeProxy*> FromProxies;
 		ALandscapeProxy* ToLandscape = nullptr;
 
-		for (ULandscapeSplineControlPoint* ControlPoint : SelectedSplineControlPoints)
+		for (ULandscapeSplineControlPoint* ControlPoint : SplineSelection->GetSelectedSplineControlPoints())
 		{
 			ULandscapeSplinesComponent* LandscapeSplinesComp = ControlPoint->GetOuterULandscapeSplinesComponent();
 			ALandscapeProxy* FromProxy = LandscapeSplinesComp ? Cast<ALandscapeProxy>(LandscapeSplinesComp->GetOuter()) : nullptr;
@@ -1227,9 +785,9 @@ public:
 	void ShowSplineProperties()
 	{
 		TArray<UObject*> Objects;
-		Objects.Reset(SelectedSplineControlPoints.Num() + SelectedSplineSegments.Num());
-		Algo::Copy(SelectedSplineControlPoints, Objects);
-		Algo::Copy(SelectedSplineSegments, Objects);
+		Objects.Reset(SplineSelection->GetSelectedSplineControlPoints().Num() + SplineSelection->GetSelectedSplineSegments().Num());
+		Algo::Copy(SplineSelection->GetSelectedSplineControlPoints(), Objects);
+		Algo::Copy(SplineSelection->GetSelectedSplineSegments(), Objects);
 
 		FPropertyEditorModule& PropertyModule = FModuleManager::Get().LoadModuleChecked<FPropertyEditorModule>(TEXT("PropertyEditor"));
 		if (!PropertyModule.HasUnlockedDetailViews())
@@ -1251,6 +809,8 @@ public:
 			ILandscapeSplineInterface* SplineOwner = nullptr;
 
 			// If we have a selection use the landscape of the selected spline
+			const TArray<ULandscapeSplineControlPoint*> SelectedSplineControlPoints = SplineSelection->GetSelectedSplineControlPoints();
+
 			if (SelectedSplineControlPoints.Num() > 0)
 			{
 				ULandscapeSplineControlPoint* FirstPoint = *SelectedSplineControlPoints.CreateConstIterator();
@@ -1361,26 +921,19 @@ public:
 			if (NewLandscapeActor && (NewLandscapeActor != EdMode->GetLandscape()))
 			{
 				EdMode->SetTargetLandscape(LandscapeProxy->GetLandscapeInfo());
-				ClearSelection();
 			}
 		}
 	}
 
 	virtual bool HandleClick(HHitProxy* HitProxy, const FViewportClick& Click) override
 	{
-		if ((!HitProxy || !HitProxy->IsA(HWidgetAxis::StaticGetType()))
-			&& !Click.IsShiftDown())
-		{
-			ClearSelection();
-			UpdatePropertiesWindows();
-			GUnrealEd->RedrawLevelEditingViewports();
-		}
+		ULandscapeSplineControlPoint* ClickedControlPoint = nullptr;
+		ULandscapeSplineSegment* ClickedSplineSegment = nullptr;
+
+		bool bIsValidSplineHitProxy = false;
 
 		if (HitProxy)
 		{
-			ULandscapeSplineControlPoint* ClickedControlPoint = nullptr;
-			ULandscapeSplineSegment* ClickedSplineSegment = nullptr;
-
 			if (HitProxy->IsA(HLandscapeSplineProxy_ControlPoint::StaticGetType()))
 			{
 				HLandscapeSplineProxy_ControlPoint* SplineProxy = (HLandscapeSplineProxy_ControlPoint*)HitProxy;
@@ -1391,6 +944,10 @@ public:
 				HLandscapeSplineProxy_Segment* SplineProxy = (HLandscapeSplineProxy_Segment*)HitProxy;
 				ClickedSplineSegment = SplineProxy->SplineSegment;
 			}
+			else if (HitProxy->IsA(HWidgetAxis::StaticGetType()) || HitProxy->IsA(HLandscapeSplineProxy_Tangent::StaticGetType()))
+			{
+				bIsValidSplineHitProxy = true;
+			}
 			else if (HitProxy->IsA(HActor::StaticGetType()))
 			{
 				HActor* ActorProxy = (HActor*)HitProxy;
@@ -1416,49 +973,60 @@ public:
 					}
 				}
 			}
+		}
+		
+		bIsValidSplineHitProxy = bIsValidSplineHitProxy || ClickedSplineSegment || ClickedControlPoint;
+
+		if (!HitProxy || !bIsValidSplineHitProxy)
+		{
+			SplineSelection->ClearSelection();
+			return false;
+		}
 
-			if (ClickedControlPoint != nullptr)
+		if (ClickedControlPoint != nullptr)
+		{
+			if (Click.IsShiftDown() && ClickedControlPoint->IsSplineSelected())
 			{
-				if (Click.IsShiftDown() && ClickedControlPoint->IsSplineSelected())
-				{
-					DeselectControlPoint(ClickedControlPoint);
-				}
-				else
-				{
-					SetTargetLandscapeBasedOnSelection(ClickedControlPoint);
-					SelectControlPoint(ClickedControlPoint);
-				}
+				FScopedTransaction Transaction(LOCTEXT("LandscapeSpline_DeselectPoint", "Deselect Landscape Spline Point"));
+				SplineSelection->DeselectControlPoint(ClickedControlPoint, ESplineNavigationFlags::UpdatePropertiesWindows);
 				GEditor->SelectNone(true, true);
-				UpdatePropertiesWindows();
-
-				GUnrealEd->RedrawLevelEditingViewports();
-				return true;
 			}
-			else if (ClickedSplineSegment != nullptr)
+			else
 			{
-				// save info about what we grabbed
-				if (Click.IsShiftDown() && ClickedSplineSegment->IsSplineSelected())
-				{
-					DeSelectSegment(ClickedSplineSegment);
-				}
-				else
-				{
-					SetTargetLandscapeBasedOnSelection(ClickedSplineSegment);
-					SelectSegment(ClickedSplineSegment);
-				}
+				FScopedTransaction Transaction(LOCTEXT("LandscapeSpline_SelectPoint", "Select Landscape Spline Point"));
+				SetTargetLandscapeBasedOnSelection(ClickedControlPoint);
+				SplineSelection->SelectControlPoint(ClickedControlPoint, ESplineNavigationFlags::UpdatePropertiesWindows | (Click.IsShiftDown() ? ESplineNavigationFlags::AddToSelection : ESplineNavigationFlags::None));
 				GEditor->SelectNone(true, true);
-				UpdatePropertiesWindows();
-
-				GUnrealEd->RedrawLevelEditingViewports();
-				return true;
 			}
+			return true;
 		}
-
+		else if (ClickedSplineSegment != nullptr)
+		{
+			// save info about what we grabbed
+			if (Click.IsShiftDown() && ClickedSplineSegment->IsSplineSelected())
+			{
+				FScopedTransaction Transaction(LOCTEXT("LandscapeSpline_DeselectSegment", "Deselect Landscape Spline Segment"));
+				SplineSelection->DeSelectSegment(ClickedSplineSegment, ESplineNavigationFlags::UpdatePropertiesWindows);
+				GEditor->SelectNone(true, true);
+			}
+			else
+			{
+				FScopedTransaction Transaction(LOCTEXT("LandscapeSpline_SelectSegment", "Select Landscape Spline Segment"));
+				SetTargetLandscapeBasedOnSelection(ClickedSplineSegment);
+				SplineSelection->SelectSegment(ClickedSplineSegment, ESplineNavigationFlags::UpdatePropertiesWindows | (Click.IsShiftDown() ? ESplineNavigationFlags::AddToSelection : ESplineNavigationFlags::None));
+				GEditor->SelectNone(true, true);
+			}
+			return true;
+		}
+		
 		return false;
 	}
 
 	virtual bool InputKey(FEditorViewportClient* InViewportClient, FViewport* InViewport, FKey InKey, EInputEvent InEvent) override
 	{
+		TArray<ULandscapeSplineControlPoint*> SelectedSplineControlPoints = SplineSelection->GetSelectedSplineControlPoints();
+		TArray<ULandscapeSplineSegment*> SelectedSplineSegments = SplineSelection->GetSelectedSplineSegments();
+
 		if (InKey == EKeys::F4 && InEvent == IE_Pressed)
 		{
 			if (SelectedSplineControlPoints.Num() > 0 || SelectedSplineSegments.Num() > 0)
@@ -1497,6 +1065,7 @@ public:
 		{
 			if (SelectedSplineSegments.Num() > 0)
 			{
+				FScopedTransaction Transaction(LOCTEXT("LandscapeSpline_FlipSegments", "Flip Selected Landscape Spline Segments"));
 				FlipSelectedSplineSegments();
 				return true;
 			}
@@ -1542,10 +1111,8 @@ public:
 					SnapControlPointToGround(Segment->Connections[0].ControlPoint);
 					SnapControlPointToGround(Segment->Connections[1].ControlPoint);
 				}
-				UpdatePropertiesWindows();
-
+				SplineSelection->UpdatePropertiesWindows();
 				EdMode->AutoUpdateDirtyLandscapeSplines();
-				GUnrealEd->RedrawLevelEditingViewports();
 				return true;
 			}
 		}
@@ -1555,11 +1122,9 @@ public:
 		{
 			if (SelectedSplineControlPoints.Num() > 0 || SelectedSplineSegments.Num() > 0)
 			{
-				SelectConnected();
-
-				UpdatePropertiesWindows();
-
-				GUnrealEd->RedrawLevelEditingViewports();
+				FScopedTransaction Transaction(LOCTEXT("LandscapeSpline_SelectConnectedSegments", "Select Landscape Spline Connected Segments"));
+				SplineSelection->SelectConnected();
+				SplineSelection->UpdatePropertiesWindows();
 				return true;
 			}
 		}
@@ -1853,7 +1418,9 @@ public:
 			return false;
 		}
 
-		ULandscapeSplineControlPoint* SelectedControlPoint = *SelectedSplineControlPoints.CreateConstIterator();
+		FScopedTransaction Transaction(LOCTEXT("LandscapeSpline_DuplicatePoint", "Duplicate Landscape Spline Point"));
+
+		ULandscapeSplineControlPoint* SelectedControlPoint = *SplineSelection->GetSelectedSplineControlPoints().CreateConstIterator();
 		ULandscapeSplinesComponent* SplinesComponent = SelectedControlPoint->GetOuterULandscapeSplinesComponent();
 		FVector LocalDrag = SplinesComponent->GetComponentTransform().InverseTransformVector(InDrag);
 
@@ -1896,7 +1463,7 @@ public:
 			SplitSegment(SegmentToSplit, Location);
 
 			UE::Widget::EWidgetMode WidgetMode = EdMode->GetModeManager()->GetWidgetMode(); 
-			SelectControlPoint(SplinesComponent->ControlPoints.Last());
+			SplineSelection->SelectControlPoint(SplinesComponent->ControlPoints.Last());
 			EdMode->GetModeManager()->SetWidgetMode(WidgetMode);
 		}
 		else
@@ -1907,7 +1474,7 @@ public:
 		}
 
 		// Get newly-created control point
-		SelectedControlPoint = *SelectedSplineControlPoints.CreateConstIterator();
+		SelectedControlPoint = *SplineSelection->GetSelectedSplineControlPoints().CreateConstIterator();
 
 		if (bAutoChangeConnectionsOnMove)
 		{
@@ -1922,7 +1489,7 @@ public:
 	// called when alt-dragging duplicated control point 
 	virtual bool UpdateDuplicateControlPoint(FVector& InDrag)
 	{
-		ULandscapeSplineControlPoint* SelectedControlPoint = *SelectedSplineControlPoints.CreateConstIterator();
+		ULandscapeSplineControlPoint* SelectedControlPoint = *SplineSelection->GetSelectedSplineControlPoints().CreateConstIterator();
 		ULandscapeSplinesComponent* SplinesComponent = SelectedControlPoint->GetOuterULandscapeSplinesComponent();
 		FVector LocalDrag = SplinesComponent->GetComponentTransform().InverseTransformVector(InDrag);
 		FVector Location = SelectedControlPoint->Location + LocalDrag;
@@ -2020,6 +1587,8 @@ public:
 			return true;
 		}
 
+		TArray<ULandscapeSplineControlPoint*> SelectedSplineControlPoints = SplineSelection->GetSelectedSplineControlPoints();
+
 		if (SelectedSplineControlPoints.Num() == 1 && InViewportClient->IsAltPressed() && InViewportClient->GetWidgetMode() == UE::Widget::WM_Translate && InViewportClient->GetCurrentWidgetAxis() != EAxisList::None)
 		{
 			static const int MaxDuplicationDelay = 3;
@@ -2083,65 +1652,15 @@ public:
 		return false;
 	}
 
-	void FixSelection()
-	{
-		TSet<ULandscapeSplineControlPoint*> SelectedSplineControlPointsBefore(ObjectPtrDecay(SelectedSplineControlPoints));
-		TSet<ULandscapeSplineSegment*> SelectedSplineSegmentsBefore(ObjectPtrDecay(SelectedSplineSegments));
-		SelectedSplineControlPoints.Empty();
-		SelectedSplineSegments.Empty();
-
-		if (EdMode->CurrentTool != nullptr && EdMode->CurrentTool == this)
-		{
-			for (const FLandscapeListInfo& Info : EdMode->GetLandscapeList())
-			{
-				Info.Info->ForAllSplineActors([this](TScriptInterface<ILandscapeSplineInterface> SplineOwner)
-				{
-					if (ULandscapeSplinesComponent* SplineComponent = SplineOwner->GetSplinesComponent())
-					{
-						Algo::CopyIf(SplineComponent->ControlPoints, SelectedSplineControlPoints, &ULandscapeSplineControlPoint::IsSplineSelected);
-						Algo::CopyIf(SplineComponent->Segments,      SelectedSplineSegments,      &ULandscapeSplineSegment::IsSplineSelected);
-					}
-				});
-			}
-
-			// If the selection has changed, update the properties panel
-			if (Algo::AnyOf(SelectedSplineControlPoints, [&SelectedSplineControlPointsBefore](TObjectPtr<ULandscapeSplineControlPoint> InControlPoint) { return !SelectedSplineControlPointsBefore.Contains(InControlPoint); })
-				|| Algo::AnyOf(SelectedSplineSegments, [&SelectedSplineSegmentsBefore](TObjectPtr<ULandscapeSplineSegment> InSegment) { return !SelectedSplineSegmentsBefore.Contains(InSegment); }))
-			{
-				UpdatePropertiesWindows();
-			}
-		}
-		else
-		{
-			for (const FLandscapeListInfo& Info : EdMode->GetLandscapeList())
-			{
-				Info.Info->ForAllSplineActors([](TScriptInterface<ILandscapeSplineInterface> SplineOwner)
-				{
-					if (ULandscapeSplinesComponent* SplineComponent = SplineOwner->GetSplinesComponent())
-					{
-						for (ULandscapeSplineControlPoint* ControlPoint : SplineComponent->ControlPoints)
-						{
-							ControlPoint->SetSplineSelected(false);
-						}
-
-						for (ULandscapeSplineSegment* Segment : SplineComponent->Segments)
-						{
-							Segment->SetSplineSelected(false);
-						}
-					}
-				});
-			}
-		}
-	}
-
 	void OnUndo()
 	{
-		FixSelection();
+		SplineSelection->UpdatePropertiesWindows();
 	}
 
 	virtual void EnterTool() override
 	{
 		GEditor->SelectNone(true, true, false);
+		SplineSelection = NewObject<ULandscapeSplineSelection>(GetTransientPackage(), TEXT("LandscapeSplineSelection"), RF_Transactional);
 
 		for (const FLandscapeListInfo& Info : EdMode->GetLandscapeList())
 		{
@@ -2157,8 +1676,7 @@ public:
 
 	virtual void ExitTool() override
 	{
-		ClearSelection();
-		UpdatePropertiesWindows();
+		SplineSelection->UpdatePropertiesWindows();
 
 		for (const FLandscapeListInfo& Info : EdMode->GetLandscapeList())
 		{
@@ -2170,6 +1688,9 @@ public:
 				}
 			});
 		}
+
+		SplineSelection->ClearSelection();
+		SplineSelection = nullptr;
 	}
 
 	virtual void Render(const FSceneView* View, FViewport* Viewport, FPrimitiveDrawInterface* PDI) override
@@ -2177,7 +1698,7 @@ public:
 		// The editor can try to render the tool before the UpdateLandscapeEditorData command runs and the landscape editor realizes that the landscape has been hidden/deleted
 		if (const ALandscapeProxy* LandscapeProxy = EdMode->CurrentToolTarget.LandscapeInfo.IsValid() ? EdMode->CurrentToolTarget.LandscapeInfo->GetLandscapeProxy() : nullptr)
 		{
-			for (ULandscapeSplineControlPoint* ControlPoint : SelectedSplineControlPoints)
+			for (ULandscapeSplineControlPoint* ControlPoint : SplineSelection->GetSelectedSplineControlPoints())
 			{
 				const ULandscapeSplinesComponent* SplinesComponent = ControlPoint->GetOuterULandscapeSplinesComponent();
 
@@ -2206,7 +1727,7 @@ public:
 
 			if (GLevelEditorModeTools().GetWidgetMode() == UE::Widget::WM_Scale && !Viewport->GetClient()->IsOrtho())
 			{
-				for (ULandscapeSplineSegment* Segment : SelectedSplineSegments)
+				for (ULandscapeSplineSegment* Segment : SplineSelection->GetSelectedSplineSegments())
 				{
 					const ULandscapeSplinesComponent* SplinesComponent = Segment->GetOuterULandscapeSplinesComponent();
 					for (int32 End = 0; End <= 1; End++)
@@ -2248,7 +1769,7 @@ public:
 
 	virtual bool UsesTransformWidget() const override
 	{
-		if (SelectedSplineControlPoints.Num() > 0 || DraggingTangent_Segment)
+		if (SplineSelection->GetSelectedSplineControlPoints().Num() > 0 || DraggingTangent_Segment)
 		{
 			// The editor can try to render the transform widget before the landscape editor ticks and realizes that the landscape has been hidden/deleted
 			return EdMode->CurrentToolTarget.LandscapeInfo.IsValid() && (EdMode->CurrentToolTarget.LandscapeInfo->GetLandscapeProxy() != nullptr);
@@ -2259,10 +1780,10 @@ public:
 
 	virtual EAxisList::Type GetWidgetAxisToDraw(UE::Widget::EWidgetMode CheckMode) const override
 	{
-		if (SelectedSplineControlPoints.Num() > 0)
+		if (SplineSelection->GetSelectedSplineControlPoints().Num() > 0)
 		{
 			//if (CheckMode == UE::Widget::WM_Rotate
-			//	&& SelectedSplineControlPoints.Num() >= 2)
+			//	&& SplineSelection->GetSelectedSplineControlPoints().Num() >= 2)
 			//{
 			//	return AXIS_X;
 			//}
@@ -2282,6 +1803,8 @@ public:
 
 	virtual FVector GetWidgetLocation() const override
 	{
+		TArray<ULandscapeSplineControlPoint*> SelectedSplineControlPoints = SplineSelection->GetSelectedSplineControlPoints();
+
 		if (const ALandscapeProxy* LandscapeProxy = EdMode->CurrentToolTarget.LandscapeInfo.IsValid() ? EdMode->CurrentToolTarget.LandscapeInfo->GetLandscapeProxy() : nullptr)
 		{
 			if (DraggingTangent_Segment)
@@ -2309,6 +1832,8 @@ public:
 
 	virtual FMatrix GetWidgetRotation() const override
 	{
+		TArray<ULandscapeSplineControlPoint*> SelectedSplineControlPoints = SplineSelection->GetSelectedSplineControlPoints();
+
 		const ALandscapeProxy* LandscapeProxy = EdMode->CurrentToolTarget.LandscapeInfo.IsValid() ? EdMode->CurrentToolTarget.LandscapeInfo->GetLandscapeProxy() : nullptr;
 		if (LandscapeProxy)
 		{
@@ -2333,7 +1858,7 @@ public:
 
 	virtual EEditAction::Type GetActionEditDuplicate() override
 	{
-		if (SelectedSplineControlPoints.Num() > 0 || SelectedSplineSegments.Num() > 0)
+		if (SplineSelection->GetSelectedSplineControlPoints().Num() > 0 || SplineSelection->GetSelectedSplineSegments().Num() > 0)
 		{
 			return EEditAction::Process;
 		}
@@ -2343,7 +1868,7 @@ public:
 
 	virtual EEditAction::Type GetActionEditDelete() override
 	{
-		if (SelectedSplineControlPoints.Num() > 0 || SelectedSplineSegments.Num() > 0)
+		if (SplineSelection->GetSelectedSplineControlPoints().Num() > 0 || SplineSelection->GetSelectedSplineSegments().Num() > 0)
 		{
 			return EEditAction::Process;
 		}
@@ -2353,7 +1878,7 @@ public:
 
 	virtual EEditAction::Type GetActionEditCut() override
 	{
-		if (SelectedSplineControlPoints.Num() > 0 || SelectedSplineSegments.Num() > 0)
+		if (SplineSelection->GetSelectedSplineControlPoints().Num() > 0 || SplineSelection->GetSelectedSplineSegments().Num() > 0)
 		{
 			return EEditAction::Process;
 		}
@@ -2363,7 +1888,7 @@ public:
 
 	virtual EEditAction::Type GetActionEditCopy() override
 	{
-		if (SelectedSplineControlPoints.Num() > 0 || SelectedSplineSegments.Num() > 0)
+		if (SplineSelection->GetSelectedSplineControlPoints().Num() > 0 || SplineSelection->GetSelectedSplineSegments().Num() > 0)
 		{
 			return EEditAction::Process;
 		}
@@ -2415,7 +1940,7 @@ public:
 
 	void InternalProcessEditDuplicate()
 	{
-		if (SelectedSplineControlPoints.Num() > 0 || SelectedSplineSegments.Num() > 0)
+		if (SplineSelection->GetSelectedSplineControlPoints().Num() > 0 || SplineSelection->GetSelectedSplineSegments().Num() > 0)
 		{
 			FScopedTransaction Transaction(LOCTEXT("LandscapeSpline_Duplicate", "Duplicate Landscape Splines"));
 
@@ -2427,6 +1952,9 @@ public:
 
 	void InternalProcessEditDelete()
 	{
+		TArray<ULandscapeSplineControlPoint*> SelectedSplineControlPoints = SplineSelection->GetSelectedSplineControlPoints();
+		TArray<ULandscapeSplineSegment*> SelectedSplineSegments = SplineSelection->GetSelectedSplineSegments();
+
 		if (SelectedSplineControlPoints.Num() > 0 || SelectedSplineSegments.Num() > 0)
 		{
 			FScopedTransaction Transaction(LOCTEXT("LandscapeSpline_Delete", "Delete Landscape Splines"));
@@ -2439,16 +1967,13 @@ public:
 			{
 				DeleteSegment(Segment, bDeleteLooseEnds);
 			}
-			ClearSelection();
-			UpdatePropertiesWindows();
-
-			GUnrealEd->RedrawLevelEditingViewports();
+			SplineSelection->ClearSelection();
 		}
 	}
 
 	void InternalProcessEditCut()
 	{
-		if (SelectedSplineControlPoints.Num() > 0 || SelectedSplineSegments.Num() > 0)
+		if (SplineSelection->GetSelectedSplineControlPoints().Num() > 0 || SplineSelection->GetSelectedSplineSegments().Num() > 0)
 		{
 			FScopedTransaction Transaction(LOCTEXT("LandscapeSpline_Cut", "Cut Landscape Splines"));
 
@@ -2473,6 +1998,9 @@ public:
 			}
 		};
 
+		TArray<ULandscapeSplineControlPoint*> SelectedSplineControlPoints = SplineSelection->GetSelectedSplineControlPoints();
+		TArray<ULandscapeSplineSegment*> SelectedSplineSegments = SplineSelection->GetSelectedSplineSegments();
+
 		if (SelectedSplineControlPoints.Num() > 0 || SelectedSplineSegments.Num() > 0)
 		{
 			TArray<UObject*> Objects;
@@ -2595,19 +2123,8 @@ protected:
 	FEdModeLandscape* EdMode;
 	TObjectPtr<ULandscapeInfo> LandscapeInfo;
 
-	TSet<TObjectPtr<ULandscapeSplineControlPoint>> SelectedSplineControlPoints;
-	TSet<TObjectPtr<ULandscapeSplineSegment>> SelectedSplineSegments;
-
-	// Last Selected
-	TObjectPtr<ULandscapeSplineControlPoint> LastSelectedControlPoint;
-	TObjectPtr<ULandscapeSplineSegment> LastSelectedSegment;
-
-	// Control Points and Segments stored in sequential order
-	TArray<TObjectPtr<ULandscapeSplineControlPoint>> LinearControlPoints;
-	TArray<TObjectPtr<ULandscapeSplineSegment>> LinearSegments;
-
-	// Persistent control point mode
-	TOptional<UE::Widget::EWidgetMode> ControlPointWidgetMode;
+	/** Tracks the selected points/segments and handles linear spline navigation */
+	TObjectPtr<ULandscapeSplineSelection> SplineSelection;
 
 	TObjectPtr<ULandscapeSplineSegment> DraggingTangent_Segment;
 	float DraggingTangent_Length;
@@ -2658,7 +2175,7 @@ protected:
 
 bool FEdModeLandscape::HasSelectedSplineSegments() const
 {
-	return SplinesTool && (SplinesTool->SelectedSplineSegments.Num() > 0);
+	return SplinesTool && (SplinesTool->SplineSelection->GetSelectedSplineSegments().Num() > 0);
 }
 
 bool FEdModeLandscape::HasAdjacentLinearSplineConnection(ESplineNavigationFlags Flags) const
@@ -2673,11 +2190,11 @@ bool FEdModeLandscape::HasAdjacentLinearSplineConnection(ESplineNavigationFlags
 
 	if (EnumHasAllFlags(Flags, ESplineNavigationFlags::SegmentSelectModeEnabled))
 	{
-		return SplinesTool->HasAdjacentSegmentInLinearPath(Flags);
+		return SplinesTool->SplineSelection->HasAdjacentSegmentInLinearPath(Flags);
 	}
 	else
 	{
-		return SplinesTool->HasAdjacentControlPointInLinearPath(Flags);
+		return SplinesTool->SplineSelection->HasAdjacentControlPointInLinearPath(Flags);
 	}
 }
 
@@ -2692,13 +2209,13 @@ bool FEdModeLandscape::HasEndLinearSplineConnection(ESplineNavigationFlags Flags
 
 	if (EnumHasAllFlags(Flags, ESplineNavigationFlags::SegmentSelectModeEnabled))
 	{
-		const ULandscapeSplineSegment* EndSegment = SplinesTool->GetEndSegmentInLinearPath(Flags);
-		return EndSegment != nullptr && !EndSegment->IsSplineSelected() && SplinesTool->IsSelectionValidForNavigation();
+		const ULandscapeSplineSegment* EndSegment = SplinesTool->SplineSelection->GetEndSegmentInLinearPath(Flags);
+		return EndSegment != nullptr && !EndSegment->IsSplineSelected() && SplinesTool->SplineSelection->IsSelectionValidForNavigation();
 	}
 	else
 	{
-		const ULandscapeSplineControlPoint* EndControlPoint = SplinesTool->GetEndControlPointInLinearPath(Flags);
-		return EndControlPoint != nullptr && !EndControlPoint->IsSplineSelected() && SplinesTool->IsSelectionValidForNavigation();
+		const ULandscapeSplineControlPoint* EndControlPoint = SplinesTool->SplineSelection->GetEndControlPointInLinearPath(Flags);
+		return EndControlPoint != nullptr && !EndControlPoint->IsSplineSelected() && SplinesTool->SplineSelection->IsSelectionValidForNavigation();
 	}
 }
 
@@ -2710,6 +2227,7 @@ void FEdModeLandscape::FlipSelectedSplineSegments()
 	}
 
 	// Do Flip
+	FScopedTransaction Transaction(LOCTEXT("LandscapeSpline_FlipSegments", "Flip Selected Landscape Spline Segments"));
 	SplinesTool->FlipSelectedSplineSegments();
 }
 
@@ -2723,12 +2241,12 @@ void FEdModeLandscape::ShowSplineProperties()
 
 void FEdModeLandscape::GetSelectedSplineOwners(TSet<AActor*>& SelectedSplineOwners) const
 {
-	for (ULandscapeSplineSegment* Segment : SplinesTool->SelectedSplineSegments)
+	for (ULandscapeSplineSegment* Segment : SplinesTool->SplineSelection->GetSelectedSplineSegments())
 	{
 		SelectedSplineOwners.Add(Segment->GetTypedOuter<AActor>());
 	}
 
-	for (ULandscapeSplineControlPoint* ControlPoint : SplinesTool->SelectedSplineControlPoints)
+	for (ULandscapeSplineControlPoint* ControlPoint : SplinesTool->SplineSelection->GetSelectedSplineControlPoints())
 	{
 		SelectedSplineOwners.Add(ControlPoint->GetTypedOuter<AActor>());
 	}
@@ -2736,53 +2254,53 @@ void FEdModeLandscape::GetSelectedSplineOwners(TSet<AActor*>& SelectedSplineOwne
 
 void FEdModeLandscape::SelectAllSplineControlPoints()
 {
+	FScopedTransaction Transaction(LOCTEXT("LandscapeSpline_SelectAllPoints", "Select All Landscape Spline Points"));
+
 	ULandscapeInfo* CurrentLandscapeInfo = CurrentToolTarget.LandscapeInfo.Get();
 	if (SplinesTool && CurrentLandscapeInfo)
 	{
-		SplinesTool->ClearSelectedSegments();
-		SplinesTool->SelectAllControlPoints(*CurrentLandscapeInfo);
-
-		SplinesTool->UpdatePropertiesWindows();
-		GUnrealEd->RedrawLevelEditingViewports();
+		SplinesTool->SplineSelection->SelectAllControlPoints(*CurrentLandscapeInfo);
+		SplinesTool->SplineSelection->UpdatePropertiesWindows();
 	}
 }
 
 void FEdModeLandscape::SelectAllSplineSegments()
 {
+	FScopedTransaction Transaction(LOCTEXT("LandscapeSpline_SelectAllSegments", "Select All Landscape Spline Segments"));
+
 	ULandscapeInfo* CurrentLandscapeInfo = CurrentToolTarget.LandscapeInfo.Get();
 	if (SplinesTool && CurrentLandscapeInfo)
 	{
-		SplinesTool->ClearSelectedControlPoints();
-		SplinesTool->SelectAllSplineSegments(*CurrentLandscapeInfo);
-
-		SplinesTool->UpdatePropertiesWindows();
-		GUnrealEd->RedrawLevelEditingViewports();
+		SplinesTool->SplineSelection->SelectAllSplineSegments(*CurrentLandscapeInfo);
+		SplinesTool->SplineSelection->UpdatePropertiesWindows();
 	}
 }
 
 void FEdModeLandscape::SelectAllConnectedSplineControlPoints()
 {
+	FScopedTransaction Transaction(LOCTEXT("LandscapeSpline_SelectConnectedPoints", "Select Landscape Spline Connected Points"));
+
 	if (SplinesTool /*&& SplinesTool == CurrentTool*/)
 	{
-		SplinesTool->SelectAdjacentControlPoints();
-		SplinesTool->ClearSelectedSegments();
-		SplinesTool->SelectConnected();
+		SplinesTool->SplineSelection->SelectAdjacentControlPoints();
+		SplinesTool->SplineSelection->ClearSelectedSegments();
+		SplinesTool->SplineSelection->SelectConnected();
 
-		SplinesTool->UpdatePropertiesWindows();
-		GUnrealEd->RedrawLevelEditingViewports();
+		SplinesTool->SplineSelection->UpdatePropertiesWindows();
 	}
 }
 
 void FEdModeLandscape::SelectAllConnectedSplineSegments()
 {
+	FScopedTransaction Transaction(LOCTEXT("LandscapeSpline_SelectConnectedSegments", "Select Landscape Spline Connected Segments"));
+
 	if (SplinesTool /*&& SplinesTool == CurrentTool*/)
 	{
-		SplinesTool->SelectAdjacentSegments();
-		SplinesTool->ClearSelectedControlPoints();
-		SplinesTool->SelectConnected();
+		SplinesTool->SplineSelection->SelectAdjacentSegments();
+		SplinesTool->SplineSelection->ClearSelectedControlPoints();
+		SplinesTool->SplineSelection->SelectConnected();
 
-		SplinesTool->UpdatePropertiesWindows();
-		GUnrealEd->RedrawLevelEditingViewports();
+		SplinesTool->SplineSelection->UpdatePropertiesWindows();
 	}
 }
 
@@ -2792,44 +2310,32 @@ void FEdModeLandscape::SelectAdjacentLinearSplineElement(ESplineNavigationFlags
 	{
 		// check only one selection mode is set
 		check(FMath::CountBits(static_cast<uint8>(Flags & ESplineNavigationFlags::SelectModeMask)) == 1);
+		const bool bAddToSelection = EnumHasAllFlags(Flags, ESplineNavigationFlags::AddToSelection);
 
 		if (EnumHasAllFlags(Flags, ESplineNavigationFlags::SegmentSelectModeEnabled))
 		{
-			ULandscapeSplineSegment* AdjacentSegment = SplinesTool->GetAdjacentSegmentInLinearPath(Flags);
+			FScopedTransaction Transaction(LOCTEXT("LandscapeSpline_SelectAdjacentSegment", "Select Landscape Spline Segment"));
+			ULandscapeSplineSegment* AdjacentSegment = SplinesTool->SplineSelection->GetAdjacentSegmentInLinearPath(Flags);
 
-			// Clear all other selections if the user only wants the next segment selected
-			if (!EnumHasAllFlags(Flags, ESplineNavigationFlags::AddToSelection))
-			{
-				SplinesTool->ClearSelectedSegments();
-			}
-
-			if (!AdjacentSegment->IsSelected())
-			{
-				SplinesTool->SelectSegment(AdjacentSegment);
-			}
+			// SelectSegment will reset selection unless AddToSelection flag is passed
+			SplinesTool->SplineSelection->SelectSegment(AdjacentSegment, bAddToSelection ? ESplineNavigationFlags::AddToSelection : ESplineNavigationFlags::None);
 
-			SplinesTool->ClearSelectedControlPoints();
+			// When AddToSelection is set, SelectSegment will not clear control points
+			SplinesTool->SplineSelection->ClearSelectedControlPoints();
 		}
 		else 
 		{
-			ULandscapeSplineControlPoint* AdjacentPoint = SplinesTool->GetAdjacentControlPointInPath(Flags);
+			FScopedTransaction Transaction(LOCTEXT("LandscapeSpline_SelectAdjacentPoint", "Select Landscape Spline Point"));
+			ULandscapeSplineControlPoint* AdjacentPoint = SplinesTool->SplineSelection->GetAdjacentControlPointInPath(Flags);
 
-			// Clear all other selections if the user only wants the next point selected
-			if (!EnumHasAllFlags(Flags, ESplineNavigationFlags::AddToSelection))
-			{
-				SplinesTool->ClearSelectedControlPoints();
-			}
+			// SelectSegment will reset selection unless AddToSelection flag is passed
+			SplinesTool->SplineSelection->SelectControlPoint(AdjacentPoint, bAddToSelection ? ESplineNavigationFlags::AddToSelection : ESplineNavigationFlags::None);
 
-			if (!AdjacentPoint->IsSplineSelected())
-			{
-				SplinesTool->SelectControlPoint(AdjacentPoint);
-			}
-
-			SplinesTool->ClearSelectedSegments();
+			// When AddToSelection is set, SelectControlPoint will not clear Segments
+			SplinesTool->SplineSelection->ClearSelectedSegments();
 		}
 		
-		SplinesTool->UpdatePropertiesWindows();
-		GUnrealEd->RedrawLevelEditingViewports();
+		SplinesTool->SplineSelection->UpdatePropertiesWindows();
 	}
 }
 
@@ -2842,56 +2348,48 @@ void FEdModeLandscape::SelectEndLinearSplineElement(ESplineNavigationFlags Flags
 
 		if (EnumHasAllFlags(Flags, ESplineNavigationFlags::SegmentSelectModeEnabled))
 		{
-			ULandscapeSplineSegment* EndSegment = SplinesTool->GetEndSegmentInLinearPath(Flags);
+			FScopedTransaction Transaction(LOCTEXT("LandscapeSpline_SelectEndSegment", "Select Landscape Spline Segment"));
 
+			ULandscapeSplineSegment* EndSegment = SplinesTool->SplineSelection->GetEndSegmentInLinearPath(Flags);
 			if (EndSegment != nullptr)
 			{
-				SplinesTool->ClearSelectedSegments();
-				if (!EndSegment->IsSplineSelected())
-				{
-					SplinesTool->SelectSegment(EndSegment);
-				}
-				SplinesTool->ClearSelectedControlPoints();
+				SplinesTool->SplineSelection->SelectSegment(EndSegment, ESplineNavigationFlags::UpdatePropertiesWindows);
 			}
 		}
 		else
 		{
-			ULandscapeSplineControlPoint* EndPoint = SplinesTool->GetEndControlPointInLinearPath(Flags);
+			FScopedTransaction Transaction(LOCTEXT("LandscapeSpline_SelectEndPoint", "Select Landscape Spline Point"));
+
+			ULandscapeSplineControlPoint* EndPoint = SplinesTool->SplineSelection->GetEndControlPointInLinearPath(Flags);
 			if (EndPoint != nullptr)
 			{
-				SplinesTool->ClearSelectedControlPoints();
-				if (!EndPoint->IsSplineSelected())
-				{
-					SplinesTool->SelectControlPoint(EndPoint);
-				}
-				SplinesTool->ClearSelectedSegments();
+				SplinesTool->SplineSelection->SelectControlPoint(EndPoint, ESplineNavigationFlags::UpdatePropertiesWindows);
 			}
 		}
-
-		SplinesTool->UpdatePropertiesWindows();
-		GUnrealEd->RedrawLevelEditingViewports();
 	}
 }
 
 void FEdModeLandscape::SelectSplineControlPointsFromCurrentSegmentSelection() const
 {
+	FScopedTransaction Transaction(LOCTEXT("LandscapeSpline_ConvertToPoints", "Switch selected Segments to Points"));
+
 	if (SplinesTool /*&& SplinesTool == CurrentTool*/)
 	{
-		SplinesTool->SelectAdjacentControlPoints();
-		SplinesTool->ClearSelectedSegments();
-		SplinesTool->UpdatePropertiesWindows();
-		GUnrealEd->RedrawLevelEditingViewports();
+		SplinesTool->SplineSelection->SelectAdjacentControlPoints();
+		SplinesTool->SplineSelection->ClearSelectedSegments();
+		SplinesTool->SplineSelection->UpdatePropertiesWindows();
 	}
 }
 
 void FEdModeLandscape::SelectSplineSegmentsFromCurrentControlPointSelection() const
 {
+	FScopedTransaction Transaction(LOCTEXT("LandscapeSpline_ConvertToSegments", "Switch selected Points to Segments"));
+
 	if (SplinesTool /*&& SplinesTool == CurrentTool*/)
 	{
-		SplinesTool->SelectAdjacentSegments();
-		SplinesTool->ClearSelectedControlPoints();
-		SplinesTool->UpdatePropertiesWindows();
-		GUnrealEd->RedrawLevelEditingViewports();
+		SplinesTool->SplineSelection->SelectAdjacentSegments();
+		SplinesTool->SplineSelection->ClearSelectedControlPoints();
+		SplinesTool->SplineSelection->UpdatePropertiesWindows();
 	}
 }
 
@@ -2902,14 +2400,13 @@ void FEdModeLandscape::SplineMoveToCurrentLevel()
 	if (SplinesTool /*&& SplinesTool == CurrentTool*/)
 	{
 		// Select all connected control points
-		SplinesTool->SelectAdjacentSegments();
-		SplinesTool->SelectAdjacentControlPoints();
-		SplinesTool->SelectConnected();
+		SplinesTool->SplineSelection->SelectAdjacentSegments();
+		SplinesTool->SplineSelection->SelectAdjacentControlPoints();
+		SplinesTool->SplineSelection->SelectConnected();
 
 		SplinesTool->MoveSelectedToLevel();
 
-		SplinesTool->ClearSelection();
-		SplinesTool->UpdatePropertiesWindows();
+		SplinesTool->SplineSelection->ClearSelection();
 	}
 }
 
diff --git a/Engine/Source/Editor/LandscapeEditor/Private/LandscapeSplineDetails.cpp b/Engine/Source/Editor/LandscapeEditor/Private/LandscapeSplineDetails.cpp
--- a/Engine/Source/Editor/LandscapeEditor/Private/LandscapeSplineDetails.cpp
+++ b/Engine/Source/Editor/LandscapeEditor/Private/LandscapeSplineDetails.cpp
@@ -1,6 +1,7 @@
 // Copyright Epic Games, Inc. All Rights Reserved.
 
 #include "LandscapeSplineDetails.h"
+#include "LandscapeSplineSelection.h"
 #include "Widgets/DeclarativeSyntaxSupport.h"
 #include "Widgets/SBoxPanel.h"
 #include "LandscapeEdMode.h"
diff --git a/Engine/Source/Editor/LandscapeEditor/Private/LandscapeSplineSelection.cpp b/Engine/Source/Editor/LandscapeEditor/Private/LandscapeSplineSelection.cpp
new file mode 100644
--- /dev/null
+++ b/Engine/Source/Editor/LandscapeEditor/Private/LandscapeSplineSelection.cpp
@@ -0,0 +1,527 @@
+// Copyright Epic Games, Inc. All Rights Reserved.
+
+#include "LandscapeSplineSelection.h"
+
+#include "LandscapeSplineControlPoint.h"
+#include "LandscapeSplineSegment.h"
+#include "LandscapeSplinesComponent.h"
+
+#include "UnrealEdGlobals.h"
+#include "Editor/UnrealEdEngine.h"
+#include "EditorModeManager.h"
+#include "EditorModes.h"
+#include "UObject/ConstructorHelpers.h"
+#include "PropertyEditorModule.h"
+#include "Modules/ModuleManager.h"
+
+#define LOCTEXT_NAMESPACE "LandscapeSplineSelection"
+
+ULandscapeSplineSelection::ULandscapeSplineSelection(const FObjectInitializer& ObjectInitializer)
+	: Super(ObjectInitializer)
+{
+}
+
+void ULandscapeSplineSelection::SelectSegment(ULandscapeSplineSegment* Segment, ESplineNavigationFlags Flags)
+{
+	// UpdateProperties, AddToSelection, SegmentModeFlags are only valid flags
+	check(Flags == ESplineNavigationFlags::None || !EnumHasAnyFlags(Flags, ~(ESplineNavigationFlags::UpdatePropertiesWindows | ESplineNavigationFlags::AddToSelection | ESplineNavigationFlags::SegmentSelectModeEnabled)));
+
+	// Do nothing if the given point is already selected
+	if (IsSegmentSelected(Segment))
+	{
+		check(Segment->IsSplineSelected());
+		return;
+	}
+
+	// Before entering segment mode, save the control point mode for persistence
+	if (SelectedSplineSegments.Num() == 0)
+	{
+		ControlPointWidgetMode = GLevelEditorModeTools().GetWidgetMode();
+	}
+
+	// Clear all previous selection
+	if (!EnumHasAllFlags(Flags, ESplineNavigationFlags::AddToSelection))
+	{
+		ClearSelection();
+	}
+
+	Modify(false);
+	SelectedSplineSegments.Add(Segment);
+
+	SelectNavigationSegment(Segment);
+
+	Segment->Modify(false);
+	Segment->SetSplineSelected(true);
+	GLevelEditorModeTools().SetWidgetMode(UE::Widget::WM_Scale);
+
+	if (EnumHasAllFlags(Flags, ESplineNavigationFlags::UpdatePropertiesWindows))
+	{
+		UpdatePropertiesWindows();
+	}
+}
+
+void ULandscapeSplineSelection::SelectControlPoint(ULandscapeSplineControlPoint* ControlPoint, ESplineNavigationFlags Flags)
+{
+	// UpdateProperties, AddToSelection, ControlPointSelect are only valid flags
+	check(Flags == ESplineNavigationFlags::None || !EnumHasAnyFlags(Flags, ~(ESplineNavigationFlags::UpdatePropertiesWindows | ESplineNavigationFlags::AddToSelection | ESplineNavigationFlags::ControlPointSelectModeEnabled)));
+
+	// Do nothing if the given point is already selected
+	if (IsControlPointSelected(ControlPoint))
+	{
+		check(ControlPoint->IsSplineSelected());
+		return;
+	}
+
+	// When switching from segment mode, restore the last used control point mode
+	if (SelectedSplineSegments.Num() > 0 && ControlPointWidgetMode.IsSet())
+	{
+		GLevelEditorModeTools().SetWidgetMode(ControlPointWidgetMode.GetValue());
+	}
+
+	// Clear all previous selection
+	if (!EnumHasAllFlags(Flags, ESplineNavigationFlags::AddToSelection))
+	{
+		ClearSelection();
+	}
+
+	Modify(false);
+	SelectedSplineControlPoints.Add(ControlPoint);
+
+	SelectNavigationControlPoint(ControlPoint);
+
+	ControlPoint->Modify(false);
+	ControlPoint->SetSplineSelected(true);
+
+	if (EnumHasAllFlags(Flags, ESplineNavigationFlags::UpdatePropertiesWindows))
+	{
+		UpdatePropertiesWindows();
+	}
+}
+
+void ULandscapeSplineSelection::ClearSelectedControlPoints()
+{
+	for (ULandscapeSplineControlPoint* ControlPoint : SelectedSplineControlPoints)
+	{
+		check(ControlPoint->IsSplineSelected());
+		ControlPoint->Modify(false);
+		ControlPoint->SetSplineSelected(false);
+	}
+	
+	Modify(false);
+	SelectedSplineControlPoints.Empty();
+}
+
+void ULandscapeSplineSelection::ClearSelectedSegments()
+{
+	for (ULandscapeSplineSegment* Segment : SelectedSplineSegments)
+	{
+		check(Segment->IsSplineSelected());
+		Segment->Modify(false);
+		Segment->SetSplineSelected(false);
+	}
+
+	Modify(false);
+	SelectedSplineSegments.Empty();
+}
+
+void ULandscapeSplineSelection::ClearSelection()
+{
+	ClearSelectedControlPoints();
+	ClearSelectedSegments();
+
+	UpdatePropertiesWindows();
+}
+
+void ULandscapeSplineSelection::DeselectControlPoint(ULandscapeSplineControlPoint* ControlPoint, ESplineNavigationFlags Flags)
+{
+	// UpdateProperties is only valid flag
+	check(Flags == ESplineNavigationFlags::None || !EnumHasAnyFlags(Flags, ~(ESplineNavigationFlags::UpdatePropertiesWindows)));
+	check(ControlPoint->IsSplineSelected());
+
+	ControlPoint->Modify(false);
+	ControlPoint->SetSplineSelected(false);
+
+	Modify(false);
+	SelectedSplineControlPoints.Remove(ControlPoint);
+
+	if (EnumHasAllFlags(Flags, ESplineNavigationFlags::UpdatePropertiesWindows))
+	{
+		UpdatePropertiesWindows();
+	}
+}
+
+void ULandscapeSplineSelection::DeSelectSegment(ULandscapeSplineSegment* Segment, ESplineNavigationFlags Flags)
+{
+	// UpdateProperties is only valid flag
+	check(Flags == ESplineNavigationFlags::None || !EnumHasAnyFlags(Flags, ~(ESplineNavigationFlags::UpdatePropertiesWindows)));
+	check(Segment->IsSplineSelected());
+
+	Segment->Modify(false);
+	Segment->SetSplineSelected(false);
+
+	Modify(false);
+	SelectedSplineSegments.Remove(Segment);
+
+	if (EnumHasAllFlags(Flags, ESplineNavigationFlags::UpdatePropertiesWindows))
+	{
+		UpdatePropertiesWindows();
+	}
+}
+
+void ULandscapeSplineSelection::SelectConnected()
+{
+	TArray<ULandscapeSplineControlPoint*> ControlPointsToProcess = ObjectPtrDecay(SelectedSplineControlPoints);
+
+	while (ControlPointsToProcess.Num() > 0)
+	{
+		const ULandscapeSplineControlPoint* ControlPoint = ControlPointsToProcess.Pop();
+
+		for (const FLandscapeSplineConnection& Connection : ControlPoint->ConnectedSegments)
+		{
+			ULandscapeSplineControlPoint* OtherEnd = Connection.GetFarConnection().ControlPoint;
+
+			if (!OtherEnd->IsSplineSelected())
+			{
+				SelectControlPoint(OtherEnd, ESplineNavigationFlags::AddToSelection);
+				ControlPointsToProcess.Add(OtherEnd);
+			}
+		}
+	}
+
+	TArray<ULandscapeSplineSegment*> SegmentsToProcess = ObjectPtrDecay(SelectedSplineSegments);
+
+	while (SegmentsToProcess.Num() > 0)
+	{
+		const ULandscapeSplineSegment* Segment = SegmentsToProcess.Pop();
+
+		for (const FLandscapeSplineSegmentConnection& SegmentConnection : Segment->Connections)
+		{
+			for (const FLandscapeSplineConnection& Connection : SegmentConnection.ControlPoint->ConnectedSegments)
+			{
+				if (Connection.Segment != Segment && !Connection.Segment->IsSplineSelected())
+				{
+					SelectSegment(Connection.Segment, ESplineNavigationFlags::AddToSelection);
+					SegmentsToProcess.Add(Connection.Segment);
+				}
+			}
+		}
+	}
+}
+
+void ULandscapeSplineSelection::SelectAllSplineSegments(const ULandscapeInfo& InLandscapeInfo)
+{
+	TArray<TScriptInterface<ILandscapeSplineInterface>> SplineActors(InLandscapeInfo.GetSplineActors());
+	for (TScriptInterface<ILandscapeSplineInterface> SplineActor : SplineActors)
+	{
+		if (ULandscapeSplinesComponent* SplineComponent = SplineActor->GetSplinesComponent())
+		{
+			SplineComponent->ForEachControlPoint([this](ULandscapeSplineControlPoint* ControlPoint)
+				{
+					for (const FLandscapeSplineConnection& Connection : ControlPoint->ConnectedSegments)
+					{
+						SelectSegment(Connection.Segment, ESplineNavigationFlags::AddToSelection);
+					}
+				});
+		}
+	}
+}
+
+void ULandscapeSplineSelection::SelectAllControlPoints(const ULandscapeInfo& InLandscapeInfo)
+{
+	TArray<TScriptInterface<ILandscapeSplineInterface>> SplineActors(InLandscapeInfo.GetSplineActors());
+	for (TScriptInterface<ILandscapeSplineInterface> SplineActor : SplineActors)
+	{
+		if (ULandscapeSplinesComponent* SplineComponent = SplineActor->GetSplinesComponent())
+		{
+			SplineComponent->ForEachControlPoint([this](ULandscapeSplineControlPoint* ControlPoint) { SelectControlPoint(ControlPoint, ESplineNavigationFlags::AddToSelection); });
+		}
+	}
+}
+
+void ULandscapeSplineSelection::SelectAdjacentControlPoints()
+{
+	for (const ULandscapeSplineSegment* Segment : SelectedSplineSegments)
+	{
+		SelectControlPoint(Segment->Connections[0].ControlPoint, ESplineNavigationFlags::AddToSelection);
+		SelectControlPoint(Segment->Connections[1].ControlPoint, ESplineNavigationFlags::AddToSelection);
+	}
+}
+
+void ULandscapeSplineSelection::SelectAdjacentSegments()
+{
+	for (ULandscapeSplineControlPoint* ControlPoint : SelectedSplineControlPoints)
+	{
+		for (const FLandscapeSplineConnection& Connection : ControlPoint->ConnectedSegments)
+		{
+			SelectSegment(Connection.Segment, ESplineNavigationFlags::AddToSelection);
+		}
+	}
+}
+
+const bool ULandscapeSplineSelection::IsSegmentSelected(ULandscapeSplineSegment* Segment) const
+{
+	return SelectedSplineSegments.Contains(Segment);
+}
+
+const bool ULandscapeSplineSelection::IsControlPointSelected(ULandscapeSplineControlPoint* ControlPoint) const
+{
+	return SelectedSplineControlPoints.Contains(ControlPoint);
+}
+
+void ULandscapeSplineSelection::UpdatePropertiesWindows() const
+{
+	if (GLevelEditorModeTools().GetActiveMode(FBuiltinEditorModes::EM_Landscape))
+	{
+		TArray<UObject*> Objects;
+		Objects.Reset(SelectedSplineControlPoints.Num() + SelectedSplineSegments.Num());
+		Algo::Copy(SelectedSplineControlPoints, Objects);
+		Algo::Copy(SelectedSplineSegments, Objects);
+
+		FPropertyEditorModule& PropertyModule = FModuleManager::Get().LoadModuleChecked<FPropertyEditorModule>(TEXT("PropertyEditor"));
+		PropertyModule.UpdatePropertyViews(Objects);
+	}
+
+	GUnrealEd->RedrawLevelEditingViewports();
+}
+
+ULandscapeSplineControlPoint* ULandscapeSplineSelection::GetLinearEndControlPointInternal(ULandscapeSplineControlPoint* SelectedControlPoint)
+{
+	TSet<ULandscapeSplineControlPoint*> ControlPointsVisited;
+	TArray<ULandscapeSplineControlPoint*> ControlPointsToProcess;
+	ControlPointsToProcess.Add(SelectedControlPoint);
+
+	Modify(false);
+	// Reset the linear path before it gets rebuilt below
+	LinearControlPoints.Empty();
+
+	// Control points with zero/multi node connections are not part of a valid linear path
+	if (SelectedControlPoint->ConnectedSegments.IsEmpty() || SelectedControlPoint->ConnectedSegments.Num() > 2)
+	{
+		LinearControlPoints.Add(SelectedControlPoint);
+		return SelectedControlPoint;
+	}
+
+	while (ControlPointsToProcess.Num() > 0)
+	{
+		ULandscapeSplineControlPoint* ControlPoint = ControlPointsToProcess.Pop();
+		ControlPointsVisited.Add(ControlPoint);
+
+		// Store nodes in order of visiting to create the linear path
+		if (!LinearControlPoints.Contains(ControlPoint))
+		{
+			LinearControlPoints.Add(ControlPoint);
+		}
+
+		// Valid end points only have one connection
+		if (ControlPoint->ConnectedSegments.Num() == 1 && ControlPoint != SelectedControlPoint)
+		{
+			return ControlPoint;
+		}
+
+		for (const FLandscapeSplineConnection& Connection : ControlPoint->ConnectedSegments)
+		{
+			ULandscapeSplineControlPoint* OtherEnd = Connection.GetFarConnection().ControlPoint;
+			check(OtherEnd != nullptr);
+
+			if (OtherEnd->ConnectedSegments.IsEmpty() || OtherEnd->ConnectedSegments.Num() > 2)
+			{
+				continue;
+			}
+
+			if (!ControlPointsVisited.Contains(OtherEnd))
+			{
+				ControlPointsToProcess.Add(OtherEnd);
+			}
+		}
+	}
+
+	return SelectedControlPoint;
+}
+
+void ULandscapeSplineSelection::ResetNavigationPath()
+{
+	Modify(false);
+	// Clear cached path
+	LinearControlPoints.Empty();
+	LinearSegments.Empty();
+}
+
+void ULandscapeSplineSelection::BuildLinearPathFromLastSelectedPointInternal()
+{
+	Modify(false);
+	// Calculate the last selected point
+	// If a control point is selected, use it directly. Otherwise, if a segment is selected, retrieve the control point from its connection
+	TObjectPtr<ULandscapeSplineSegment> LastSelectedSegment = SelectedSplineSegments.IsEmpty() ? nullptr : SelectedSplineSegments.Last();
+	TObjectPtr<ULandscapeSplineControlPoint> LastSelectedControlPoint = SelectedSplineControlPoints.IsEmpty() && LastSelectedSegment ? LastSelectedSegment->Connections[0].ControlPoint : SelectedSplineControlPoints.Last();
+
+	check(LastSelectedControlPoint != nullptr);
+
+	LinearControlPoints.Empty();
+	LinearSegments.Empty();
+
+	ULandscapeSplineControlPoint* StartControlPoint = GetLinearEndControlPointInternal(LastSelectedControlPoint);
+	ULandscapeSplineControlPoint* EndControlPoint = GetLinearEndControlPointInternal(StartControlPoint);
+
+	for (const ULandscapeSplineControlPoint* Point : LinearControlPoints)
+	{
+		for (const FLandscapeSplineConnection& Connection : Point->ConnectedSegments)
+		{
+			if (!LinearSegments.Contains(Connection.Segment))
+			{
+				LinearSegments.Add(Connection.Segment);
+			}
+		}
+	}
+}
+
+void ULandscapeSplineSelection::SelectNavigationSegment(ULandscapeSplineSegment* Segment)
+{
+	Modify(false);
+
+	if (!LinearSegments.Contains(Segment))
+	{
+		BuildLinearPathFromLastSelectedPointInternal();
+	}
+}
+
+void ULandscapeSplineSelection::SelectNavigationControlPoint(ULandscapeSplineControlPoint* ControlPoint)
+{
+	Modify(false);
+
+	if (!LinearControlPoints.Contains(ControlPoint))
+	{
+		BuildLinearPathFromLastSelectedPointInternal();
+	}
+}
+
+bool ULandscapeSplineSelection::IsSelectionValidForNavigation() const
+{
+	// If any elements in selection are part of different linear paths, disable navigation
+	for (ULandscapeSplineControlPoint* ControlPoint : SelectedSplineControlPoints)
+	{
+		if (!LinearControlPoints.Contains(ControlPoint))
+		{
+			return false;
+		}
+	}
+
+	for (ULandscapeSplineSegment* Segment : SelectedSplineSegments)
+	{
+		if (!LinearSegments.Contains(Segment))
+		{
+			return false;
+		}
+	}
+
+	return true;
+}
+
+ULandscapeSplineSegment* ULandscapeSplineSelection::GetEndSegmentInLinearPath(ESplineNavigationFlags Flags) const
+{
+	check(EnumHasAllFlags(Flags, ESplineNavigationFlags::SegmentSelectModeEnabled));
+	check(FMath::CountBits(static_cast<uint8>(Flags & ESplineNavigationFlags::DirectionMask)) == 1);
+
+	if (LinearSegments.IsEmpty())
+	{
+		return nullptr;
+	}
+
+	if (EnumHasAllFlags(Flags, ESplineNavigationFlags::DirectionForward))
+	{
+		return LinearSegments.Top();
+	}
+
+	return LinearSegments[0];
+}
+
+ULandscapeSplineControlPoint* ULandscapeSplineSelection::GetEndControlPointInLinearPath(ESplineNavigationFlags Flags) const
+{
+	check(EnumHasAllFlags(Flags, ESplineNavigationFlags::ControlPointSelectModeEnabled));
+	check(FMath::CountBits(static_cast<uint8>(Flags & ESplineNavigationFlags::SelectModeMask)) == 1);
+
+	if (LinearControlPoints.IsEmpty())
+	{
+		return nullptr;
+	}
+
+	if (EnumHasAllFlags(Flags, ESplineNavigationFlags::DirectionForward))
+	{
+		return LinearControlPoints.Top();
+	}
+
+	return LinearControlPoints[0];
+}
+
+ULandscapeSplineSegment* ULandscapeSplineSelection::GetAdjacentSegmentInLinearPath(ESplineNavigationFlags Flags) const
+{
+	check(EnumHasAllFlags(Flags, ESplineNavigationFlags::SegmentSelectModeEnabled));
+	check(FMath::CountBits(static_cast<uint8>(Flags & ESplineNavigationFlags::SelectModeMask)) == 1);
+
+	const int32 LastSelectedSegmentIndex = LinearSegments.Find(SelectedSplineSegments.Last());
+	const int32 NextIndexOffset = EnumHasAllFlags(Flags, ESplineNavigationFlags::DirectionForward) ? 1 : -1;
+	check(LastSelectedSegmentIndex != INDEX_NONE); // last selected should always be part of the linear path
+
+	if (LinearSegments.IsValidIndex(LastSelectedSegmentIndex + NextIndexOffset))
+	{
+		return LinearSegments[LastSelectedSegmentIndex + NextIndexOffset];
+	}
+	return nullptr;
+}
+
+ULandscapeSplineControlPoint* ULandscapeSplineSelection::GetAdjacentControlPointInPath(ESplineNavigationFlags Flags) const
+{
+	check(EnumHasAllFlags(Flags, ESplineNavigationFlags::ControlPointSelectModeEnabled));
+	check(FMath::CountBits(static_cast<uint8>(Flags & ESplineNavigationFlags::SelectModeMask)) == 1);
+
+	const int32 LastSelectedControlPointIndex = LinearControlPoints.Find(SelectedSplineControlPoints.Last());
+	const int32 NextIndexOffset = EnumHasAllFlags(Flags, ESplineNavigationFlags::DirectionForward) ? 1 : -1;
+	check(LastSelectedControlPointIndex != INDEX_NONE); // last selected should always be part of the linear path
+
+	if (LinearControlPoints.IsValidIndex(LastSelectedControlPointIndex + NextIndexOffset))
+	{
+		return LinearControlPoints[LastSelectedControlPointIndex + NextIndexOffset];
+	}
+	return nullptr;
+}
+
+bool ULandscapeSplineSelection::HasAdjacentSegmentInLinearPath(ESplineNavigationFlags Flags) const
+{
+	check(EnumHasAllFlags(Flags, ESplineNavigationFlags::SegmentSelectModeEnabled));
+	check(FMath::CountBits(static_cast<uint8>(Flags & ESplineNavigationFlags::SelectModeMask)) == 1);
+
+	const ULandscapeSplineSegment* AdjacentSegment = GetAdjacentSegmentInLinearPath(Flags);
+
+	if (!AdjacentSegment || !IsSelectionValidForNavigation())
+	{
+		return false;
+	}
+
+	if (EnumHasAllFlags(Flags, ESplineNavigationFlags::AddToSelection))
+	{
+		return !AdjacentSegment->IsSplineSelected();
+	}
+
+	return true;
+}
+
+bool ULandscapeSplineSelection::HasAdjacentControlPointInLinearPath(ESplineNavigationFlags Flags) const
+{
+	check(EnumHasAllFlags(Flags, ESplineNavigationFlags::ControlPointSelectModeEnabled));
+	check(FMath::CountBits(static_cast<uint8>(Flags & ESplineNavigationFlags::SelectModeMask)) == 1);
+
+	const ULandscapeSplineControlPoint* AdjacentPoint = GetAdjacentControlPointInPath(Flags);
+
+	if (!AdjacentPoint || !IsSelectionValidForNavigation())
+	{
+		return false;
+	}
+
+	if (EnumHasAllFlags(Flags, ESplineNavigationFlags::AddToSelection))
+	{
+		return !AdjacentPoint->IsSplineSelected();
+	}
+
+	return true;
+}
+
+#undef LOCTEXT_NAMESPACE
\ No newline at end of file
diff --git a/Engine/Source/Editor/LandscapeEditor/Private/LandscapeSplineSelection.h b/Engine/Source/Editor/LandscapeEditor/Private/LandscapeSplineSelection.h
new file mode 100644
--- /dev/null
+++ b/Engine/Source/Editor/LandscapeEditor/Private/LandscapeSplineSelection.h
@@ -0,0 +1,103 @@
+// Copyright Epic Games, Inc. All Rights Reserved.
+
+#pragma once
+
+#include "CoreMinimal.h"
+#include "UnrealWidgetFwd.h"
+
+#include "LandscapeSplineSelection.generated.h"
+
+class ULandscapeSplineControlPoint;
+class ULandscapeSplineSegment;
+class ULandscapeInfo;
+
+enum class ESplineNavigationFlags : uint8
+{
+	None = 0,
+	DirectionForward = 1 << 0,
+	DirectionBackward = 1 << 1,
+
+	SegmentSelectModeEnabled = 1 << 2,
+	ControlPointSelectModeEnabled = 1 << 3,
+
+	AddToSelection = 1 << 4,
+	UpdatePropertiesWindows = 1 << 5,
+
+	SelectModeMask = SegmentSelectModeEnabled | ControlPointSelectModeEnabled,
+	DirectionMask = DirectionForward | DirectionBackward
+};
+ENUM_CLASS_FLAGS(ESplineNavigationFlags);
+
+//
+// ULandscapeSplineSelection
+// 
+UCLASS(Transient, MinimalAPI)
+class ULandscapeSplineSelection : public UObject
+{
+	GENERATED_UCLASS_BODY()
+
+	// Spline Selection
+	void SelectSegment(ULandscapeSplineSegment* Segment, ESplineNavigationFlags Flags = ESplineNavigationFlags::None);
+	void SelectControlPoint(ULandscapeSplineControlPoint* ControlPoint, ESplineNavigationFlags Flags = ESplineNavigationFlags::None);
+
+	void ClearSelectedControlPoints();
+	void ClearSelectedSegments();
+	void ClearSelection();
+	void DeselectControlPoint(ULandscapeSplineControlPoint* ControlPoint, ESplineNavigationFlags Flags);
+	void DeSelectSegment(ULandscapeSplineSegment* Segment, ESplineNavigationFlags Flags);
+
+	void SelectConnected();
+	void SelectAllSplineSegments(const ULandscapeInfo& InLandscapeInfo);
+	void SelectAllControlPoints(const ULandscapeInfo& InLandscapeInfo);
+	void SelectAdjacentControlPoints();
+	void SelectAdjacentSegments();
+
+	const bool IsSegmentSelected(ULandscapeSplineSegment* Segment) const;
+	const bool IsControlPointSelected(ULandscapeSplineControlPoint* ControlPoint) const;
+
+	void UpdatePropertiesWindows() const;
+
+	// Spline Navigation (path traversal based on current selection)
+	void ResetNavigationPath();
+
+	void SelectNavigationSegment(ULandscapeSplineSegment* Segment);
+	void SelectNavigationControlPoint(ULandscapeSplineControlPoint* ControlPoint);
+	bool IsSelectionValidForNavigation() const;
+
+	ULandscapeSplineSegment* GetEndSegmentInLinearPath(ESplineNavigationFlags Flags) const;
+	ULandscapeSplineControlPoint* GetEndControlPointInLinearPath(ESplineNavigationFlags Flags) const;
+	ULandscapeSplineSegment* GetAdjacentSegmentInLinearPath(ESplineNavigationFlags Flags) const;
+	ULandscapeSplineControlPoint* GetAdjacentControlPointInPath(ESplineNavigationFlags Flags) const;
+
+	bool HasAdjacentSegmentInLinearPath(ESplineNavigationFlags Flags) const;
+	bool HasAdjacentControlPointInLinearPath(ESplineNavigationFlags Flags) const;
+
+	TArray<ULandscapeSplineControlPoint*> GetSelectedSplineControlPoints() const
+	{
+		return SelectedSplineControlPoints;
+	};
+
+	TArray<ULandscapeSplineSegment*> GetSelectedSplineSegments() const
+	{
+		return SelectedSplineSegments;
+	};
+
+private:
+	// Spline navigation
+	void BuildLinearPathFromLastSelectedPointInternal();
+	ULandscapeSplineControlPoint* GetLinearEndControlPointInternal(ULandscapeSplineControlPoint* SelectedControlPoint);
+
+	UPROPERTY(Transient)
+	TArray<TObjectPtr<ULandscapeSplineControlPoint>> SelectedSplineControlPoints;
+	UPROPERTY(Transient)
+	TArray<TObjectPtr<ULandscapeSplineSegment>> SelectedSplineSegments;
+
+	// Linear representation of the current spline based on the last selected point/segment
+	UPROPERTY(Transient)
+	TArray<TObjectPtr<ULandscapeSplineControlPoint>> LinearControlPoints;
+	UPROPERTY(Transient)
+	TArray<TObjectPtr<ULandscapeSplineSegment>> LinearSegments;
+
+	// Persistent control point mode
+	TOptional<UE::Widget::EWidgetMode> ControlPointWidgetMode;
+};
diff --git a/Engine/Source/Runtime/Landscape/Classes/LandscapeSplineControlPoint.h b/Engine/Source/Runtime/Landscape/Classes/LandscapeSplineControlPoint.h
--- a/Engine/Source/Runtime/Landscape/Classes/LandscapeSplineControlPoint.h
+++ b/Engine/Source/Runtime/Landscape/Classes/LandscapeSplineControlPoint.h
@@ -264,6 +264,9 @@ public:
 	/** Update spline points */
 	virtual void UpdateSplinePoints(bool bUpdateCollision = true, bool bUpdateAttachedSegments = true, bool bUpdateMeshLevel = false);
 
+	/** Update spline mesh selection state */
+	virtual void UpdateSplineMeshSelectionState();
+
 	/** Delete spline points */
 	virtual void DeleteSplinePoints();
 
diff --git a/Engine/Source/Runtime/Landscape/Classes/LandscapeSplineSegment.h b/Engine/Source/Runtime/Landscape/Classes/LandscapeSplineSegment.h
--- a/Engine/Source/Runtime/Landscape/Classes/LandscapeSplineSegment.h
+++ b/Engine/Source/Runtime/Landscape/Classes/LandscapeSplineSegment.h
@@ -350,7 +350,7 @@ public:
 	TArray<USplineMeshComponent*> GetLocalMeshComponents() const;
 	
 	virtual void UpdateSplinePoints(bool bUpdateCollision = true, bool bUpdateMeshLevel = false);
-
+	virtual void UpdateSplineMeshSelectionState();
 	void UpdateSplineEditorMesh();
 	virtual void DeleteSplinePoints();
 
diff --git a/Engine/Source/Runtime/Landscape/Private/LandscapeSplines.cpp b/Engine/Source/Runtime/Landscape/Private/LandscapeSplines.cpp
--- a/Engine/Source/Runtime/Landscape/Private/LandscapeSplines.cpp
+++ b/Engine/Source/Runtime/Landscape/Private/LandscapeSplines.cpp
@@ -2032,23 +2032,14 @@ void ULandscapeSplineControlPoint::GetConnectionLocationAndRotation(FName Socket
 
 void ULandscapeSplineControlPoint::SetSplineSelected(bool bInSelected)
 {
-	bSelected = bInSelected;
-	GetOuterULandscapeSplinesComponent()->MarkRenderStateDirty();
-
-	if (LocalMeshComponent != nullptr)
+	if (bSelected == bInSelected)
 	{
-		LocalMeshComponent->bSelected = bInSelected;
-		LocalMeshComponent->PushSelectionToProxy();
+		return;
 	}
+	bSelected = bInSelected;
+	GetOuterULandscapeSplinesComponent()->MarkRenderStateDirty();
 
-	auto ForeignMeshComponentsMap = GetForeignMeshComponents();
-	for (auto& ForeignMeshComponentsPair : ForeignMeshComponentsMap)
-	{
-		ULandscapeSplinesComponent* MeshComponentOuterSplines = ForeignMeshComponentsPair.Key;
-		auto* MeshComponent = ForeignMeshComponentsPair.Value;
-		MeshComponent->bSelected = bInSelected;
-		MeshComponent->PushSelectionToProxy();
-	}
+	UpdateSplineMeshSelectionState();
 }
 
 void ULandscapeSplineControlPoint::AutoCalcRotation(bool bAlwaysRotateForward)
@@ -2552,6 +2543,24 @@ void ULandscapeSplineControlPoint::UpdateSplinePoints(bool bUpdateCollision, boo
 	}
 }
 
+void ULandscapeSplineControlPoint::UpdateSplineMeshSelectionState()
+{
+	if (LocalMeshComponent != nullptr)
+	{
+		LocalMeshComponent->bSelected = bSelected;
+		LocalMeshComponent->PushSelectionToProxy();
+	}
+
+	auto ForeignMeshComponentsMap = GetForeignMeshComponents();
+	for (auto& ForeignMeshComponentsPair : ForeignMeshComponentsMap)
+	{
+		ULandscapeSplinesComponent* MeshComponentOuterSplines = ForeignMeshComponentsPair.Key;
+		auto* MeshComponent = ForeignMeshComponentsPair.Value;
+		MeshComponent->bSelected = bSelected;
+		MeshComponent->PushSelectionToProxy();
+	}
+}
+
 void ULandscapeSplineControlPoint::DeleteSplinePoints()
 {
 	Modify();
@@ -2600,6 +2609,9 @@ void ULandscapeSplineControlPoint::PostEditUndo()
 	Super::PostEditUndo();
 	bHackIsUndoingSplines = false;
 
+	// Update mesh components to make sure the selection state stays in sync
+	UpdateSplineMeshSelectionState();
+
 	ULandscapeSplinesComponent* SplineComponent = GetOuterULandscapeSplinesComponent();
 	SplineComponent->MarkRenderStateDirty();
 	SplineComponent->RequestSplineLayerUpdate();
@@ -2914,25 +2926,15 @@ void ULandscapeSplineSegment::UpdateMeshCollisionProfile(USplineMeshComponent* M
 #if WITH_EDITOR
 void ULandscapeSplineSegment::SetSplineSelected(bool bInSelected)
 {
-	bSelected = bInSelected;
-	GetOuterULandscapeSplinesComponent()->MarkRenderStateDirty();
-
-	for (auto& LocalMeshComponent : LocalMeshComponents)
+	if (bInSelected == bSelected)
 	{
-		LocalMeshComponent->bSelected = bInSelected;
-		LocalMeshComponent->PushSelectionToProxy();
+		return;
 	}
 
-	auto ForeignMeshComponentsMap = GetForeignMeshComponents();
-	for (auto& ForeignMeshComponentsPair : ForeignMeshComponentsMap)
-	{
-		ULandscapeSplinesComponent* MeshComponentOuterSplines = ForeignMeshComponentsPair.Key;
-		for (auto* ForeignMeshComponent : ForeignMeshComponentsPair.Value)
-		{
-			ForeignMeshComponent->bSelected = bInSelected;
-			ForeignMeshComponent->PushSelectionToProxy();
-		}
-	}
+	bSelected = bInSelected;
+	GetOuterULandscapeSplinesComponent()->MarkRenderStateDirty();
+
+	UpdateSplineMeshSelectionState();
 }
 
 void ULandscapeSplineSegment::AutoFlipTangents()
@@ -3540,6 +3542,26 @@ void ULandscapeSplineSegment::UpdateSplinePoints(bool bUpdateCollision, bool bUp
 	ForeignMeshComponentsMap.Empty();
 }
 
+void ULandscapeSplineSegment::UpdateSplineMeshSelectionState()
+{
+	for (auto& LocalMeshComponent : LocalMeshComponents)
+	{
+		LocalMeshComponent->bSelected = bSelected;
+		LocalMeshComponent->PushSelectionToProxy();
+	}
+
+	auto ForeignMeshComponentsMap = GetForeignMeshComponents();
+	for (auto& ForeignMeshComponentsPair : ForeignMeshComponentsMap)
+	{
+		ULandscapeSplinesComponent* MeshComponentOuterSplines = ForeignMeshComponentsPair.Key;
+		for (auto* ForeignMeshComponent : ForeignMeshComponentsPair.Value)
+		{
+			ForeignMeshComponent->bSelected = bSelected;
+			ForeignMeshComponent->PushSelectionToProxy();
+		}
+	}
+}
+
 void ULandscapeSplineSegment::UpdateSplineEditorMesh()
 {
 	ULandscapeSplinesComponent* OuterSplines = CastChecked<ULandscapeSplinesComponent>(GetOuter());
@@ -3645,6 +3667,9 @@ void ULandscapeSplineSegment::PostEditUndo()
 	Super::PostEditUndo();
 	bHackIsUndoingSplines = false;
 
+	// Update mesh components to make sure the selection state stays in sync
+	UpdateSplineMeshSelectionState();
+
 	ULandscapeSplinesComponent* SplineComponent = GetOuterULandscapeSplinesComponent();
 	SplineComponent->MarkRenderStateDirty();
 	SplineComponent->RequestSplineLayerUpdate();
