Landscape Editor - Fix spline crash on Undo after using brush tools

diff --git a/Engine/Source/Editor/LandscapeEditor/Private/LandscapeEdModeSplineTools.cpp b/Engine/Source/Editor/LandscapeEditor/Private/LandscapeEdModeSplineTools.cpp
--- a/Engine/Source/Editor/LandscapeEditor/Private/LandscapeEdModeSplineTools.cpp
+++ b/Engine/Source/Editor/LandscapeEditor/Private/LandscapeEdModeSplineTools.cpp
@@ -1654,7 +1654,10 @@ public:
 
 	void OnUndo()
 	{
-		SplineSelection->UpdatePropertiesWindows();
+		if (EdMode && EdMode->CurrentToolMode && EdMode->CurrentToolMode->CurrentToolName == "Splines" && SplineSelection)
+		{
+			SplineSelection->UpdatePropertiesWindows();
+		}
 	}
 
 	virtual void EnterTool() override
diff --git a/Engine/Source/Editor/LandscapeEditor/Private/LandscapeSplineSelection.cpp b/Engine/Source/Editor/LandscapeEditor/Private/LandscapeSplineSelection.cpp
--- a/Engine/Source/Editor/LandscapeEditor/Private/LandscapeSplineSelection.cpp
+++ b/Engine/Source/Editor/LandscapeEditor/Private/LandscapeSplineSelection.cpp
@@ -270,16 +270,13 @@ const bool ULandscapeSplineSelection::IsControlPointSelected(ULandscapeSplineCon
 
 void ULandscapeSplineSelection::UpdatePropertiesWindows() const
 {
-	if (GLevelEditorModeTools().GetActiveMode(FBuiltinEditorModes::EM_Landscape))
-	{
-		TArray<UObject*> Objects;
-		Objects.Reset(SelectedSplineControlPoints.Num() + SelectedSplineSegments.Num());
-		Algo::Copy(SelectedSplineControlPoints, Objects);
-		Algo::Copy(SelectedSplineSegments, Objects);
+	TArray<UObject*> Objects;
+	Objects.Reset(SelectedSplineControlPoints.Num() + SelectedSplineSegments.Num());
+	Algo::Copy(SelectedSplineControlPoints, Objects);
+	Algo::Copy(SelectedSplineSegments, Objects);
 
-		FPropertyEditorModule& PropertyModule = FModuleManager::Get().LoadModuleChecked<FPropertyEditorModule>(TEXT("PropertyEditor"));
-		PropertyModule.UpdatePropertyViews(Objects);
-	}
+	FPropertyEditorModule& PropertyModule = FModuleManager::Get().LoadModuleChecked<FPropertyEditorModule>(TEXT("PropertyEditor"));
+	PropertyModule.UpdatePropertyViews(Objects);
 
 	GUnrealEd->RedrawLevelEditingViewports();
 }
