Landscape Editor - Added CanEditTargetLayer to remove target layer UI warning from showing under Edit Layer category

diff --git a/Engine/Source/Editor/LandscapeEditor/Private/LandscapeEdMode.cpp b/Engine/Source/Editor/LandscapeEditor/Private/LandscapeEdMode.cpp
--- a/Engine/Source/Editor/LandscapeEditor/Private/LandscapeEdMode.cpp
+++ b/Engine/Source/Editor/LandscapeEditor/Private/LandscapeEdMode.cpp
@@ -1977,7 +1977,7 @@ bool FEdModeLandscape::InputKey(FEditorViewportClient* ViewportClient, FViewport
 					FVector HitLocation;
 					if (LandscapeMouseTrace(ViewportClient, HitLocation))
 					{
-						if (CurrentTool->AffectsEditLayers() && !CanEditLayer(&ErrorReasonOnMouseUp))
+						if (CurrentTool->AffectsEditLayers() && !CanEditTargetLayer(&ErrorReasonOnMouseUp))
 						{
 							return true;
 						}
@@ -4221,18 +4221,24 @@ bool FEdModeLandscape::CanEditLayer(FText* Reason /*=nullptr*/, const ULandscape
 		}
 	}
 
-	if (CurrentToolTarget.TargetType == ELandscapeToolTargetType::Weightmap && CurrentToolTarget.LayerInfo == nullptr && (CurrentTool && CurrentTool->GetToolName() != FName("BlueprintBrush")))
+	return true;
+}
+
+bool FEdModeLandscape::CanEditTargetLayer(FText* Reason /*=nullptr*/, const ULandscapeEditLayerBase* InLayer /*= nullptr*/)
+{
+	// Target layer is only editable when the edit layer is
+	if (!CanEditLayer(Reason, InLayer))
+	{
+		return false;
+	}
+
+	if (CurrentToolTarget.TargetType == ELandscapeToolTargetType::Weightmap && CurrentToolTarget.LayerInfo == nullptr)
 	{
 		if (Reason)
 		{
 			*Reason = NSLOCTEXT("UnrealEd", "LandscapeNeedToCreateLayerInfo", "This layer has no layer info assigned yet. You must create or assign a layer info before you can paint this layer.");
 		}
 		return false;
-		// TODO: FName to LayerInfo: do we want to create the layer info here?
-		//if (FMessageDialog::Open(EAppMsgType::YesNo, NSLOCTEXT("UnrealEd", "LandscapeNeedToCreateLayerInfo", "This layer has no layer info assigned yet. You must create or assign a layer info before you can paint this layer.")) == EAppReturnType::Yes)
-		//{
-		//	CurrentToolTarget.LandscapeInfo->LandscapeProxy->CreateLayerInfo(*CurrentToolTarget.PlaceholderLayerName.ToString());
-		//}
 	}
 	return true;
 }
diff --git a/Engine/Source/Editor/LandscapeEditor/Private/LandscapeEdMode.h b/Engine/Source/Editor/LandscapeEditor/Private/LandscapeEdMode.h
--- a/Engine/Source/Editor/LandscapeEditor/Private/LandscapeEdMode.h
+++ b/Engine/Source/Editor/LandscapeEditor/Private/LandscapeEdMode.h
@@ -583,6 +583,7 @@ public:
 	void UpdateLandscapeSplines(bool bUpdateOnlySelected = false);
 	void AutoUpdateDirtyLandscapeSplines();
 	bool CanEditLayer(FText* Reason = nullptr, const ULandscapeEditLayerBase* InLayer = nullptr);
+	bool CanEditTargetLayer(FText* Reason = nullptr, const ULandscapeEditLayerBase* InLayer = nullptr);
 
 	void AddBrushToCurrentLayer(class ALandscapeBlueprintBrushBase* InBrush);
 	void RemoveBrushFromCurrentLayer(int32 InBrushIndex);
diff --git a/Engine/Source/Editor/LandscapeEditor/Private/LandscapeEdModeBrushes.cpp b/Engine/Source/Editor/LandscapeEditor/Private/LandscapeEdModeBrushes.cpp
--- a/Engine/Source/Editor/LandscapeEditor/Private/LandscapeEdModeBrushes.cpp
+++ b/Engine/Source/Editor/LandscapeEditor/Private/LandscapeEdModeBrushes.cpp
@@ -113,7 +113,7 @@ public:
 			return false;
 		}
 
-		if (!EdMode->CanEditLayer())
+		if (!EdMode->CanEditTargetLayer())
 		{
 			return false;
 		}
diff --git a/Engine/Source/Editor/LandscapeEditor/Private/LandscapeEditorDetailCustomization_TargetLayers.cpp b/Engine/Source/Editor/LandscapeEditor/Private/LandscapeEditorDetailCustomization_TargetLayers.cpp
--- a/Engine/Source/Editor/LandscapeEditor/Private/LandscapeEditorDetailCustomization_TargetLayers.cpp
+++ b/Engine/Source/Editor/LandscapeEditor/Private/LandscapeEditorDetailCustomization_TargetLayers.cpp
@@ -246,7 +246,7 @@ void FLandscapeEditorDetailCustomization_TargetLayers::CustomizeDetails(IDetailL
 			SNew(SMultiLineEditableTextBox)
 				.IsReadOnly(true)
 				.Font(DetailBuilder.GetDetailFontBold())
-				.BackgroundColor(FAppStyle::GetColor("InfoReporting.BackgroundColor"))
+				.BackgroundColor(FAppStyle::GetColor("ErrorReporting.WarningBackgroundColor"))
 				.Text(LOCTEXT("InvalidTargetLayers_Tip", "There are target layers with no layer info asset assigned. Create a new asset, select an existing, or use the Auto-Fill button above to quickly set assets for all layers"))
 				.AutoWrapText(true)
 		];
