Landscape Editor - Fix assert when deforming non-edit layer landscape to spline
* Disables deform buttons once a landscape has a Spline Edit Layer
* Fixes assert when toggling bCanHaveLayersContent

diff --git a/Engine/Source/Editor/LandscapeEditor/Private/LandscapeEdMode.cpp b/Engine/Source/Editor/LandscapeEditor/Private/LandscapeEdMode.cpp
--- a/Engine/Source/Editor/LandscapeEditor/Private/LandscapeEdMode.cpp
+++ b/Engine/Source/Editor/LandscapeEditor/Private/LandscapeEdMode.cpp
@@ -3987,6 +3987,16 @@ bool FEdModeLandscape::HasLandscapeLayersContent() const
 	return Landscape != nullptr ? Landscape->HasLayersContent() : false;
 }
 
+bool FEdModeLandscape::HasSplinesEditLayer() const
+{
+	if (ALandscape* Landscape = GetLandscape(); Landscape != nullptr && CanHaveLandscapeLayersContent())
+	{
+		const ULandscapeEditLayerBase* SplinesLayer = Landscape->FindEditLayerOfTypeConst(ULandscapeEditLayerSplines::StaticClass());
+		return SplinesLayer != nullptr;
+	}
+	return false;
+}
+
 int32 FEdModeLandscape::GetLayerCount() const
 {
 	ALandscape* Landscape = GetLandscape();
diff --git a/Engine/Source/Editor/LandscapeEditor/Private/LandscapeEdMode.h b/Engine/Source/Editor/LandscapeEditor/Private/LandscapeEdMode.h
--- a/Engine/Source/Editor/LandscapeEditor/Private/LandscapeEdMode.h
+++ b/Engine/Source/Editor/LandscapeEditor/Private/LandscapeEdMode.h
@@ -570,6 +570,7 @@ public:
 	bool HasValidLandscapeEditLayerSelection() const;
 	bool CanHaveLandscapeLayersContent() const;
 	bool HasLandscapeLayersContent() const;
+	bool HasSplinesEditLayer() const;
 	int32 GetLayerCount() const;
 	void SetSelectedEditLayer(int32 InLayerIndex);
 	int32 GetSelectedEditLayerIndex() const;
diff --git a/Engine/Source/Editor/LandscapeEditor/Private/LandscapeEditorDetailCustomization_Layers.cpp b/Engine/Source/Editor/LandscapeEditor/Private/LandscapeEditorDetailCustomization_Layers.cpp
--- a/Engine/Source/Editor/LandscapeEditor/Private/LandscapeEditorDetailCustomization_Layers.cpp
+++ b/Engine/Source/Editor/LandscapeEditor/Private/LandscapeEditorDetailCustomization_Layers.cpp
@@ -866,9 +866,9 @@ void FLandscapeEditorCustomNodeBuilder_Layers::DeleteLayer(int32 InLayerIndex)
 			if (Result == EAppReturnType::Yes)
 			{
 				const FScopedTransaction Transaction(LOCTEXT("Landscape_Layers_Delete", "Delete Layer"));
+				// Delete layer will update the selected edit layer index
 				Landscape->DeleteLayer(InLayerIndex);
-				int32 NewLayerSelectionIndex = Landscape->GetLayerConst(InLayerIndex) ? InLayerIndex : 0;
-				OnLayerSelectionChanged(NewLayerSelectionIndex);
+				LandscapeEdMode->UpdateTargetList();
 				LandscapeEdMode->RefreshDetailPanel();
 			}
 		}
diff --git a/Engine/Source/Editor/LandscapeEditor/Private/LandscapeEditorDetailCustomization_MiscTools.cpp b/Engine/Source/Editor/LandscapeEditor/Private/LandscapeEditorDetailCustomization_MiscTools.cpp
--- a/Engine/Source/Editor/LandscapeEditor/Private/LandscapeEditorDetailCustomization_MiscTools.cpp
+++ b/Engine/Source/Editor/LandscapeEditor/Private/LandscapeEditorDetailCustomization_MiscTools.cpp
@@ -138,24 +138,44 @@ void FLandscapeEditorDetailCustomization_MiscTools::CustomizeDetails(IDetailLayo
 				.Text(LOCTEXT("Spline.ApplySplines", "Deform Landscape to Splines:"))
 			]
 		];
+
+		const FEdModeLandscape* LandscapeEdMode = GetEditorMode();
+
+		// Once a Splines Edit Layer exists, disable the buttons since spline updates happen automatically
 		ToolsCategory.AddCustomRow(LOCTEXT("ApplySplinesLabel", "Apply Splines"))
 		[
 			SNew(SHorizontalBox)
 			+ SHorizontalBox::Slot()
 			[
 				SNew(SButton)
-				.ToolTipText(LOCTEXT("Spline.ApplySplines.All.Tooltip", "Deforms and paints the landscape to fit all the landscape spline segments and points."))
+				.ToolTipText_Lambda([LandscapeEdMode]()
+				{
+					return LandscapeEdMode && LandscapeEdMode->HasSplinesEditLayer() ? LOCTEXT("Spline.SplineEditLayerEnabled", "The landscape deformation is controlled by the Splines Edit Layer. Changes made to the spline are instantly reflected to the landscape.")
+						: LOCTEXT("Spline.ApplySplines.All.Tooltip", "Deforms and paints the landscape to fit all the landscape spline segments and points.");
+				})
 				.Text(LOCTEXT("Spline.ApplySplines.All", "All Splines"))
 				.HAlign(HAlign_Center)
 				.OnClicked_Static(&FLandscapeEditorDetailCustomization_MiscTools::OnApplyAllSplinesButtonClicked)
+				.IsEnabled_Lambda([LandscapeEdMode]()
+				{
+					return LandscapeEdMode && !LandscapeEdMode->HasSplinesEditLayer();
+				})
 			]
 			+ SHorizontalBox::Slot()
 			[
 				SNew(SButton)
-				.ToolTipText(LOCTEXT("Spline.ApplySplines.Tooltip", "Deforms and paints the landscape to fit only the selected landscape spline segments and points."))
-				.Text(LOCTEXT("Spline.ApplySplines.Selected", "Only Selected"))
+				.ToolTipText_Lambda([LandscapeEdMode]()
+				{
+					return LandscapeEdMode && LandscapeEdMode->HasSplinesEditLayer() ? LOCTEXT("Spline.SplineEditLayerEnabled", "The landscape deformation is controlled by the Splines Edit Layer. Changes made to the spline are instantly reflected to the landscape.")
+						: LOCTEXT("Spline.ApplySplines.Segments.Tooltip", "Deforms and paints the landscape to fit only the selected landscape spline segments.");
+				})
+				.Text(LOCTEXT("Spline.ApplySplines.Selected", "Selected Segments"))
 				.HAlign(HAlign_Center)
 				.OnClicked_Static(&FLandscapeEditorDetailCustomization_MiscTools::OnApplySelectedSplinesButtonClicked)
+				.IsEnabled_Lambda([LandscapeEdMode]()
+				{
+					return LandscapeEdMode && !LandscapeEdMode->HasSplinesEditLayer();
+				})
 			]
 		];
 
diff --git a/Engine/Source/Runtime/Landscape/Private/LandscapeEditLayers.cpp b/Engine/Source/Runtime/Landscape/Private/LandscapeEditLayers.cpp
--- a/Engine/Source/Runtime/Landscape/Private/LandscapeEditLayers.cpp
+++ b/Engine/Source/Runtime/Landscape/Private/LandscapeEditLayers.cpp
@@ -2554,13 +2554,15 @@ void ALandscape::ToggleCanHaveLayersContent()
 	}
 	check(InLayerUpdateCount == 0);
 
+	// TODO: All the toggle functions (handling edit layers) should ideally run before updating bCanHaveLayersContent
 	bCanHaveLayersContent = !bCanHaveLayersContent;
 
 	if (!bCanHaveLayersContent)
 	{
+		check(!LandscapeEditLayers.IsEmpty());
 		ReleaseLayersRenderingResource();
 		DeleteLayers();
-		SetSelectedEditLayerIndex(INDEX_NONE);
+		check(SelectedEditLayerIndex == INDEX_NONE);
 	}
 	else
 	{
@@ -12708,6 +12710,8 @@ TArray<const FLandscapeLayer*> ALandscape::GetLayersOfTypeConst(const TSubclassO
 
 void ALandscape::DeleteLayers()
 {
+	SelectedEditLayerIndex = INDEX_NONE;
+
 	for (int32 LayerIndex = LandscapeEditLayers.Num() - 1; LayerIndex >= 0; --LayerIndex)
 	{
 		DeleteLayer(LayerIndex);
@@ -12716,7 +12720,10 @@ void ALandscape::DeleteLayers()
 
 bool ALandscape::DeleteLayer(int32 InLayerIndex)
 {
-	ensure(HasLayersContent());
+	// ToggleCanHaveLayers sets the bCanHaveLayers flag before calling DeleteLayers()
+	// When toggling from an layer to a non-layer state, HasLayersContent becomes false but this function still needs to run
+	// ensure Selected Index is INDEX_NONE in this case
+	ensure(HasLayersContent() || SelectedEditLayerIndex == INDEX_NONE);
 
 	// Detect any attempt to call this in the middle of UpdateLayersContent.  If called from blueprint, log an error and return early instead of asserting.
 	if (InLayerUpdateCount > 0 && UE::Landscape::Private::InBPCallstack())
@@ -12734,8 +12741,9 @@ bool ALandscape::DeleteLayer(int32 InLayerIndex)
 
 	Modify();
 
-	// If this edit layer is currently selected, select the one below it or the first layer
-	if (GetSelectedEditLayerIndex() == InLayerIndex)
+	// If the layer to delete is below the current selected layer index, shift the selected index down
+	// Ignore index update when toggling from an layer to a non-layer state 
+	if (HasLayersContent() && SelectedEditLayerIndex >= InLayerIndex)
 	{
 		if (LandscapeEditLayers.IsValidIndex(InLayerIndex - 1))
 		{
diff --git a/Engine/Source/Runtime/Landscape/Private/LandscapeSplineRaster.cpp b/Engine/Source/Runtime/Landscape/Private/LandscapeSplineRaster.cpp
--- a/Engine/Source/Runtime/Landscape/Private/LandscapeSplineRaster.cpp
+++ b/Engine/Source/Runtime/Landscape/Private/LandscapeSplineRaster.cpp
@@ -291,11 +291,15 @@ void RasterizeHeight(int32& MinX, int32& MinY, int32& MaxX, int32& MaxY, FLandsc
 	const ALandscape* Landscape = LandscapeEdit.GetTargetLandscape();
 	check(Landscape != nullptr);
 
-	const ULandscapeEditLayerBase* EditingLayer = Landscape->GetEditLayerConst(Landscape->GetEditingLayer());
-	check(EditingLayer != nullptr);
-	const ULandscapeEditLayerBase* SplinesLayer = Landscape->FindEditLayerOfTypeConst(ULandscapeEditLayerSplines::StaticClass());
-	bool bIsEditingLayerReservedForSplines = (EditingLayer != nullptr) && (SplinesLayer == EditingLayer);
-	check(!bIsEditingLayerReservedForSplines || (SplinesLayer->GetBlendMode() == LSBM_AlphaBlend));
+	bool bIsEditingLayerReservedForSplines = false;
+	if (Landscape->CanHaveLayersContent())
+	{
+		const ULandscapeEditLayerBase* EditingLayer = Landscape->GetEditLayerConst(Landscape->GetEditingLayer());
+		check(EditingLayer != nullptr);
+		const ULandscapeEditLayerBase* SplinesLayer = Landscape->FindEditLayerOfTypeConst(ULandscapeEditLayerSplines::StaticClass());
+		bIsEditingLayerReservedForSplines = (SplinesLayer == EditingLayer);
+		check(!bIsEditingLayerReservedForSplines || (SplinesLayer->GetBlendMode() == LSBM_AlphaBlend));
+	}
 
 	TArray<uint16> HeightAlphaBlendData;
 	TArray<uint8> HeightFlagsData;
@@ -303,6 +307,8 @@ void RasterizeHeight(int32& MinX, int32& MinY, int32& MaxX, int32& MaxY, FLandsc
 	TArray<uint8>* HeightFlagsDataPtr = nullptr;
 
 	bool bCalculateNormals = true;
+
+	// when using spline edit layer there is additional alpha blending data
 	if (bIsEditingLayerReservedForSplines)
 	{
 		bCalculateNormals = false;
