Display error message instead of crashing when failing to import an image file to a landscape layer.

diff --git a/Engine/Source/Editor/LandscapeEditor/Private/LandscapeImageFileCache.cpp b/Engine/Source/Editor/LandscapeEditor/Private/LandscapeImageFileCache.cpp
--- a/Engine/Source/Editor/LandscapeEditor/Private/LandscapeImageFileCache.cpp
+++ b/Engine/Source/Editor/LandscapeEditor/Private/LandscapeImageFileCache.cpp
@@ -69,13 +69,31 @@ FLandscapeFileInfo FLandscapeImageFileCache::FindImage(const TCHAR* InImageFilen
 		FLandscapeFileResolution ExpectedResolution = FileInfo.PossibleResolutions[0];
 		FLandscapeImportData<T> ImportData = FileFormat->Import(InImageFilename, ExpectedResolution);
 
+		if (ImportData.ResultCode == ELandscapeImportResult::Error)
+		{
+			// Validate didn't error, but Import did.  Return the Import error.
+			Result.ResultCode = ELandscapeImportResult::Error;
+			Result.ErrorMessage = ImportData.ErrorMessage;
+			return Result;
+		}
+		else if (ImportData.ResultCode == ELandscapeImportResult::Warning)
+		{
+			// New warning from Import.
+			NewImageData.Result = ELandscapeImportResult::Warning;
+			NewImageData.ErrorMessage = ImportData.ErrorMessage;
+		}
+		else
+		{
+			// No new warning or error, use the warning from the validate step, if there was one.
+			NewImageData.Result = FileInfo.ResultCode;
+			NewImageData.ErrorMessage = FileInfo.ErrorMessage;
+		}
+
 		const int32 BufferSize = ImportData.Data.Num() * sizeof(T);
 		NewImageData.Data = TSharedPtr<TArray<uint8>>(new TArray<uint8>());
 		NewImageData.Data->SetNumUninitialized(BufferSize);
 		FMemory::Memcpy(NewImageData.Data->GetData(), ImportData.Data.GetData(), BufferSize);
 		NewImageData.Resolution = FIntPoint(ExpectedResolution.Width, ExpectedResolution.Height);
-		NewImageData.Result = FileInfo.ResultCode;
-		NewImageData.ErrorMessage = FileInfo.ErrorMessage;
 		NewImageData.BytesPerPixel = BufferSize / (ExpectedResolution.Width * ExpectedResolution.Height);
 	}
 	else
@@ -88,8 +106,8 @@ FLandscapeFileInfo FLandscapeImageFileCache::FindImage(const TCHAR* InImageFilen
 	Add<T>(FString(InImageFilename), OutImageData);
 
 	Result.PossibleResolutions.Add(FLandscapeFileResolution(OutImageData.Resolution.X, OutImageData.Resolution.Y));
-	Result.ResultCode = FileInfo.ResultCode;
-	Result.ErrorMessage = FileInfo.ErrorMessage;
+	Result.ResultCode = NewImageData.Result;
+	Result.ErrorMessage = NewImageData.ErrorMessage;
 
 	return Result;
 }
