Landscape Editor - Disable SLandscapeAssetThumbnail Tooltips when the asset is null (Sculpt, Visibility, Target layers with no layer info asset)

diff --git a/Engine/Source/Editor/LandscapeEditor/Private/SLandscapeEditor.cpp b/Engine/Source/Editor/LandscapeEditor/Private/SLandscapeEditor.cpp
--- a/Engine/Source/Editor/LandscapeEditor/Private/SLandscapeEditor.cpp
+++ b/Engine/Source/Editor/LandscapeEditor/Private/SLandscapeEditor.cpp
@@ -56,13 +56,22 @@ void SLandscapeAssetThumbnail::Construct(const FArguments& InArgs, UObject* Asse
 		AssetThumbnailConfig.AllowAssetStatusThumbnailOverlay = false;
 		AssetThumbnailConfig.ThumbnailLabel = EThumbnailLabel::Type::NoLabel;
 	}
+
+	TSharedRef<SWidget> AssetThumbnailWidget = AssetThumbnail->MakeThumbnailWidget(AssetThumbnailConfig);
+
+	// Disable the tooltip instead of showing a "None" tooltip for invalid assets
+	if (Asset == nullptr)
+	{
+		AssetThumbnailWidget->SetToolTip(TSharedPtr<IToolTip>());
+	}
+
 	ChildSlot
 	[
 		SNew(SBox)
 		.WidthOverride(static_cast<float>(ThumbnailSize.X))
 		.HeightOverride(static_cast<float>(ThumbnailSize.Y))
 		[
-			AssetThumbnail->MakeThumbnailWidget(AssetThumbnailConfig)
+			AssetThumbnailWidget
 		]
 	];
 
