Landscape PNG loader replace with modern FImage based loader
works for all image formats now
avoids use of deprecated color conversions inside ImageWrapper load/save code

diff --git a/Engine/Source/Editor/LandscapeEditor/Private/LandscapeFileFormatPng.cpp b/Engine/Source/Editor/LandscapeEditor/Private/LandscapeFileFormatPng.cpp
--- a/Engine/Source/Editor/LandscapeEditor/Private/LandscapeFileFormatPng.cpp
+++ b/Engine/Source/Editor/LandscapeEditor/Private/LandscapeFileFormatPng.cpp
@@ -9,8 +9,8 @@
 #include "HAL/FileManager.h"
 #include "HAL/PlatformCrt.h"
 #include "HAL/UnrealMemory.h"
-#include "IImageWrapper.h"
-#include "IImageWrapperModule.h"
+#include "ImageCore.h"
+#include "ImageUtils.h"
 #include "Internationalization/Internationalization.h"
 #include "Internationalization/Text.h"
 #include "Math/Vector.h"
@@ -24,6 +24,12 @@
 
 FLandscapeHeightmapFileFormat_Png::FLandscapeHeightmapFileFormat_Png()
 {
+	// @todo : can support all image extensions here
+	//	since changing the loader to use FImageUtils
+	//	this is no longer png specific, can load all image formats here
+	//	 (eg. EXR and DDS might be useful)
+	//	ideally FLandscapeHeightmapFileFormat_Png should be renamed too
+	//	 (note that only PNG, EXR and DDS support 16-bit write ; TIFF can read 16 but not write)
 	FileTypeInfo.Description = LOCTEXT("FileFormatPng_HeightmapDesc", "Heightmap .png files");
 	FileTypeInfo.Extensions.Add(".png");
 	FileTypeInfo.bSupportsExport = true;
@@ -33,255 +39,217 @@ FLandscapeFileInfo FLandscapeHeightmapFileFormat_Png::Validate(const TCHAR* Heig
 {
 	FLandscapeFileInfo Result;
 
-	TArray64<uint8> ImportData;
-	if (!FFileHelper::LoadFileToArray(ImportData, HeightmapFilename, FILEREAD_Silent))
+	FImage Image;
+	if ( !FImageUtils::LoadImage(HeightmapFilename,Image) )
 	{
 		Result.ResultCode = ELandscapeImportResult::Error;
 		Result.ErrorMessage = LOCTEXT("Import_HeightmapFileReadError", "Error reading heightmap file");
+		return Result;
 	}
-	else
+
+	// SizeXY are S32
+	if ( Image.SizeX <= 0 || Image.SizeY <= 0)
+	{
+		Result.ResultCode = ELandscapeImportResult::Error;
+		Result.ErrorMessage = LOCTEXT("Import_HeightmapFileCorruptPng", "The heightmap file cannot be read (corrupt png?)");
+		return Result;
+	}
+
+	FLandscapeFileResolution ImportResolution;
+	ImportResolution.Width  = Image.SizeX;
+	ImportResolution.Height = Image.SizeY;
+	Result.PossibleResolutions.Add(ImportResolution);
+	
+	if ( ERawImageFormat::NumChannels(Image.Format) != 1 )
+	{
+		// @todo : could run UE::TextureUtilitiesCommon::AutoDetectAndChangeGrayScale here
+		//	no need to complain if it's an RGB but actually contains gray colors
+
+		Result.ResultCode = ELandscapeImportResult::Warning;
+		Result.ErrorMessage = LOCTEXT("Import_HeightmapFileColorPng", "The imported layer is not Grayscale. Results in-Editor will not be consistent with the source file.");
+	}
+
+	if ( Result.ResultCode != ELandscapeImportResult::Warning && IsU8Channels(Image.Format) )
 	{
-		IImageWrapperModule& ImageWrapperModule = FModuleManager::LoadModuleChecked<IImageWrapperModule>("ImageWrapper");
-		TSharedPtr<IImageWrapper> ImageWrapper = ImageWrapperModule.CreateImageWrapper(EImageFormat::PNG);
-
-		if (!ImageWrapper->SetCompressed(ImportData.GetData(), ImportData.Num()) 
-			|| ImageWrapper->GetWidth() <= 0 
-			|| ImageWrapper->GetHeight() <= 0)
-		{
-			Result.ResultCode = ELandscapeImportResult::Error;
-			Result.ErrorMessage = LOCTEXT("Import_HeightmapFileCorruptPng", "The heightmap file cannot be read (corrupt png?)");
-		}
-		else if ((ImageWrapper->GetWidth() > MAX_int32)
-			|| (ImageWrapper->GetHeight() > MAX_int32)
-			|| (ImageWrapper->GetWidth() > MAX_int64 / ImageWrapper->GetHeight()))	// the total pixel count should fit in an int64 to avoid overflow issues
-		{
-			Result.ResultCode = ELandscapeImportResult::Error;
-			Result.ErrorMessage = LOCTEXT("Import_HeightmapFileTooLargePng", "The heightmap file is too large to load");
-		}
-		else
-		{
-			FLandscapeFileResolution ImportResolution;
-			ImportResolution.Width = static_cast<uint32>(ImageWrapper->GetWidth());
-			ImportResolution.Height = static_cast<uint32>(ImageWrapper->GetHeight());
-			Result.PossibleResolutions.Add(ImportResolution);
-
-			if (ImageWrapper->GetFormat() != ERGBFormat::Gray)
-			{
-				Result.ResultCode = ELandscapeImportResult::Warning;
-				Result.ErrorMessage = LOCTEXT("Import_HeightmapFileColorPng", "The imported layer is not Grayscale. Results in-Editor will not be consistent with the source file.");
-			}
-			else if (ImageWrapper->GetBitDepth() != 16)
-			{
-				Result.ResultCode = ELandscapeImportResult::Warning;
-				Result.ErrorMessage = LOCTEXT("Import_HeightmapFileLowBitDepth", "The heightmap file appears to be an 8-bit png, 16-bit is preferred. The import *can* continue, but the result may be lower quality than desired.");
-			}
-			
-		}
+		Result.ResultCode = ELandscapeImportResult::Warning;
+		Result.ErrorMessage = LOCTEXT("Import_HeightmapFileLowBitDepth", "The heightmap file appears to be an 8-bit png, 16-bit is preferred. The import *can* continue, but the result may be lower quality than desired.");
 	}
+	
 
 	// possible todo: support sCAL (XY scale) and pCAL (Z scale) png chunks for filling out Result.DataScale
 	// I don't know if any heightmap generation software uses these or not
 	// if we support their import we should make the exporter write them too
 
+	// @todo : Efficiency : "Image" is just discarded
+	//	 could get the image info without actually reading the whole thing?
+	//	 or don't discard if we will wind up loading later?
+	// the calling code does :
+	//const FLandscapeFileInfo FileInfo = FileFormat->Validate(InImageFilename);
+	//FLandscapeImportData<T> ImportData = FileFormat->Import(InImageFilename, ExpectedResolution);
+	//	so the image is loaded twice
+
 	return Result;
 }
 
 FLandscapeImportData<uint16> FLandscapeHeightmapFileFormat_Png::Import(const TCHAR* HeightmapFilename, FName LayerName, FLandscapeFileResolution ExpectedResolution) const
 {
 	FLandscapeImportData<uint16> Result;
-
-	TArray<uint8> TempData;
-	if (!FFileHelper::LoadFileToArray(TempData, HeightmapFilename, FILEREAD_Silent))
+	
+	FImage Image;
+	if ( !FImageUtils::LoadImage(HeightmapFilename,Image) )
 	{
 		Result.ResultCode = ELandscapeImportResult::Error;
 		Result.ErrorMessage = LOCTEXT("Import_HeightmapFileReadError", "Error reading heightmap file");
+		return Result;
 	}
-	else
+
+	if (Image.SizeX != ExpectedResolution.Width || Image.SizeY != ExpectedResolution.Height)
 	{
-		IImageWrapperModule& ImageWrapperModule = FModuleManager::LoadModuleChecked<IImageWrapperModule>("ImageWrapper");
-		TSharedPtr<IImageWrapper> ImageWrapper = ImageWrapperModule.CreateImageWrapper(EImageFormat::PNG);
-
-		const int64 Width = ExpectedResolution.Width;
-		const int64 Height = ExpectedResolution.Height;
-
-		if (!ImageWrapper->SetCompressed(TempData.GetData(), TempData.Num()))
-		{
-			Result.ResultCode = ELandscapeImportResult::Error;
-			Result.ErrorMessage = LOCTEXT("Import_HeightmapFileCorruptPng", "The heightmap file cannot be read (corrupt png?)");
-		}
-		else if (ImageWrapper->GetWidth() != Width || ImageWrapper->GetHeight() != Height)
-		{
-			Result.ResultCode = ELandscapeImportResult::Error;
-			Result.ErrorMessage = LOCTEXT("Import_HeightmapResolutionMismatch", "The heightmap file's resolution does not match the requested resolution");
-		}
-		else if (Width > MAX_int64 / Height) // total pixel count must fit in an int64
-		{
-			Result.ResultCode = ELandscapeImportResult::Error;
-			Result.ErrorMessage = LOCTEXT("Import_HeightmapFileTooLargePng", "The heightmap file is too large to load");
-		}
-		else
-		{
-			if (ImageWrapper->GetFormat() != ERGBFormat::Gray)
-			{
-				Result.ResultCode = ELandscapeImportResult::Warning;
-				Result.ErrorMessage = LOCTEXT("Import_HeightmapFileColorPng", "The imported layer is not Grayscale. Results in-Editor will not be consistent with the source file.");
-			}
-			else if (ImageWrapper->GetBitDepth() != 16)
-			{
-				Result.ResultCode = ELandscapeImportResult::Warning;
-				Result.ErrorMessage = LOCTEXT("Import_HeightmapFileLowBitDepth", "The heightmap file appears to be an 8-bit png, 16-bit is preferred. The import *can* continue, but the result may be lower quality than desired.");
-			}
-
-			const int64 TotalPixels = Width * Height;
-
-			TArray64<uint8> RawData;
-			if (ImageWrapper->GetBitDepth() <= 8)
-			{
-				if (!ImageWrapper->GetRaw(ERGBFormat::Gray, 8, RawData)
-					|| (RawData.Num() != TotalPixels))
-				{
-					Result.ResultCode = ELandscapeImportResult::Error;
-					Result.ErrorMessage = LOCTEXT("Import_HeightmapFileCorruptPng", "The heightmap file cannot be read (corrupt png?)");
-				}
-				else
-				{
-					Result.Data.Empty(TotalPixels);
-					Algo::Transform(RawData, Result.Data, [](uint8 Value) { return static_cast<uint16>(Value * 0x101); }); // Expand to 16-bit
-				}
-			}
-			else
-			{
-				if (!ImageWrapper->GetRaw(ERGBFormat::Gray, 16, RawData)
-					|| (RawData.Num()/2 != TotalPixels))
-				{
-					Result.ResultCode = ELandscapeImportResult::Error;
-					Result.ErrorMessage = LOCTEXT("Import_HeightmapFileCorruptPng", "The heightmap file cannot be read (corrupt png?)");
-				}
-				else
-				{
-					Result.Data.Empty(TotalPixels);
-					Result.Data.AddUninitialized(TotalPixels);
-					const uint64 TotalBytes = static_cast<uint64>(TotalPixels) * 2;
-					FMemory::Memcpy(Result.Data.GetData(), RawData.GetData(), TotalBytes);
-				}
-			}
-		}
+		Result.ResultCode = ELandscapeImportResult::Error;
+		Result.ErrorMessage = LOCTEXT("Import_HeightmapResolutionMismatch", "The heightmap file's resolution does not match the requested resolution");
+		return Result;
 	}
 
+	// by default 8-bit source images come in as SRGB but we want to treat them as Linear here
+	Image.GammaSpace = EGammaSpace::Linear;
+
+	FImageView G16Image;
+	G16Image.SizeX = Image.SizeX;
+	G16Image.SizeY = Image.SizeY;
+	G16Image.NumSlices = 1;
+	G16Image.Format = ERawImageFormat::G16;
+	G16Image.GammaSpace = EGammaSpace::Linear;
+	
+	Result.Data.SetNumUninitialized( G16Image.GetNumPixels() );
+	G16Image.RawData = Result.Data.GetData();
+
+	// blit into Result.Data with conversion to G16 :
+	FImageCore::CopyImage(Image,G16Image);
+
 	return Result;
 }
 
 void FLandscapeHeightmapFileFormat_Png::Export(const TCHAR* HeightmapFilename, FName LayerName, TArrayView<const uint16> Data, FLandscapeFileResolution DataResolution, FVector Scale) const
 {
-	IImageWrapperModule& ImageWrapperModule = FModuleManager::LoadModuleChecked<IImageWrapperModule>("ImageWrapper");
-	TSharedPtr<IImageWrapper> ImageWrapper = ImageWrapperModule.CreateImageWrapper(EImageFormat::PNG);
-
-	if (ImageWrapper->SetRaw(Data.GetData(), Data.Num() * 2, DataResolution.Width, DataResolution.Height, ERGBFormat::Gray, 16))
-	{
-		const TArray64<uint8> TempData = ImageWrapper->GetCompressed();
-		FFileHelper::SaveArrayToFile(TempData, HeightmapFilename);
-	}
+	FImageView G16Image;
+	G16Image.SizeX = DataResolution.Width;
+	G16Image.SizeY = DataResolution.Height;
+	G16Image.NumSlices = 1;
+	G16Image.Format = ERawImageFormat::G16;
+	G16Image.GammaSpace = EGammaSpace::Linear;
+	G16Image.RawData = (void *)Data.GetData();
+	check( Data.NumBytes() >= (SIZE_T)G16Image.GetImageSizeBytes() );
+
+	FImageUtils::SaveImageByExtension(HeightmapFilename, G16Image);
+	// note: no error return
+	//	not all image formats support 16-bit, they will convert to 8-bit if necessary
 }
 
 //////////////////////////////////////////////////////////////////////////
 
 FLandscapeWeightmapFileFormat_Png::FLandscapeWeightmapFileFormat_Png()
 {
+	// @todo : support more than .png here
+	//	 Weightmps are 8 bit so all formats are fully supported
 	FileTypeInfo.Description = LOCTEXT("FileFormatPng_WeightmapDesc", "Layer .png files");
 	FileTypeInfo.Extensions.Add(".png");
 	FileTypeInfo.bSupportsExport = true;
 }
 
+// very similar to FLandscapeHeightmapFileFormat_Png::Validate
+
 FLandscapeFileInfo FLandscapeWeightmapFileFormat_Png::Validate(const TCHAR* WeightmapFilename, FName LayerName) const
 {
 	FLandscapeFileInfo Result;
-
-	TArray64<uint8> ImportData;
-	if (!FFileHelper::LoadFileToArray(ImportData, WeightmapFilename, FILEREAD_Silent))
+	
+	FImage Image;
+	if ( !FImageUtils::LoadImage(WeightmapFilename,Image) )
 	{
 		Result.ResultCode = ELandscapeImportResult::Error;
 		Result.ErrorMessage = LOCTEXT("Import_LayerFileReadError", "Error reading layer file");
+		return Result;
+	}
+
+	// SizeXY are S32
+	if ( Image.SizeX <= 0 || Image.SizeY <= 0)
+	{
+		Result.ResultCode = ELandscapeImportResult::Error;
+		Result.ErrorMessage = LOCTEXT("Import_LayerCorruptPng", "The layer file cannot be read (corrupt png?)");
+		return Result;
 	}
-	else
+
+	if ( ERawImageFormat::NumChannels(Image.Format) != 1 )
 	{
-		IImageWrapperModule& ImageWrapperModule = FModuleManager::LoadModuleChecked<IImageWrapperModule>("ImageWrapper");
-		TSharedPtr<IImageWrapper> ImageWrapper = ImageWrapperModule.CreateImageWrapper(EImageFormat::PNG);
-
-		if (!ImageWrapper->SetCompressed(ImportData.GetData(), ImportData.Num()))
-		{
-			Result.ResultCode = ELandscapeImportResult::Error;
-			Result.ErrorMessage = LOCTEXT("Import_LayerCorruptPng", "The layer file cannot be read (corrupt png?)");
-		}
-		else
-		{
-			if (ImageWrapper->GetFormat() != ERGBFormat::Gray)
-			{
-				Result.ResultCode = ELandscapeImportResult::Warning;
-				Result.ErrorMessage = LOCTEXT("Import_LayerColorPng", "The imported layer is not Grayscale. Results in-Editor will not be consistent with the source file.");
-			}
-			FLandscapeFileResolution ImportResolution;
-			ImportResolution.Width = ImageWrapper->GetWidth();
-			ImportResolution.Height = ImageWrapper->GetHeight();
-			Result.PossibleResolutions.Add(ImportResolution);
-		}
+		// @todo : could run UE::TextureUtilitiesCommon::AutoDetectAndChangeGrayScale here
+		//	no need to complain if it's an RGB but actually contains gray colors
+
+		Result.ResultCode = ELandscapeImportResult::Warning;
+		Result.ErrorMessage = LOCTEXT("Import_LayerColorPng", "The imported layer is not Grayscale. Results in-Editor will not be consistent with the source file.");
 	}
 
+	FLandscapeFileResolution ImportResolution;
+	ImportResolution.Width = Image.SizeX;
+	ImportResolution.Height = Image.SizeY;
+	Result.PossibleResolutions.Add(ImportResolution);
+	
+	// @todo : Efficiency : "Image" is just discarded
+	//	 could get the image info without actually reading the whole thing?
+	//	 or don't discard if we will wind up loading later?
+
 	return Result;
 }
 
 FLandscapeImportData<uint8> FLandscapeWeightmapFileFormat_Png::Import(const TCHAR* WeightmapFilename, FName LayerName, FLandscapeFileResolution ExpectedResolution) const
 {
 	FLandscapeImportData<uint8> Result;
-
-	TArray64<uint8> TempData;
-	if (!FFileHelper::LoadFileToArray(TempData, WeightmapFilename, FILEREAD_Silent))
+	
+	FImage Image;
+	if ( !FImageUtils::LoadImage(WeightmapFilename,Image) )
 	{
 		Result.ResultCode = ELandscapeImportResult::Error;
 		Result.ErrorMessage = LOCTEXT("Import_LayerFileReadError", "Error reading layer file");
+		return Result;
 	}
-	else
+	
+	if (Image.SizeX != ExpectedResolution.Width || Image.SizeY != ExpectedResolution.Height)
 	{
-		IImageWrapperModule& ImageWrapperModule = FModuleManager::LoadModuleChecked<IImageWrapperModule>("ImageWrapper");
-		TSharedPtr<IImageWrapper> ImageWrapper = ImageWrapperModule.CreateImageWrapper(EImageFormat::PNG);
-
-		if (!ImageWrapper->SetCompressed(TempData.GetData(), TempData.Num()))
-		{
-			Result.ResultCode = ELandscapeImportResult::Error;
-			Result.ErrorMessage = LOCTEXT("Import_LayerCorruptPng", "The layer file cannot be read (corrupt png?)");
-		}
-		else if (ImageWrapper->GetWidth() != ExpectedResolution.Width || ImageWrapper->GetHeight() != ExpectedResolution.Height)
-		{
-			Result.ResultCode = ELandscapeImportResult::Error;
-			Result.ErrorMessage = LOCTEXT("Import_LayerResolutionMismatch", "The layer file's resolution does not match the requested resolution");
-		}
-		else if (!ImageWrapper->GetRaw(ERGBFormat::Gray, 8, Result.Data))
-		{
-			Result.ResultCode = ELandscapeImportResult::Error;
-			Result.ErrorMessage = LOCTEXT("Import_LayerCorruptPng", "The layer file cannot be read (corrupt png?)");
-		}
-		else
-		{
-			if (ImageWrapper->GetFormat() != ERGBFormat::Gray)
-			{
-				Result.ResultCode = ELandscapeImportResult::Warning;
-				Result.ErrorMessage = LOCTEXT("Import_LayerColorPng", "The imported layer is not Grayscale. Results in-Editor will not be consistent with the source file.");
-			}
-		}
+		Result.ResultCode = ELandscapeImportResult::Error;
+		Result.ErrorMessage = LOCTEXT("Import_LayerResolutionMismatch", "The layer file's resolution does not match the requested resolution");
+		return Result;
 	}
+	
+	// by default 8-bit source images come in as SRGB but we want to treat them as Linear here
+	Image.GammaSpace = EGammaSpace::Linear;
+
+	FImageView G8Image;
+	G8Image.SizeX = Image.SizeX;
+	G8Image.SizeY = Image.SizeY;
+	G8Image.NumSlices = 1;
+	G8Image.Format = ERawImageFormat::G8;
+	G8Image.GammaSpace = EGammaSpace::Linear;
+
+	Result.Data.SetNumUninitialized( G8Image.GetNumPixels() );
+	G8Image.RawData = Result.Data.GetData();
+	
+	// blit into Result.Data with conversion to G8 :
+	FImageCore::CopyImage(Image,G8Image);
 
 	return Result;
 }
 
 void FLandscapeWeightmapFileFormat_Png::Export(const TCHAR* WeightmapFilename, FName LayerName, TArrayView<const uint8> Data, FLandscapeFileResolution DataResolution, FVector Scale) const
 {
-	IImageWrapperModule& ImageWrapperModule = FModuleManager::LoadModuleChecked<IImageWrapperModule>("ImageWrapper");
-	TSharedPtr<IImageWrapper> ImageWrapper = ImageWrapperModule.CreateImageWrapper(EImageFormat::PNG);
-
-	if (ImageWrapper->SetRaw(Data.GetData(), Data.Num(), DataResolution.Width, DataResolution.Height, ERGBFormat::Gray, 8))
-	{
-		const TArray64<uint8> TempData = ImageWrapper->GetCompressed();
-		FFileHelper::SaveArrayToFile(TempData, WeightmapFilename);
-	}
+	FImageView G8Image;
+	G8Image.SizeX = DataResolution.Width;
+	G8Image.SizeY = DataResolution.Height;
+	G8Image.NumSlices = 1;
+	G8Image.Format = ERawImageFormat::G8;
+	G8Image.GammaSpace = EGammaSpace::Linear;
+	G8Image.RawData = (void *)Data.GetData();
+	check( Data.NumBytes() >= (SIZE_T)G8Image.GetImageSizeBytes() );
+
+	FImageUtils::SaveImageByExtension(WeightmapFilename, G8Image);
+	// note: no error return
 }
 
 #undef LOCTEXT_NAMESPACE
diff --git a/Engine/Source/Runtime/ImageCore/Private/ImageCore.cpp b/Engine/Source/Runtime/ImageCore/Private/ImageCore.cpp
--- a/Engine/Source/Runtime/ImageCore/Private/ImageCore.cpp
+++ b/Engine/Source/Runtime/ImageCore/Private/ImageCore.cpp
@@ -1449,6 +1449,31 @@ IMAGECORE_API int64 ERawImageFormat::GetBytesPerPixel(Type Format)
 }
 
 
+IMAGECORE_API int ERawImageFormat::NumChannels(Type Format)
+{
+	switch(Format)
+	{
+	case ERawImageFormat::G8: return 1;
+	case ERawImageFormat::BGRA8: return 4;
+	case ERawImageFormat::BGRE8: return 3; // <-
+	case ERawImageFormat::RGBA16: return 4;
+	case ERawImageFormat::RGBA16F: return 4;
+	case ERawImageFormat::RGBA32F: return 4;
+	case ERawImageFormat::G16: return 1;
+	case ERawImageFormat::R16F: return 1;
+	case ERawImageFormat::R32F: return 1;
+	default:
+		check(0);
+		return 0;
+	}
+}
+
+IMAGECORE_API bool ERawImageFormat::IsU8Channels(Type Format)
+{
+	// "BGRE8" is 8-bit but not *U8*
+	return Format == G8 || Format == BGRA8;
+}
+
 IMAGECORE_API bool ERawImageFormat::IsHDR(Type Format)
 {
 	return Format == RGBA16F || Format == RGBA32F || Format == R16F || Format == R32F || Format == BGRE8;
diff --git a/Engine/Source/Runtime/ImageCore/Public/ImageCore.h b/Engine/Source/Runtime/ImageCore/Public/ImageCore.h
--- a/Engine/Source/Runtime/ImageCore/Public/ImageCore.h
+++ b/Engine/Source/Runtime/ImageCore/Public/ImageCore.h
@@ -77,6 +77,9 @@ namespace ERawImageFormat
 	IMAGECORE_API const FUtf8StringView GetNameView(Type Format);
 	IMAGECORE_API bool GetFormatFromString(FUtf8StringView InString, Type& OutFormat);
 	
+	// IsU8Channels is not true for BGRE8, it means channels of uint8 in [0,255]
+	IMAGECORE_API bool IsU8Channels(Type Format);
+	IMAGECORE_API int NumChannels(Type Format);
 	IMAGECORE_API bool IsHDR(Type Format);
 	IMAGECORE_API bool HasAlphaChannel(Type Format);
 	
@@ -89,8 +92,8 @@ namespace ERawImageFormat
 		return GetOnePixelLinear(PixelData,Format,bSRGB ? EGammaSpace::sRGB : EGammaSpace::Linear);
 	}
 
-	// G8 and BGRA8 are affected by Gamma
-	//	16/32 is NOT
+	// G8 and BGRA8 are affected by Gamma, 16/32 is NOT (they are always Linear)
+	// this happens to be the same as IsU8Channels() at the moment, but don't assume that
 	FORCEINLINE bool GetFormatNeedsGammaSpace(Type Format)
 	{
 		if ( Format == G8 || Format == BGRA8 )
@@ -99,7 +102,7 @@ namespace ERawImageFormat
 		}
 		else
 		{
-			// these formats ignore GammaSpace setting
+			// these formats ignore GammaSpace setting, always Linear
 			return false;
 		}
 	}
