Landscape Editor - Adding EditLayer validation checks to fix crash when slate IsLayerSubtractiveBlendChecked runs

diff --git a/Engine/Source/Editor/LandscapeEditor/Private/LandscapeEditorDetailCustomization_TargetLayers.cpp b/Engine/Source/Editor/LandscapeEditor/Private/LandscapeEditorDetailCustomization_TargetLayers.cpp
--- a/Engine/Source/Editor/LandscapeEditor/Private/LandscapeEditorDetailCustomization_TargetLayers.cpp
+++ b/Engine/Source/Editor/LandscapeEditor/Private/LandscapeEditorDetailCustomization_TargetLayers.cpp
@@ -2027,9 +2027,13 @@ EVisibility FLandscapeEditorCustomNodeBuilder_TargetLayers::GetDebugModeLayerUsa
 EVisibility FLandscapeEditorCustomNodeBuilder_TargetLayers::GetLayersSubstractiveBlendVisibility(const TSharedRef<FLandscapeTargetListInfo> Target)
 {
 	FEdModeLandscape* LandscapeEdMode = GetEditorMode();
-	if (LandscapeEdMode != nullptr && LandscapeEdMode->CanHaveLandscapeLayersContent() && Target->TargetType != ELandscapeToolTargetType::Heightmap && Target->LayerInfoObj.IsValid())
+	if (LandscapeEdMode != nullptr && Target->TargetType == ELandscapeToolTargetType::Weightmap && Target->LayerInfoObj.IsValid())
 	{
-		return EVisibility::Visible;
+		ULandscapeEditLayerBase* EditLayer = LandscapeEdMode->GetEditLayer(LandscapeEdMode->GetSelectedEditLayerIndex());
+		if (EditLayer && LandscapeEdMode->CanHaveLandscapeLayersContent())
+		{
+			return EVisibility::Visible;
+		}
 	}
 
 	return EVisibility::Collapsed;
@@ -2040,11 +2044,14 @@ ECheckBoxState FLandscapeEditorCustomNodeBuilder_TargetLayers::IsLayersSubstract
 	if (const FEdModeLandscape* LandscapeEdMode = GetEditorMode())
 	{
 		ULandscapeEditLayerBase* EditLayer = LandscapeEdMode->GetEditLayer(LandscapeEdMode->GetSelectedEditLayerIndex());
-		const bool* AllocationBlend = EditLayer->GetWeightmapLayerAllocationBlend().Find(Target.Get().LayerInfoObj.Get());
-
-		if (AllocationBlend != nullptr)
+		if (EditLayer && Target->LayerInfoObj.IsValid())
 		{
-			return (*AllocationBlend) ? ECheckBoxState::Checked : ECheckBoxState::Unchecked;
+			const bool* AllocationBlend = EditLayer->GetWeightmapLayerAllocationBlend().Find(Target.Get().LayerInfoObj.Get());
+
+			if (AllocationBlend != nullptr)
+			{
+				return (*AllocationBlend) ? ECheckBoxState::Checked : ECheckBoxState::Unchecked;
+			}
 		}
 	}
 
