Fixed landscape localization issues

diff --git a/Engine/Source/Editor/LandscapeEditor/Private/LandscapeEditorDetailCustomization_Layers.cpp b/Engine/Source/Editor/LandscapeEditor/Private/LandscapeEditorDetailCustomization_Layers.cpp
--- a/Engine/Source/Editor/LandscapeEditor/Private/LandscapeEditorDetailCustomization_Layers.cpp
+++ b/Engine/Source/Editor/LandscapeEditor/Private/LandscapeEditorDetailCustomization_Layers.cpp
@@ -617,7 +617,7 @@ TSharedPtr<SWidget> FLandscapeEditorCustomNodeBuilder_Layers::OnLayerContextMenu
 		const TArray<ALandscapeBlueprintBrushBase*>& Brushes = LandscapeEdMode->GetBrushList();
 		TArray<ALandscapeBlueprintBrushBase*> FilteredBrushes = Brushes.FilterByPredicate([](ALandscapeBlueprintBrushBase* Brush) { return Brush->GetOwningLandscape() == nullptr; });
 		
-		MenuBuilder.BeginSection("LandscapeEditorBrushActions", LOCTEXT("LandscapeEditorBrushActions.Heading", "Unassigned Brushes"));
+		MenuBuilder.BeginSection("LandscapeEditorBrushActions", LOCTEXT("LandscapeEditorBrushActions.Heading", "Brushes"));
 		{
 			if (EditLayer)
 			{
@@ -636,7 +636,7 @@ TSharedPtr<SWidget> FLandscapeEditorCustomNodeBuilder_Layers::OnLayerContextMenu
 				{
 					FUIAction NoOpAction = FUIAction(FExecuteAction::CreateLambda([]() {}),
 						FCanExecuteAction::CreateLambda([]() { return false; }));
-					const FText NoOpToolTip = EditLayer->SupportsBlueprintBrushes() ? LOCTEXT("LandscapeEditorBrushNone", "All Blueprint Brush actors are assigned to a landscape edit layer")
+					const FText NoOpToolTip = EditLayer->SupportsBlueprintBrushes() ? LOCTEXT("LandscapeEditorBrushAllBrushActorsAssigned", "All Blueprint Brush actors are assigned to a landscape edit layer")
 						: FText::Format(LOCTEXT("LandscapeEditorBrushUnsupported", "This layer's type ({0}) doesn't support blueprint brushes."), EditLayer->GetClass()->GetDisplayNameText());
 
 					MenuBuilder.AddMenuEntry(
