Landscape Editor - Fix to ensure Paint mode renames the correct target layer when a new layer is added after another layer has been deleted

diff --git a/Engine/Source/Editor/LandscapeEditor/Private/LandscapeEditorDetailCustomization_TargetLayers.cpp b/Engine/Source/Editor/LandscapeEditor/Private/LandscapeEditorDetailCustomization_TargetLayers.cpp
--- a/Engine/Source/Editor/LandscapeEditor/Private/LandscapeEditorDetailCustomization_TargetLayers.cpp
+++ b/Engine/Source/Editor/LandscapeEditor/Private/LandscapeEditorDetailCustomization_TargetLayers.cpp
@@ -746,19 +746,26 @@ void FLandscapeEditorCustomNodeBuilder_TargetLayers::HandleCreateLayer()
 		}
 
 		FScopedTransaction Transaction(LOCTEXT("LandscapeTargetLayer_Create", "Create a Target Layer"));
-
-		Landscape->AddTargetLayer();
+		// TODO [jared.ritchie] - AddTargetLayer returns a FLandscapeTargetLayerSettings which does not give
+		// any context about the newly created LayerName, position, etc. Deprecate AddTargetLayer
+		// and use the return of new function to simplify setting the PendingRenameTargetLayerIndex
+		const FName& NewTargetLayerName = Landscape->GenerateUniqueTargetLayerName();
+		Landscape->AddTargetLayer(NewTargetLayerName, FLandscapeTargetLayerSettings());
 
 		LandscapeEdMode->CurrentToolTarget.LandscapeInfo->UpdateLayerInfoMap();
 		LandscapeEdMode->UpdateTargetList();
 
 		// Trigger a rename for the new layer, enters edit mode on next tick
-		LandscapeEdMode->PendingRenameTargetLayerIndex = GetWeightmapTargetLayerCount() - 1;
+		const int32 NewTargetLayerIndex = LandscapeEdMode->GetTargetDisplayOrderList()->Find(NewTargetLayerName);
+		LandscapeEdMode->PendingRenameTargetLayerIndex = NewTargetLayerIndex - LandscapeEdMode->GetTargetLayerStartingIndex();
 
 		LandscapeEdMode->RefreshDetailPanel();
 
 		// Auto-select new target layer in details panel
-		OnTargetSelectionChanged(LandscapeEdMode->GetTargetList().Last());
+		if (LandscapeEdMode->GetTargetList().IsValidIndex(NewTargetLayerIndex))
+		{
+			OnTargetSelectionChanged(LandscapeEdMode->GetTargetList()[NewTargetLayerIndex]);
+		}
 	}
 }
 
