Landscape Editor - Updating Edit/Target layer context menus to simplify Clear layer actions
* Deprecate ELandscapeClearMode and ELandscapeToolTargetTypeMask in favor or ELandscapeToolTargetTypeFlags
* Update LandscapeEdMode FLandscapeToolMode to track ELandscapeToolTargetTypeFlags

diff --git a/Engine/Source/Editor/LandscapeEditor/Private/LandscapeEdMode.cpp b/Engine/Source/Editor/LandscapeEditor/Private/LandscapeEdMode.cpp
--- a/Engine/Source/Editor/LandscapeEditor/Private/LandscapeEdMode.cpp
+++ b/Engine/Source/Editor/LandscapeEditor/Private/LandscapeEdMode.cpp
@@ -212,7 +212,7 @@ void FEdModeLandscape::UpdateToolModes()
 
 	LandscapeToolModes.Reset();
 	
-	FLandscapeToolMode* ToolMode_Manage = new(LandscapeToolModes)FLandscapeToolMode(TEXT("ToolMode_Manage"), ELandscapeToolTargetTypeMask::NA);
+	FLandscapeToolMode* ToolMode_Manage = new(LandscapeToolModes)FLandscapeToolMode(TEXT("ToolMode_Manage"), ELandscapeToolTargetTypeFlags::None);
 	ToolMode_Manage->ValidTools.Add(TEXT("NewLandscape"));
 	ToolMode_Manage->ValidTools.Add(TEXT("Select"));
 	ToolMode_Manage->ValidTools.Add(TEXT("AddComponent"));
@@ -234,7 +234,7 @@ void FEdModeLandscape::UpdateToolModes()
 	ToolMode_Manage->CurrentTargetLayerName = PreviousTargetLayerName ? *PreviousTargetLayerName : NAME_None;
 		
 
-	FLandscapeToolMode* ToolMode_Sculpt = new(LandscapeToolModes)FLandscapeToolMode(TEXT("ToolMode_Sculpt"), ELandscapeToolTargetTypeMask::Heightmap | ELandscapeToolTargetTypeMask::Visibility);
+	FLandscapeToolMode* ToolMode_Sculpt = new(LandscapeToolModes)FLandscapeToolMode(TEXT("ToolMode_Sculpt"), ELandscapeToolTargetTypeFlags::Heightmap | ELandscapeToolTargetTypeFlags::Visibility);
 	ToolMode_Sculpt->ValidTools.Add(TEXT("Sculpt"));
 	if (CanHaveLandscapeLayersContent())
 	{
@@ -260,7 +260,7 @@ void FEdModeLandscape::UpdateToolModes()
 	PreviousTargetLayerName = PreviousTargetLayerNames.Find(ToolMode_Sculpt->ToolModeName);
 	ToolMode_Sculpt->CurrentTargetLayerName = PreviousTargetLayerName ? *PreviousTargetLayerName : NAME_None;
 
-	FLandscapeToolMode* ToolMode_Paint = new(LandscapeToolModes)FLandscapeToolMode(TEXT("ToolMode_Paint"), ELandscapeToolTargetTypeMask::Weightmap);
+	FLandscapeToolMode* ToolMode_Paint = new(LandscapeToolModes)FLandscapeToolMode(TEXT("ToolMode_Paint"), ELandscapeToolTargetTypeFlags::Weightmap);
 	ToolMode_Paint->ValidTools.Add(TEXT("Paint"));
 	ToolMode_Paint->ValidTools.Add(TEXT("Smooth"));
 	ToolMode_Paint->ValidTools.Add(TEXT("Flatten"));
@@ -1972,12 +1972,12 @@ bool FEdModeLandscape::InputKey(FEditorViewportClient* ViewportClient, FViewport
 					 (LandscapeEditorControlType == ELandscapeFoliageEditorControlType::RequireCtrl   && IsCtrlDown(Viewport)) ||
 					 (LandscapeEditorControlType == ELandscapeFoliageEditorControlType::RequireNoCtrl && !IsCtrlDown(Viewport)))))
 			{
-				if (CurrentTool && (CurrentTool->GetSupportedTargetTypes() == ELandscapeToolTargetTypeMask::NA || CurrentToolTarget.TargetType != ELandscapeToolTargetType::Invalid))
+				if (CurrentTool && (CurrentTool->GetSupportedTargetTypes() == ELandscapeToolTargetTypeFlags::None || CurrentToolTarget.TargetType != ELandscapeToolTargetType::Invalid))
 				{
 					FVector HitLocation;
 					if (LandscapeMouseTrace(ViewportClient, HitLocation))
 					{
-						if (CurrentTool->AffectsEditLayers() && !CanEditTargetLayer(&ErrorReasonOnMouseUp))
+						if (CurrentTool->AffectsEditLayers() && !CanEditSelectedTargetLayer(&ErrorReasonOnMouseUp))
 						{
 							return true;
 						}
@@ -2134,7 +2134,7 @@ void FEdModeLandscape::SetCurrentTool(FName ToolName, FName TargetLayerName)
 		FLandscapeTool* Tool = LandscapeTools[i].Get();
 		if (ToolName == Tool->GetToolName())
 		{
-			if ((Tool->GetSupportedTargetTypes() & CurrentToolMode->SupportedTargetTypes) != 0)
+			if (EnumHasAnyFlags(Tool->GetSupportedTargetTypes(), CurrentToolMode->SupportedTargetTypes))
 			{
 				ToolIndex = i;
 				break;
@@ -2204,19 +2204,19 @@ void FEdModeLandscape::SetCurrentTool(int32 ToolIndex, FName TargetLayerName)
 	CurrentTool = NewTool;
 
 	// Set target type appropriate for tool
-	if (CurrentTool->GetSupportedTargetTypes() == ELandscapeToolTargetTypeMask::NA)
+	if (CurrentTool->GetSupportedTargetTypes() == ELandscapeToolTargetTypeFlags::None)
 	{
 		CurrentToolTarget.TargetType = ELandscapeToolTargetType::Invalid;
 		SetCurrentTargetLayer(NAME_None, nullptr);
 	}
 	else
 	{
-		const uint8 TargetTypeMask = CurrentToolMode->SupportedTargetTypes & CurrentTool->GetSupportedTargetTypes();
-		checkSlow(TargetTypeMask != 0);
-
-		if ((TargetTypeMask & ELandscapeToolTargetTypeMask::FromType(CurrentToolTarget.TargetType)) == 0)
+		const ELandscapeToolTargetTypeFlags TargetTypeMask = CurrentTool->GetSupportedTargetTypes();
+		checkSlow(TargetTypeMask != ELandscapeToolTargetTypeFlags::None);
+		// GetLandscapeToolTargetTypeAsFlags does not support ELandscapeToolTargetType::Invalid so explicitly check this first 
+		if (CurrentToolTarget.TargetType == ELandscapeToolTargetType::Invalid || !EnumHasAnyFlags(UE::Landscape::GetLandscapeToolTargetTypeAsFlags(CurrentToolTarget.TargetType), TargetTypeMask))
 		{
-			auto filter = [TargetTypeMask, TargetLayerName](const TSharedRef<FLandscapeTargetListInfo>& Target){ return (TargetTypeMask & ELandscapeToolTargetTypeMask::FromType(Target->TargetType)) != 0 && (TargetLayerName == NAME_None || TargetLayerName == Target->LayerName); };
+			auto filter = [TargetTypeMask, TargetLayerName](const TSharedRef<FLandscapeTargetListInfo>& Target) { return (EnumHasAnyFlags(UE::Landscape::GetLandscapeToolTargetTypeAsFlags(Target->TargetType), TargetTypeMask)) && (TargetLayerName == NAME_None || TargetLayerName == Target->LayerName); };
 			const TSharedRef<FLandscapeTargetListInfo>* Target = LandscapeTargetList.FindByPredicate(filter);
 			if (Target != nullptr)
 			{
@@ -4234,7 +4234,12 @@ bool FEdModeLandscape::CanEditLayer(FText* Reason /*=nullptr*/, const ULandscape
 	return true;
 }
 
-bool FEdModeLandscape::CanEditTargetLayer(FText* Reason /*=nullptr*/, const ULandscapeEditLayerBase* InLayer /*= nullptr*/)
+bool FEdModeLandscape::CanEditSelectedTargetLayer(FText* Reason /*=nullptr*/, const ULandscapeEditLayerBase* InLayer /*= nullptr*/)
+{
+	return CanEditTargetLayer(Reason, CurrentToolTarget.LayerInfo.Get(), InLayer);
+}
+
+bool FEdModeLandscape::CanEditTargetLayer(FText* Reason /*=nullptr*/, const ULandscapeLayerInfoObject* InLayerInfo  /*=nullptr*/, const ULandscapeEditLayerBase* InLayer /*= nullptr*/)
 {
 	// Target layer is only editable when the edit layer is
 	if (!CanEditLayer(Reason, InLayer))
@@ -4242,7 +4247,7 @@ bool FEdModeLandscape::CanEditTargetLayer(FText* Reason /*=nullptr*/, const ULan
 		return false;
 	}
 
-	if (CurrentToolTarget.TargetType == ELandscapeToolTargetType::Weightmap && CurrentToolTarget.LayerInfo == nullptr)
+	if (CurrentToolTarget.TargetType == ELandscapeToolTargetType::Weightmap && InLayerInfo == nullptr)
 	{
 		if (Reason)
 		{
diff --git a/Engine/Source/Editor/LandscapeEditor/Private/LandscapeEdMode.h b/Engine/Source/Editor/LandscapeEditor/Private/LandscapeEdMode.h
--- a/Engine/Source/Editor/LandscapeEditor/Private/LandscapeEdMode.h
+++ b/Engine/Source/Editor/LandscapeEditor/Private/LandscapeEdMode.h
@@ -54,13 +54,13 @@ template<typename TargetType> class FLandscapeToolCopyPaste;
 struct FLandscapeToolMode
 {
 	const FName				ToolModeName;
-	int32					SupportedTargetTypes; // ELandscapeToolTargetTypeMask::Type
+	ELandscapeToolTargetTypeFlags	SupportedTargetTypes;
 
 	TArray<FName>			ValidTools;
 	FName					CurrentToolName;
 	FName					CurrentTargetLayerName;
 
-	FLandscapeToolMode(FName InToolModeName, int32 InSupportedTargetTypes)
+	FLandscapeToolMode(FName InToolModeName, ELandscapeToolTargetTypeFlags InSupportedTargetTypes)
 		: ToolModeName(InToolModeName)
 		, SupportedTargetTypes(InSupportedTargetTypes)
 	{
@@ -584,7 +584,8 @@ public:
 	void UpdateLandscapeSplines(bool bUpdateOnlySelected = false);
 	void AutoUpdateDirtyLandscapeSplines();
 	bool CanEditLayer(FText* Reason = nullptr, const ULandscapeEditLayerBase* InLayer = nullptr);
-	bool CanEditTargetLayer(FText* Reason = nullptr, const ULandscapeEditLayerBase* InLayer = nullptr);
+	bool CanEditSelectedTargetLayer(FText* Reason = nullptr, const ULandscapeEditLayerBase* InLayer = nullptr);
+	bool CanEditTargetLayer(FText* Reason = nullptr, const ULandscapeLayerInfoObject* InLayerInfo = nullptr, const ULandscapeEditLayerBase* InLayer = nullptr);
 
 	void AddBrushToCurrentLayer(class ALandscapeBlueprintBrushBase* InBrush);
 	void RemoveBrushFromCurrentLayer(int32 InBrushIndex);
diff --git a/Engine/Source/Editor/LandscapeEditor/Private/LandscapeEdModeBrushes.cpp b/Engine/Source/Editor/LandscapeEditor/Private/LandscapeEdModeBrushes.cpp
--- a/Engine/Source/Editor/LandscapeEditor/Private/LandscapeEdModeBrushes.cpp
+++ b/Engine/Source/Editor/LandscapeEditor/Private/LandscapeEdModeBrushes.cpp
@@ -113,7 +113,7 @@ public:
 			return false;
 		}
 
-		if (!EdMode->CanEditTargetLayer())
+		if (!EdMode->CanEditSelectedTargetLayer())
 		{
 			return false;
 		}
diff --git a/Engine/Source/Editor/LandscapeEditor/Private/LandscapeEdModeComponentTools.cpp b/Engine/Source/Editor/LandscapeEditor/Private/LandscapeEdModeComponentTools.cpp
--- a/Engine/Source/Editor/LandscapeEditor/Private/LandscapeEdModeComponentTools.cpp
+++ b/Engine/Source/Editor/LandscapeEditor/Private/LandscapeEdModeComponentTools.cpp
@@ -363,9 +363,9 @@ public:
 	virtual void SetEditRenderType() override { GLandscapeEditRenderMode = ELandscapeEditRenderMode::None | (GLandscapeEditRenderMode & ELandscapeEditRenderMode::BitMaskForMask); }
 	virtual bool SupportsMask() override { return false; }
 
-	virtual ELandscapeToolTargetTypeMask::Type GetSupportedTargetTypes() override
+	virtual ELandscapeToolTargetTypeFlags GetSupportedTargetTypes() override
 	{
-		return ELandscapeToolTargetTypeMask::Visibility;
+		return ELandscapeToolTargetTypeFlags::Visibility;
 	}
 };
 
@@ -1327,9 +1327,9 @@ public:
 		GLandscapeEditRenderMode |= (this->EdMode && this->EdMode->CurrentToolTarget.LandscapeInfo.IsValid() && this->EdMode->CurrentToolTarget.LandscapeInfo->SelectedRegion.Num()) ? ELandscapeEditRenderMode::SelectRegion : ELandscapeEditRenderMode::SelectComponent;
 	}
 
-	virtual ELandscapeToolTargetTypeMask::Type GetSupportedTargetTypes() override
+	virtual ELandscapeToolTargetTypeFlags GetSupportedTargetTypes() override
 	{
-		return ELandscapeToolTargetTypeMask::FromType(ToolTarget::TargetType);
+		return UE::Landscape::GetLandscapeToolTargetTypeAsFlags(ToolTarget::TargetType);
 	}
 
 	virtual bool BeginTool(FEditorViewportClient* ViewportClient, const FLandscapeToolTarget& InTarget, const FVector& InHitLocation) override
@@ -1658,9 +1658,9 @@ public:
 		bUseGizmoRegion = InbUseGizmoRegion;
 	}
 
-	virtual ELandscapeToolTargetTypeMask::Type GetSupportedTargetTypes() override
+	virtual ELandscapeToolTargetTypeFlags GetSupportedTargetTypes() override
 	{
-		return ELandscapeToolTargetTypeMask::FromType(ToolTarget::TargetType);
+		return UE::Landscape::GetLandscapeToolTargetTypeAsFlags(ToolTarget::TargetType);
 	}
 
 	virtual ELandscapeLayerUpdateMode GetBeginToolContentUpdateFlag() const override { return ELandscapeLayerUpdateMode::Update_All_Editing; }
diff --git a/Engine/Source/Editor/LandscapeEditor/Private/LandscapeEdModeErosionTools.cpp b/Engine/Source/Editor/LandscapeEditor/Private/LandscapeEdModeErosionTools.cpp
--- a/Engine/Source/Editor/LandscapeEditor/Private/LandscapeEdModeErosionTools.cpp
+++ b/Engine/Source/Editor/LandscapeEditor/Private/LandscapeEdModeErosionTools.cpp
@@ -47,9 +47,9 @@ public:
 	{
 	}
 
-	virtual ELandscapeToolTargetTypeMask::Type GetSupportedTargetTypes() override
+	virtual ELandscapeToolTargetTypeFlags GetSupportedTargetTypes() override
 	{
-		return ELandscapeToolTargetTypeMask::Heightmap;
+		return ELandscapeToolTargetTypeFlags::Heightmap;
 	}
 };
 
diff --git a/Engine/Source/Editor/LandscapeEditor/Private/LandscapeEdModeMirrorTool.cpp b/Engine/Source/Editor/LandscapeEditor/Private/LandscapeEdModeMirrorTool.cpp
--- a/Engine/Source/Editor/LandscapeEditor/Private/LandscapeEdModeMirrorTool.cpp
+++ b/Engine/Source/Editor/LandscapeEditor/Private/LandscapeEdModeMirrorTool.cpp
@@ -52,9 +52,9 @@ public:
 	virtual void SetEditRenderType() override { GLandscapeEditRenderMode = ELandscapeEditRenderMode::None | (GLandscapeEditRenderMode & ELandscapeEditRenderMode::BitMaskForMask); }
 	virtual bool SupportsMask() override { return false; }
 
-	virtual ELandscapeToolTargetTypeMask::Type GetSupportedTargetTypes() override
+	virtual ELandscapeToolTargetTypeFlags GetSupportedTargetTypes() override
 	{
-		return ELandscapeToolTargetTypeMask::Heightmap;
+		return ELandscapeToolTargetTypeFlags::Heightmap;
 	}
 
 	virtual void EnterTool() override
diff --git a/Engine/Source/Editor/LandscapeEditor/Private/LandscapeEdModePaintTools.cpp b/Engine/Source/Editor/LandscapeEditor/Private/LandscapeEdModePaintTools.cpp
--- a/Engine/Source/Editor/LandscapeEditor/Private/LandscapeEdModePaintTools.cpp
+++ b/Engine/Source/Editor/LandscapeEditor/Private/LandscapeEdModePaintTools.cpp
@@ -56,9 +56,9 @@ public:
 	{
 	}
 
-	virtual ELandscapeToolTargetTypeMask::Type GetSupportedTargetTypes() override
+	virtual ELandscapeToolTargetTypeFlags GetSupportedTargetTypes() override
 	{
-		return ELandscapeToolTargetTypeMask::FromType(TToolTarget::TargetType);
+		return UE::Landscape::GetLandscapeToolTargetTypeAsFlags(TToolTarget::TargetType);
 	}
 };
 
diff --git a/Engine/Source/Editor/LandscapeEditor/Private/LandscapeEdModeRampTool.cpp b/Engine/Source/Editor/LandscapeEditor/Private/LandscapeEdModeRampTool.cpp
--- a/Engine/Source/Editor/LandscapeEditor/Private/LandscapeEdModeRampTool.cpp
+++ b/Engine/Source/Editor/LandscapeEditor/Private/LandscapeEdModeRampTool.cpp
@@ -126,9 +126,9 @@ public:
 	virtual void SetEditRenderType() override { GLandscapeEditRenderMode = ELandscapeEditRenderMode::None | (GLandscapeEditRenderMode & ELandscapeEditRenderMode::BitMaskForMask); }
 	virtual bool SupportsMask() override { return false; }
 
-	virtual ELandscapeToolTargetTypeMask::Type GetSupportedTargetTypes() override
+	virtual ELandscapeToolTargetTypeFlags GetSupportedTargetTypes() override
 	{
-		return ELandscapeToolTargetTypeMask::Heightmap;
+		return ELandscapeToolTargetTypeFlags::Heightmap;
 	}
 
 	virtual void EnterTool() override
diff --git a/Engine/Source/Editor/LandscapeEditor/Private/LandscapeEdModeXYOffset.cpp b/Engine/Source/Editor/LandscapeEditor/Private/LandscapeEdModeXYOffset.cpp
--- a/Engine/Source/Editor/LandscapeEditor/Private/LandscapeEdModeXYOffset.cpp
+++ b/Engine/Source/Editor/LandscapeEditor/Private/LandscapeEdModeXYOffset.cpp
@@ -567,10 +567,10 @@ public:
 	virtual FText GetDisplayName() const override { return NSLOCTEXT("UnrealEd", "LandscapeMode_Retopologize", "Retopologize"); }
 	virtual FText GetDisplayMessage() const override { return NSLOCTEXT("UnrealEd", "LandscapeMode_Retopologize_Message", "Similar to the smooth tool, Retopologize will push and pull triangles to smooth the transition of the terrain, however it will attempt to keep the basic form of the terrain, minimizing change in the Z direction. An X/Y offset map makes the Landscape slower to render and paint on with other tools, so only use the Retopologize tool if needed."); }
 
-	virtual ELandscapeToolTargetTypeMask::Type GetSupportedTargetTypes() override
+	virtual ELandscapeToolTargetTypeFlags GetSupportedTargetTypes() override
 	{
 		// technically not entirely accurate, also modifies the XYOffset map
-		return ELandscapeToolTargetTypeMask::Heightmap;
+		return ELandscapeToolTargetTypeFlags::Heightmap;
 	}
 };
 
diff --git a/Engine/Source/Editor/LandscapeEditor/Private/LandscapeEditorDetailCustomization_Layers.cpp b/Engine/Source/Editor/LandscapeEditor/Private/LandscapeEditorDetailCustomization_Layers.cpp
--- a/Engine/Source/Editor/LandscapeEditor/Private/LandscapeEditorDetailCustomization_Layers.cpp
+++ b/Engine/Source/Editor/LandscapeEditor/Private/LandscapeEditorDetailCustomization_Layers.cpp
@@ -453,75 +453,11 @@ FSlateColor FLandscapeEditorCustomNodeBuilder_Layers::GetLayerTextColor(int32 In
 	return IsLayerSelected(InLayerIndex) ? FStyleColors::ForegroundHover : FSlateColor::UseForeground();
 }
 
-
-void FLandscapeEditorCustomNodeBuilder_Layers::FillClearTargetLayerMenu(FMenuBuilder& MenuBuilder, int32 InLayerIndex, TArray<ULandscapeLayerInfoObject*> InUsedLayerInfos)
-{
-	// Clear All Weightmap Data
-	FUIAction ClearAction = FUIAction(FExecuteAction::CreateSP(this, &FLandscapeEditorCustomNodeBuilder_Layers::ClearTargetLayersOnLayer, InLayerIndex, ELandscapeClearMode::Clear_Weightmap));
-	MenuBuilder.AddMenuEntry(LOCTEXT("LandscapeClearAllWeightmap", "All"), FText(), FSlateIcon(), ClearAction);
-	MenuBuilder.AddMenuSeparator();
-
-	// Clear Per LayerInfo
-	for (ULandscapeLayerInfoObject* LayerInfo : InUsedLayerInfos)
-	{
-		FUIAction ClearLayerInfoAction = FUIAction(
-			FExecuteAction::CreateSP(this, &FLandscapeEditorCustomNodeBuilder_Layers::ClearTargetLayerOnLayer,InLayerIndex, LayerInfo),
-			FCanExecuteAction::CreateSPLambda(this, [this, InLayerIndex, LayerInfo]() { FText Reason; return CanClearTargetLayerOnLayer(InLayerIndex, LayerInfo, Reason); }));
-		TAttribute<FText> CanClearTargetLayerOnLayerTooltipText = MakeAttributeSPLambda(this, [this, InLayerIndex, LayerInfo] { FText Reason; CanClearTargetLayerOnLayer(InLayerIndex, LayerInfo, Reason); return Reason; });
-		MenuBuilder.AddMenuEntry(FText::FromName(LayerInfo->LayerName), CanClearTargetLayerOnLayerTooltipText, FSlateIcon(), ClearLayerInfoAction);
-	}
-}
-
-void FLandscapeEditorCustomNodeBuilder_Layers::FillClearLayerMenu(FMenuBuilder& MenuBuilder, int32 InLayerIndex)
-{
-	const UEnum* EnumPtr = FindObject<UEnum>(nullptr, TEXT("/Script/Landscape.ELandscapeClearMode"), EFindObjectFlags::ExactClass);
-	if (ensure(EnumPtr != nullptr))
-	{
-		// NumEnums()-1 to exclude Enum Max Value
-		for (int32 i = 0; i < EnumPtr->NumEnums()-1; ++i)
-		{
-			ELandscapeClearMode EnumValue = (ELandscapeClearMode)EnumPtr->GetValueByIndex(i);
-			if(EnumValue == ELandscapeClearMode::Clear_Weightmap)
-			{
-				TArray<ULandscapeLayerInfoObject*> UsedLayerInfos;
-				FEdModeLandscape* LandscapeEdMode = GetEditorMode();
-				check(LandscapeEdMode);
-				ALandscape* Landscape = LandscapeEdMode->GetLandscape();
-				check(Landscape);
-
-				Landscape->GetUsedPaintLayers(InLayerIndex, UsedLayerInfos);
-				if (UsedLayerInfos.Num() > 0)
-				{
-					FUIAction ClearLayerAction = FUIAction(
-						FExecuteAction(),
-						FCanExecuteAction::CreateSPLambda(this, [this, InLayerIndex, EnumValue] { FText Reason; return CanClearTargetLayersOnLayer(InLayerIndex, EnumValue, Reason); }));
-					TAttribute<FText> CanClearLayerTooltipText = MakeAttributeSPLambda(this, [this, InLayerIndex, EnumValue] { FText Reason; CanClearTargetLayersOnLayer(InLayerIndex, EnumValue, Reason); return Reason; });
-					MenuBuilder.AddSubMenu(
-						EnumPtr->GetDisplayNameTextByIndex(i),
-						CanClearLayerTooltipText,
-						FNewMenuDelegate::CreateSP(this, &FLandscapeEditorCustomNodeBuilder_Layers::FillClearTargetLayerMenu, InLayerIndex, UsedLayerInfos),
-						ClearLayerAction,
-						/*InExtensionHook = */NAME_None,
-						EUserInterfaceActionType::None);
-				}
-			}
-			else
-			{
-				FUIAction ClearLayerAction = FUIAction(
-					FExecuteAction::CreateSP(this, &FLandscapeEditorCustomNodeBuilder_Layers::ClearTargetLayersOnLayer, InLayerIndex, EnumValue),
-					FCanExecuteAction::CreateSPLambda(this, [this, InLayerIndex, EnumValue]() { FText Reason; return CanClearTargetLayersOnLayer(InLayerIndex, EnumValue, Reason); }));
-				TAttribute<FText> CanClearLayerActionTooltipText = MakeAttributeSPLambda(this, [this, InLayerIndex, EnumValue] { FText Reason; CanClearTargetLayersOnLayer(InLayerIndex, EnumValue, Reason); return Reason; });
-				MenuBuilder.AddMenuEntry(EnumPtr->GetDisplayNameTextByIndex(i), CanClearLayerActionTooltipText, FSlateIcon(), ClearLayerAction);
-			}
-		}
-	}
-}
-
 TSharedPtr<SWidget> FLandscapeEditorCustomNodeBuilder_Layers::OnLayerContextMenuOpening(int32 InLayerIndex)
 {
 	FEdModeLandscape* LandscapeEdMode = GetEditorMode();
 	ALandscape* Landscape = LandscapeEdMode ? LandscapeEdMode->GetLandscape() : nullptr;
-	if (Landscape && LandscapeEdMode && LandscapeEdMode->DoesCurrentToolAffectEditLayers())
+	if (Landscape && LandscapeEdMode)
 	{
 		const ULandscapeEditLayerBase* EditLayer = LandscapeEdMode->GetEditLayerConst(InLayerIndex);
 		FMenuBuilder MenuBuilder(true, nullptr);
@@ -534,34 +470,44 @@ TSharedPtr<SWidget> FLandscapeEditorCustomNodeBuilder_Layers::OnLayerContextMenu
 					FExecuteAction::CreateSP(this, &FLandscapeEditorCustomNodeBuilder_Layers::RenameLayer, InLayerIndex),
 					FCanExecuteAction::CreateSPLambda(this, [this, InLayerIndex] { FText Reason; return CanRenameLayer(InLayerIndex, Reason); }));
 				TAttribute<FText> CanRenameLayerTooltipText = MakeAttributeSPLambda(this, [this, InLayerIndex] { FText Reason; CanRenameLayer(InLayerIndex, Reason); return Reason; });
-				MenuBuilder.AddMenuEntry(LOCTEXT("RenameLayer", "Rename..."), CanRenameLayerTooltipText, FSlateIcon(), RenameLayerAction);
-
-				// Clear Layer
-				FUIAction ClearLayerAction = FUIAction(
-					FExecuteAction(),
-					FCanExecuteAction::CreateSPLambda(this, [this, InLayerIndex] { FText Reason; return CanClearLayer(InLayerIndex, Reason); }));
-				TAttribute<FText> CanClearLayerTooltipText = MakeAttributeSPLambda(this, [this, InLayerIndex] { FText Reason; CanClearLayer(InLayerIndex, Reason); return Reason; });
-				MenuBuilder.AddSubMenu(
-					LOCTEXT("LandscapeEditorClearLayerSubMenu", "Clear"),
-					CanClearLayerTooltipText,
-					FNewMenuDelegate::CreateSP(this, &FLandscapeEditorCustomNodeBuilder_Layers::FillClearLayerMenu, InLayerIndex),
-					ClearLayerAction, 
-					/*InExtensionHook = */NAME_None,
-					EUserInterfaceActionType::None);
+				MenuBuilder.AddMenuEntry(LOCTEXT("RenameLayer", "Rename"), CanRenameLayerTooltipText, FSlateIcon(), RenameLayerAction);
 
 				// Delete Layer
 				FUIAction DeleteLayerAction = FUIAction(
 					FExecuteAction::CreateSP(this, &FLandscapeEditorCustomNodeBuilder_Layers::DeleteLayer, InLayerIndex),
 					FCanExecuteAction::CreateSPLambda(this, [this, InLayerIndex] { FText Reason; return CanDeleteLayer(InLayerIndex, Reason); }));
 				TAttribute<FText> DeleteLayerTooltipText = MakeAttributeSPLambda(this, [this, InLayerIndex] { FText Reason; CanDeleteLayer(InLayerIndex, Reason); return Reason; });
-				MenuBuilder.AddMenuEntry(LOCTEXT("DeleteLayer", "Delete..."), DeleteLayerTooltipText, FSlateIcon(), DeleteLayerAction);
+				MenuBuilder.AddMenuEntry(LOCTEXT("DeleteLayer", "Delete"), DeleteLayerTooltipText, FSlateIcon(), DeleteLayerAction);
 
 				// Collapse Layer
 				FUIAction CollapseLayerAction = FUIAction(
 					FExecuteAction::CreateSP(this, &FLandscapeEditorCustomNodeBuilder_Layers::CollapseLayer, InLayerIndex),
 					FCanExecuteAction::CreateSPLambda(this, [this, InLayerIndex] { FText Reason; return CanCollapseLayer(InLayerIndex, Reason); }));
 				TAttribute<FText> CollapseLayerTooltipText = MakeAttributeSPLambda(this, [this, InLayerIndex] { FText Reason; CanCollapseLayer(InLayerIndex, Reason); return Reason; });
-				MenuBuilder.AddMenuEntry(LOCTEXT("CollapseLayer", "Collapse..."), CollapseLayerTooltipText, FSlateIcon(), CollapseLayerAction);
+				MenuBuilder.AddMenuEntry(LOCTEXT("CollapseLayer", "Collapse"), CollapseLayerTooltipText, FSlateIcon(), CollapseLayerAction);
+
+				if (LandscapeEdMode->DoesCurrentToolAffectEditLayers())
+				{
+					const ELandscapeToolTargetType CurrentTargetType = LandscapeEdMode->CurrentToolTarget.TargetType;
+					// Use ELandscapeToolTargetType::Weightmap as default mode: used for Target Layers and ELandscapeToolTargetType::Invalid (Paint tool with no target layers)
+					FText ClearEditLayerDataText = LOCTEXT("ClearAllTargetLayers", "Clear All Target Layers");
+
+					if (CurrentTargetType == ELandscapeToolTargetType::Heightmap)
+					{
+						ClearEditLayerDataText = LOCTEXT("ClearSculptLayer", "Clear Sculpt Layer");
+					}
+					else if (LandscapeEdMode->CurrentToolTarget.TargetType == ELandscapeToolTargetType::Visibility)
+					{
+						ClearEditLayerDataText = LOCTEXT("ClearVisibilityLayer", "Clear Visibility Layer");
+					}
+
+					FUIAction ClearEditLayerAction = FUIAction(
+						FExecuteAction::CreateSP(this, &FLandscapeEditorCustomNodeBuilder_Layers::ClearEditLayerData, InLayerIndex, CurrentTargetType),
+						FCanExecuteAction::CreateSPLambda(this, [this, InLayerIndex, CurrentTargetType] { FText Reason; return CanClearEditLayerData(InLayerIndex, CurrentTargetType, Reason); }));
+					TAttribute<FText> ClearEditLayerDataTooltip = MakeAttributeSPLambda(this, [this, InLayerIndex, CurrentTargetType] { FText Reason; CanClearEditLayerData(InLayerIndex, CurrentTargetType, Reason); return Reason; });
+
+					MenuBuilder.AddMenuEntry(ClearEditLayerDataText, ClearEditLayerDataTooltip, FSlateIcon(), ClearEditLayerAction);
+				}
 
 				// Add custom actions : 
 				if (!EditLayer->GetActions().IsEmpty())
@@ -683,148 +629,62 @@ void FLandscapeEditorCustomNodeBuilder_Layers::RenameLayer(int32 InLayerIndex)
 	}
 }
 
-
-bool FLandscapeEditorCustomNodeBuilder_Layers::CanClearTargetLayerOnLayer(int32 InLayerIndex, ULandscapeLayerInfoObject* InLayerInfo, FText& OutReason) const
+void FLandscapeEditorCustomNodeBuilder_Layers::ClearEditLayerData(int32 InEditLayerIndex, ELandscapeToolTargetType InClearMode)
 {
-	FEdModeLandscape* LandscapeEdMode = GetEditorMode();
-	const ULandscapeEditLayerBase* EditLayer = LandscapeEdMode->GetEditLayerConst(InLayerIndex);
-
-	if (EditLayer == nullptr)
-	{
-		return false;
-	}
-
-	// Visibility case : 
-	if (InLayerInfo->LayerName == UMaterialExpressionLandscapeVisibilityMask::ParameterName)
+	if (FEdModeLandscape* LandscapeEdMode = GetEditorMode())
 	{
-		if (!EditLayer->SupportsTargetType(ELandscapeToolTargetType::Visibility))
+		ALandscape* Landscape = LandscapeEdMode->GetLandscape();
+		const ULandscapeEditLayerBase* EditLayer = LandscapeEdMode->GetEditLayerConst(InEditLayerIndex);
+		if (Landscape != nullptr && EditLayer != nullptr)
 		{
-			OutReason = FText::Format(LOCTEXT("ClearTargetLayer_CantClearVisibilityOnLayerType", "Cannot clear visibility : the type of layer {0} ({1}) doesn't support visibility"),
-				FText::FromName(EditLayer->GetName()), EditLayer->GetClass()->GetDisplayNameText());
-			return false;
-		}
-	}
-		
-	if (!EditLayer->SupportsTargetType(ELandscapeToolTargetType::Weightmap))
-	{
-		OutReason = FText::Format(LOCTEXT("ClearTargetLayer_CantClearWeightmapOnLayerType", "Cannot clear weightmap : the type of layer {0} ({1}) doesn't support weightmaps"),
-			FText::FromName(EditLayer->GetName()), EditLayer->GetClass()->GetDisplayNameText());
-		return false;
-	}
-
-	OutReason = LOCTEXT("ClearTargetLayer_CanClear", "Clear the target layer on this edit layer");
-	return true;
-}
+			const FText ClearType = UEnum::GetDisplayValueAsText(InClearMode);
 
-void FLandscapeEditorCustomNodeBuilder_Layers::ClearTargetLayerOnLayer(int32 InLayerIndex, ULandscapeLayerInfoObject* InLayerInfo)
-{
-	FEdModeLandscape* LandscapeEdMode = GetEditorMode();
-	ALandscape* Landscape = LandscapeEdMode ? LandscapeEdMode->GetLandscape() : nullptr;
-	if (Landscape)
-	{
-		const ULandscapeEditLayerBase* EditLayer = LandscapeEdMode->GetEditLayerConst(InLayerIndex);
-		if (EditLayer)
-		{
-			EAppReturnType::Type Result = FMessageDialog::Open(EAppMsgType::YesNo, FText::Format(LOCTEXT("Landscape_ClearTargetLayer_Message", "The layer {0} : {1} content will be completely cleared.  Continue?"), FText::FromName(EditLayer->GetName()), FText::FromName(InLayerInfo->LayerName)));
+			EAppReturnType::Type Result = FMessageDialog::Open(EAppMsgType::YesNo, FText::Format(LOCTEXT("Landscape_ClearEditLayer_Dialog", "Edit layer {0}'s {1} data will be completely cleared.  Continue?"), FText::FromName(EditLayer->GetName()), ClearType));
 			if (Result == EAppReturnType::Yes)
 			{
-				const FScopedTransaction Transaction(LOCTEXT("Landscape_Layers_TargetClear", "Clear Target Layer"));
-				Landscape->ClearPaintLayer(InLayerIndex, InLayerInfo);
-				LandscapeEdMode->RequestUpdateLayerUsageInformation();
+				const FScopedTransaction Transaction(FText::Format(LOCTEXT("Landscape_Edit_Layers_Clear", "Clear Edit Layer {0} Data"), ClearType));
+				Landscape->ClearEditLayer(InEditLayerIndex, /*InComponents = */ nullptr, UE::Landscape::GetLandscapeToolTargetTypeAsFlags(InClearMode));
+				OnLayerSelectionChanged(InEditLayerIndex);
+
+				if (InClearMode == ELandscapeToolTargetType::Weightmap || InClearMode == ELandscapeToolTargetType::Visibility)
+				{
+					LandscapeEdMode->RequestUpdateLayerUsageInformation();
+				}
 			}
 		}
 	}
 }
 
-bool FLandscapeEditorCustomNodeBuilder_Layers::CanClearLayer(int32 InLayerIndex, FText& OutReason) const
-{
-	FEdModeLandscape* LandscapeEdMode = GetEditorMode();
-	const ULandscapeEditLayerBase* EditLayer = LandscapeEdMode->GetEditLayerConst(InLayerIndex);
-
-	if (EditLayer == nullptr)
-	{
-		return false;
-	}
-
-	if (EditLayer->IsLocked())
-	{
-		OutReason = LOCTEXT("ClearLayer_CantClearLocked", "Cannot clear a locked edit layer");
-		return false;
-	}
-
-	if (!EditLayer->SupportsEditingTools())
-	{
-		OutReason = FText::Format(LOCTEXT("ClearLayer_CantClearLayerWithoutPersistentTextures", "Cannot clear an edit layer which doesn't have editable textures (procedural)"),
-			FText::FromName(EditLayer->GetName()), EditLayer->GetClass()->GetDisplayNameText());
-		return false;
-	}
-
-	OutReason = LOCTEXT("ClearLayer_CanClear", "Clear the edit layer");
-	return true;
-}
-
-bool FLandscapeEditorCustomNodeBuilder_Layers::CanClearTargetLayersOnLayer(int32 InLayerIndex, ELandscapeClearMode InClearMode, FText& OutReason) const
-{
-	FEdModeLandscape* LandscapeEdMode = GetEditorMode();
-	const ULandscapeEditLayerBase* EditLayer = LandscapeEdMode->GetEditLayerConst(InLayerIndex);
-	if (EditLayer == nullptr)
-	{
-		return false;
-	}
-
-	// TODO [jonathan.bard] : Ideally, ELandscapeClearMode should be deprecated and we would treat Visibility and Weightmap separately but for now just keep on treating visibility as a standard weightmap : 
-	if ((InClearMode == ELandscapeClearMode::Clear_Weightmap)
-		&& !EditLayer->SupportsTargetType(ELandscapeToolTargetType::Weightmap) 
-		&& !EditLayer->SupportsTargetType(ELandscapeToolTargetType::Visibility))
-	{
-			OutReason = FText::Format(LOCTEXT("ClearTargetLayers_CantClearWeightmapsOnLayerType", "Cannot clear weightmaps : the type of layer {0} ({1}) doesn't support weightmaps"), 
-				FText::FromName(EditLayer->GetName()), EditLayer->GetClass()->GetDisplayNameText());
-			return false;
-	}
-	else if ((InClearMode == ELandscapeClearMode::Clear_Heightmap)
-		&& !EditLayer->SupportsTargetType(ELandscapeToolTargetType::Heightmap))
-	{
-		OutReason = FText::Format(LOCTEXT("ClearTargetLayers_CantClearHeightmapOnLayerType", "Cannot clear heightmap : the type of layer {0} ({1}) doesn't support heightmaps"),
-			FText::FromName(EditLayer->GetName()), EditLayer->GetClass()->GetDisplayNameText());
-			return false;
-	}
-	if ((InClearMode == ELandscapeClearMode::Clear_All)
-		&& !EditLayer->SupportsTargetType(ELandscapeToolTargetType::Heightmap)
-		&& !EditLayer->SupportsTargetType(ELandscapeToolTargetType::Weightmap)
-		&& !EditLayer->SupportsTargetType(ELandscapeToolTargetType::Visibility))
-	if ((InClearMode == ELandscapeClearMode::Clear_Heightmap) && !EditLayer->SupportsTargetType(ELandscapeToolTargetType::Heightmap))
-	{
-		OutReason = FText::Format(LOCTEXT("ClearTargetLayers_CantClearOnLayerType", "Cannot clear : the type of layer {0} ({1}) doesn't support heightmaps or weightmaps"),
-			FText::FromName(EditLayer->GetName()), EditLayer->GetClass()->GetDisplayNameText());
-		return false;
-	}
-
-	OutReason = LOCTEXT("ClearTargetLayers_CanClear", "Clear the target layers on the edit layer");
-	return true;
-}
-
-void FLandscapeEditorCustomNodeBuilder_Layers::ClearTargetLayersOnLayer(int32 InLayerIndex, ELandscapeClearMode InClearMode)
+bool FLandscapeEditorCustomNodeBuilder_Layers::CanClearEditLayerData(int32 InEditLayerIndex, ELandscapeToolTargetType InClearMode, FText& OutToolTip) const
 {
-	FEdModeLandscape* LandscapeEdMode = GetEditorMode();
-	ALandscape* Landscape = LandscapeEdMode ? LandscapeEdMode->GetLandscape() : nullptr;
-	if (Landscape)
+	if (FEdModeLandscape* LandscapeEdMode = GetEditorMode())
 	{
-		const ULandscapeEditLayerBase* EditLayer = LandscapeEdMode->GetEditLayerConst(InLayerIndex);
-		if (EditLayer)
+		switch (InClearMode)
 		{
-			EAppReturnType::Type Result = FMessageDialog::Open(EAppMsgType::YesNo, FText::Format(LOCTEXT("Landscape_ClearLayer_Message", "The layer {0} content will be completely cleared.  Continue?"), FText::FromName(EditLayer->GetName())));
-			if (Result == EAppReturnType::Yes)
-			{
-				const FScopedTransaction Transaction(LOCTEXT("Landscape_Layers_Clean", "Clear Layer"));
-				Landscape->ClearLayer(InLayerIndex, nullptr, InClearMode);
-				OnLayerSelectionChanged(InLayerIndex);
-				if (InClearMode & ELandscapeClearMode::Clear_Weightmap)
-				{
-					LandscapeEdMode->RequestUpdateLayerUsageInformation();
-				}
-			}
+			case ELandscapeToolTargetType::Invalid:
+				// In paint mode with no target layers 
+				OutToolTip = LOCTEXT("EditLayerContextMenu.Clear_WeightmapNoneToolTip", "No target layer data to clear for this edit layer.");
+				return false;
+			case ELandscapeToolTargetType::Heightmap:
+				OutToolTip = LOCTEXT("EditLayerContextMenu.Clear_HeightmapTooltip", "Clears all sculpting heightmap data for this edit layer.");
+				break;
+			case ELandscapeToolTargetType::Visibility:
+				OutToolTip = LOCTEXT("EditLayerContextMenu.Clear_VisibilityLayerTooltip", "Removes all holes painted in visibility for this edit layer.");
+				break;
+			case ELandscapeToolTargetType::Weightmap:
+				OutToolTip = LOCTEXT("EditLayerContextMenu.Clear_AllTargetLayerTooltip", "Clears all target layer weightmap data for this edit layer.");
+				break;
+			default:
+				check(false);
+				return false;
 		}
+
+		// InClearMode should always set the tooltip when valid
+		check(OutToolTip.ToString() != FString());
+		const ULandscapeEditLayerBase* EditLayer = LandscapeEdMode->GetEditLayerConst(InEditLayerIndex);
+		return EditLayer && LandscapeEdMode->GetLandscape() && LandscapeEdMode->CanEditLayer(&OutToolTip, EditLayer);
 	}
+	return false;
 }
 
 bool FLandscapeEditorCustomNodeBuilder_Layers::CanDeleteLayer(int32 InLayerIndex, FText& OutReason) const
diff --git a/Engine/Source/Editor/LandscapeEditor/Private/LandscapeEditorDetailCustomization_Layers.h b/Engine/Source/Editor/LandscapeEditor/Private/LandscapeEditorDetailCustomization_Layers.h
--- a/Engine/Source/Editor/LandscapeEditor/Private/LandscapeEditorDetailCustomization_Layers.h
+++ b/Engine/Source/Editor/LandscapeEditor/Private/LandscapeEditorDetailCustomization_Layers.h
@@ -85,11 +85,8 @@ protected:
 
 	const FSlateBrush* GetEditLayerIconBrush(int32 InLayerIndex) const;
 
-	bool CanClearTargetLayerOnLayer(int32 InLayerIndex, ULandscapeLayerInfoObject* InLayerInfo, FText& OutReason) const;
-	void ClearTargetLayerOnLayer(int32 InLayerIndex, ULandscapeLayerInfoObject* InLayerInfo);
-	bool CanClearLayer(int32 InLayerIndex, FText& OutReason) const;
-	bool CanClearTargetLayersOnLayer(int32 InLayerIndex, ELandscapeClearMode InClearMode, FText& OutReason) const;
-	void ClearTargetLayersOnLayer(int32 InLayerIndex, ELandscapeClearMode InClearMode);
+	void ClearEditLayerData(int32 InEditLayerIndex, ELandscapeToolTargetType InClearMode);
+	bool CanClearEditLayerData(int32 InEditLayerIndex, ELandscapeToolTargetType InClearMode, FText& OutToolTip) const;
 
 	bool CanRenameLayerTo(const FText& NewText, FText& OutErrorMessage, int32 InLayerIndex);
 	bool CanRenameLayer(int32 InLayerIndex, FText& OutReason) const;
@@ -132,9 +129,6 @@ protected:
 	void FillUnassignedBrushMenu(FMenuBuilder& MenuBuilder, TArray<ALandscapeBlueprintBrushBase*> Brushes, int32 InLayerIndex);
 	void AssignBrushToEditLayer(ALandscapeBlueprintBrushBase* Brush, int32 InLayerIndex);
 
-	void FillClearTargetLayerMenu(FMenuBuilder& MenuBuilder, int32 InLayerIndex, TArray<ULandscapeLayerInfoObject*> InUsedLayerInfos);
-	void FillClearLayerMenu(FMenuBuilder& MenuBuilder, int32 InLayerIndex);
-
 private:
 
 	/** Widgets for displaying and editing the layer name */
diff --git a/Engine/Source/Editor/LandscapeEditor/Private/LandscapeEditorDetailCustomization_TargetLayers.cpp b/Engine/Source/Editor/LandscapeEditor/Private/LandscapeEditorDetailCustomization_TargetLayers.cpp
--- a/Engine/Source/Editor/LandscapeEditor/Private/LandscapeEditorDetailCustomization_TargetLayers.cpp
+++ b/Engine/Source/Editor/LandscapeEditor/Private/LandscapeEditorDetailCustomization_TargetLayers.cpp
@@ -271,13 +271,8 @@ bool FLandscapeEditorDetailCustomization_TargetLayers::ShouldShowTargetLayers()
 	{
 		const FName CurrentToolName = LandscapeEdMode->CurrentTool->GetToolName();
 
-		//bool bSupportsHeightmap = (LandscapeEdMode->CurrentToolMode->SupportedTargetTypes & ELandscapeToolTargetTypeMask::Heightmap) != 0;
-		//bool bSupportsWeightmap = (LandscapeEdMode->CurrentToolMode->SupportedTargetTypes & ELandscapeToolTargetTypeMask::Weightmap) != 0;
-		//bool bSupportsVisibility = (LandscapeEdMode->CurrentToolMode->SupportedTargetTypes & ELandscapeToolTargetTypeMask::Visibility) != 0;
-
-		//// Visible if there are possible choices
-		//if (bSupportsWeightmap || bSupportsHeightmap || bSupportsVisibility)
-		if (LandscapeEdMode->CurrentToolMode->SupportedTargetTypes != 0 
+		// Target layers are included unless the tool focuses specifically on edit layers or components
+		if (LandscapeEdMode->CurrentToolMode->SupportedTargetTypes != ELandscapeToolTargetTypeFlags::None
 			&& CurrentToolName != TEXT("BlueprintBrush")
 			&& CurrentToolName != TEXT("Mask"))
 		{
@@ -440,7 +435,7 @@ void FLandscapeEditorCustomNodeBuilder_TargetLayers::GenerateHeaderRowContent(FD
 			.Text(LOCTEXT("LayersLabel", "Layers"))
 		];
 
-	if (LandscapeEdMode->CurrentToolMode->SupportedTargetTypes & ELandscapeToolTargetTypeMask::Weightmap)
+	if (EnumHasAllFlags(LandscapeEdMode->CurrentToolMode->SupportedTargetTypes, ELandscapeToolTargetTypeFlags::Weightmap))
 	{
 		NodeRow.ValueWidget
 		[
@@ -1102,7 +1097,7 @@ TSharedPtr<SWidget> FLandscapeEditorCustomNodeBuilder_TargetLayers::GenerateRow(
 
 	if (FEdModeLandscape* LandscapeEdMode = GetEditorMode())
 	{
-		if ((LandscapeEdMode->CurrentTool->GetSupportedTargetTypes() & LandscapeEdMode->CurrentToolMode->SupportedTargetTypes & ELandscapeToolTargetTypeMask::FromType(Target->TargetType)) == 0)
+		if ((LandscapeEdMode->CurrentTool->GetSupportedTargetTypes() & LandscapeEdMode->CurrentToolMode->SupportedTargetTypes & UE::Landscape::GetLandscapeToolTargetTypeAsFlags(Target->TargetType)) == ELandscapeToolTargetTypeFlags::None)
 		{
 			return RowWidget;
 		}
@@ -1549,8 +1544,8 @@ TSharedPtr<SWidget> FLandscapeEditorCustomNodeBuilder_TargetLayers::OnTargetLaye
 		if (FEdModeLandscape* LandscapeEdMode = GetEditorMode())
 		{
 			FUIAction LandscapeHeightmapChangeToolsAction = FUIAction(FExecuteAction::CreateStatic(&FLandscapeEditorCustomNodeBuilder_TargetLayers::OnHeightmapLayerContextMenu, Target));
-			MenuBuilder.AddMenuEntry(LOCTEXT("LayerContextMenu.Heightmap", "Import From/Export To File..."),
-				LOCTEXT("LayerContextMenu.HeightmapToolTip", "Opens the Landscape Import tool in order to import / export layers from / to external files."), FSlateIcon(), LandscapeHeightmapChangeToolsAction);
+			MenuBuilder.AddMenuEntry(LOCTEXT("LayerContextMenu.ImportExport", "Import From/Export To File..."),
+				LOCTEXT("LayerContextMenu.ImportExportToolTip", "Opens the Landscape Import tool in order to import / export layers from / to external files."), FSlateIcon(), LandscapeHeightmapChangeToolsAction);
 
 			if (Target->TargetType == ELandscapeToolTargetType::Weightmap && InLayerIndex != INDEX_NONE)
 			{
@@ -1576,32 +1571,20 @@ TSharedPtr<SWidget> FLandscapeEditorCustomNodeBuilder_TargetLayers::OnTargetLaye
 					LOCTEXT("LayerContextMenu.Fill_Tooltip", "Fills this layer to 100% across the entire landscape. If this is a weight-blended layer, all other weight-blended layers will be cleared."),
 					FSlateIcon(),
 					FillAction);
-
-				// Clear
-				FUIAction ClearAction = FUIAction(FExecuteAction::CreateStatic(&FLandscapeEditorCustomNodeBuilder_TargetLayers::OnClearLayer, Target),
-					FCanExecuteAction::CreateLambda([Target]()
-						{
-							return Target->LayerInfoObj != nullptr;
-						}));
-
-				MenuBuilder.AddMenuEntry(LOCTEXT("LayerContextMenu.Clear", "Clear Layer"),
-					LOCTEXT("LayerContextMenu.Clear_Tooltip", "Clears this layer to 0% across the entire landscape. If this is a weight-blended layer, other weight-blended layers will be adjusted to compensate."),
-					FSlateIcon(),
-					ClearAction);
 			}
-			else if (Target->TargetType == ELandscapeToolTargetType::Visibility)
-			{
-				MenuBuilder.AddMenuSeparator();
 
-				// Clear
-				FUIAction ClearAction = FUIAction(FExecuteAction::CreateStatic(&FLandscapeEditorCustomNodeBuilder_TargetLayers::OnClearLayer, Target),
-					FCanExecuteAction::CreateLambda([Target]()
-						{
-							return Target->LayerInfoObj != nullptr;
-						}));
-
-				MenuBuilder.AddMenuEntry(LOCTEXT("LayerContextMenu.ClearHoles", "Remove All Holes"), FText(), FSlateIcon(), ClearAction);
-			}
+			// Clear action shown for Weightmap, Heightmap, and Visibility 
+			TAttribute<FText> CanClearLayerToolTipText = MakeAttributeSPLambda(this, [this, Target] {FText Reason; CanClearLayer(Target, Reason); return Reason; });
+			FUIAction ClearAction = FUIAction(FExecuteAction::CreateStatic(&FLandscapeEditorCustomNodeBuilder_TargetLayers::OnClearLayer, Target),
+				FCanExecuteAction::CreateSPLambda(this, [this, Target] {FText Reason; return CanClearLayer(Target, Reason); }));
+			// Convert TargetType to UI label 
+			const ELandscapeToolTargetTypeFlags TargetTypeFlags = UE::Landscape::GetLandscapeToolTargetTypeAsFlags(Target->TargetType);
+			const FText& TargetTypeText = FText::FromString(UE::Landscape::GetLandscapeToolTargetTypeFlagsAsString(TargetTypeFlags));
+
+			MenuBuilder.AddMenuEntry(FText::Format(LOCTEXT("LayerContextMenu.ClearLayer", "Clear {0} Layer"), TargetTypeText),
+				CanClearLayerToolTipText,
+				FSlateIcon(),
+				ClearAction);
 		}
 	}
 	MenuBuilder.EndSection();
@@ -1774,20 +1757,67 @@ void FLandscapeEditorCustomNodeBuilder_TargetLayers::OnFillLayer(const TSharedRe
 
 void FLandscapeEditorCustomNodeBuilder_TargetLayers::OnClearLayer(const TSharedRef<FLandscapeTargetListInfo> Target)
 {
-	FScopedTransaction Transaction(LOCTEXT("Undo_ClearLayer", "Clearing Landscape Layer"));
-	if (Target->LandscapeInfo.IsValid() && Target->LayerInfoObj.IsValid())
+	FEdModeLandscape* LandscapeEdMode = GetEditorMode();
+	// Heightmap and Visibility will always have valid landscape info but Target layer may be unset
+	if (LandscapeEdMode && Target->LandscapeInfo.IsValid())
 	{
-		FEdModeLandscape* LandscapeEdMode = GetEditorMode();
-		if (LandscapeEdMode)
+		ALandscape* Landscape = LandscapeEdMode->GetLandscape();
+		const ULandscapeEditLayerBase* EditLayer = LandscapeEdMode->GetLandscapeSelectedLayer();
+		if (Landscape != nullptr && EditLayer != nullptr)
 		{
-			FScopedSetLandscapeEditingLayer Scope(LandscapeEdMode->GetLandscape(), LandscapeEdMode->GetCurrentLayerGuid(), [&] { LandscapeEdMode->RequestLayersContentUpdateForceAll(); });
-			FLandscapeEditDataInterface LandscapeEdit(Target->LandscapeInfo.Get());
-			LandscapeEdit.DeleteLayer(Target->LayerInfoObj.Get());
-			LandscapeEdMode->RequestUpdateLayerUsageInformation();
+			const FScopedTransaction Transaction(FText::Format(LOCTEXT("Landscape_TargetLayers_Clear", "Clear Edit Layer {0} Data"), UEnum::GetDisplayValueAsText(Target->TargetType)));
+
+			switch (Target->TargetType)
+			{
+				case ELandscapeToolTargetType::Weightmap: 
+					Landscape->ClearPaintLayer(LandscapeEdMode->GetSelectedEditLayerIndex(), Target->LayerInfoObj.Get()); 
+					break;
+				case ELandscapeToolTargetType::Heightmap:
+				case ELandscapeToolTargetType::Visibility: 
+					Landscape->ClearEditLayer(LandscapeEdMode->GetSelectedEditLayerIndex(), /*InComponents =*/nullptr, UE::Landscape::GetLandscapeToolTargetTypeAsFlags(Target->TargetType));
+					break;
+				default:
+					check(false);
+					return;
+			}
+	
+			if (Target->TargetType == ELandscapeToolTargetType::Weightmap || Target->TargetType == ELandscapeToolTargetType::Visibility)
+			{
+				LandscapeEdMode->RequestUpdateLayerUsageInformation();
+			}
 		}
 	}
 }
 
+bool FLandscapeEditorCustomNodeBuilder_TargetLayers::CanClearLayer(const TSharedRef<FLandscapeTargetListInfo> Target, FText& OutToolTip)
+{
+	FEdModeLandscape* LandscapeEdMode = GetEditorMode();
+	if (LandscapeEdMode && Target->LandscapeInfo.IsValid())
+	{
+		const ULandscapeEditLayerBase* EditLayer = LandscapeEdMode->GetLandscapeSelectedLayer();
+
+		switch (Target->TargetType)
+		{
+			case ELandscapeToolTargetType::Heightmap:
+				OutToolTip = LOCTEXT("LayerContextMenu.Clear_HeightmapTooltip", "Clears all sculpting heightmap data for this edit layer.");
+				break;
+			case ELandscapeToolTargetType::Visibility:
+				OutToolTip = LOCTEXT("LayerContextMenu.Clear_VisibilityLayerTooltip", "Removes all holes painted in the visibility for this edit layer.");
+				break;
+			case ELandscapeToolTargetType::Weightmap:
+				OutToolTip = LOCTEXT("LayerContextMenu.Clear_TargetLayerTooltip", "Clears this target layer to 0% across the entire landscape. If this is a weight-blended layer, other weight-blended layers will be adjusted to compensate.");
+				return LandscapeEdMode->CanEditTargetLayer(&OutToolTip, Target->LayerInfoObj.Get(), EditLayer);
+			default:
+				check(false);
+				return false;
+		}
+		// Return tooltip should always be set (Shows valid tooltip or reason why the operation is disabled)
+		check(OutToolTip.ToString() != FString());
+		return LandscapeEdMode->CanEditLayer(&OutToolTip, EditLayer);
+	}
+	return false;
+}
+
 void FLandscapeEditorCustomNodeBuilder_TargetLayers::OnRebuildMICs(const TSharedRef<FLandscapeTargetListInfo> Target)
 {
 	if (Target->LandscapeInfo.IsValid())
@@ -2249,4 +2279,4 @@ TSharedPtr<SWidget> FTargetLayerDragDropOp::GetDefaultDecorator() const
 		
 }
 
-#undef LOCTEXT_NAMESPACE
+#undef LOCTEXT_NAMESPACE
\ No newline at end of file
diff --git a/Engine/Source/Editor/LandscapeEditor/Private/LandscapeEditorDetailCustomization_TargetLayers.h b/Engine/Source/Editor/LandscapeEditor/Private/LandscapeEditorDetailCustomization_TargetLayers.h
--- a/Engine/Source/Editor/LandscapeEditor/Private/LandscapeEditorDetailCustomization_TargetLayers.h
+++ b/Engine/Source/Editor/LandscapeEditor/Private/LandscapeEditorDetailCustomization_TargetLayers.h
@@ -101,6 +101,7 @@ protected:
 	static void OnHeightmapLayerContextMenu(const TSharedRef<FLandscapeTargetListInfo> Target);
 	static void OnFillLayer(const TSharedRef<FLandscapeTargetListInfo> Target);
 	static void OnClearLayer(const TSharedRef<FLandscapeTargetListInfo> Target);
+	static bool CanClearLayer(const TSharedRef<FLandscapeTargetListInfo> Target, FText& OutToolTip);
 	static void OnRebuildMICs(const TSharedRef<FLandscapeTargetListInfo> Target);
 	static bool ShouldFilterLayerInfo(const struct FAssetData& AssetData, FName LayerName);
 	static void OnTargetLayerSetObject(const FAssetData& AssetData, const TSharedRef<FLandscapeTargetListInfo> Target); 
diff --git a/Engine/Source/Editor/LandscapeEditor/Public/LandscapeToolInterface.h b/Engine/Source/Editor/LandscapeEditor/Public/LandscapeToolInterface.h
--- a/Engine/Source/Editor/LandscapeEditor/Public/LandscapeToolInterface.h
+++ b/Engine/Source/Editor/LandscapeEditor/Public/LandscapeToolInterface.h
@@ -145,7 +145,7 @@ struct FLandscapeBrushSet
 
 namespace ELandscapeToolTargetTypeMask
 {
-	enum Type : uint8
+	enum UE_DEPRECATED(5.7, "Use ELandscapeToolTargetTypeFlags instead.")  Type : uint8
 	{
 		Heightmap  = 1 << static_cast<uint8>(ELandscapeToolTargetType::Heightmap),
 		Weightmap  = 1 << static_cast<uint8>(ELandscapeToolTargetType::Weightmap),
@@ -155,6 +155,8 @@ namespace ELandscapeToolTargetTypeMask
 		All = 0xFF,
 	};
 
+	UE_DEPRECATED(5.7, "Use UE::Landscape::GetLandscapeToolTargetTypeAsFlags with ELandscapeToolTargetType instead.")
+PRAGMA_DISABLE_DEPRECATION_WARNINGS
 	inline ELandscapeToolTargetTypeMask::Type FromType(ELandscapeToolTargetType TargetType)
 	{
 		if (TargetType == ELandscapeToolTargetType::Invalid)
@@ -162,6 +164,7 @@ namespace ELandscapeToolTargetTypeMask
 			return ELandscapeToolTargetTypeMask::NA;
 		}
 		return (ELandscapeToolTargetTypeMask::Type)(1 << static_cast<uint8>(TargetType));
+PRAGMA_ENABLE_DEPRECATION_WARNINGS
 	}
 }
 
@@ -259,7 +262,7 @@ public:
 	// Functions which doesn't need Viewport data...
 	virtual void Process(int32 Index, int32 Arg) {}
 	virtual ELandscapeToolType GetToolType() { return ELandscapeToolType::Normal; }
-	virtual ELandscapeToolTargetTypeMask::Type GetSupportedTargetTypes() { return ELandscapeToolTargetTypeMask::NA; };
+	virtual ELandscapeToolTargetTypeFlags GetSupportedTargetTypes() { return ELandscapeToolTargetTypeFlags::None; };
 
 	// FGCObject interface
 	virtual void AddReferencedObjects(FReferenceCollector& Collector) override {}
diff --git a/Engine/Source/Runtime/Landscape/Classes/Landscape.h b/Engine/Source/Runtime/Landscape/Classes/Landscape.h
--- a/Engine/Source/Runtime/Landscape/Classes/Landscape.h
+++ b/Engine/Source/Runtime/Landscape/Classes/Landscape.h
@@ -453,8 +453,17 @@ public:
 	LANDSCAPE_API void GetUsedPaintLayers(const FGuid& InLayerGuid, TArray<ULandscapeLayerInfoObject*>& OutUsedLayerInfos) const;
 	LANDSCAPE_API void ClearPaintLayer(int32 InLayerIndex, ULandscapeLayerInfoObject* InLayerInfo);
 	LANDSCAPE_API void ClearPaintLayer(const FGuid& InLayerGuid, ULandscapeLayerInfoObject* InLayerInfo);
-	LANDSCAPE_API void ClearLayer(int32 InLayerIndex, TSet<TObjectPtr<ULandscapeComponent>>* InComponents = nullptr, ELandscapeClearMode InClearMode = ELandscapeClearMode::Clear_All);
+
+PRAGMA_DISABLE_DEPRECATION_WARNINGS
+	UE_DEPRECATED(5.7, "ELandscapeClearMode is deprecated, use ClearEditLayer with ELandscapeToolTargetType instead.")
+	LANDSCAPE_API void ClearLayer(int32 InLayerIndex, TSet<TObjectPtr<ULandscapeComponent>>*InComponents = nullptr, ELandscapeClearMode InClearMode = ELandscapeClearMode::Clear_All);
+	UE_DEPRECATED(5.7, "ELandscapeClearMode is deprecated, use ClearEditLayer with ELandscapeToolTargetType instead.")
 	LANDSCAPE_API void ClearLayer(const FGuid& InLayerGuid, TSet<TObjectPtr<ULandscapeComponent>>* InComponents = nullptr, ELandscapeClearMode InClearMode = ELandscapeClearMode::Clear_All, bool bMarkPackageDirty = true);
+PRAGMA_ENABLE_DEPRECATION_WARNINGS
+
+	LANDSCAPE_API void ClearEditLayer(int32 InLayerIndex, TSet<TObjectPtr<ULandscapeComponent>>* InComponents = nullptr, ELandscapeToolTargetTypeFlags InClearMode = ELandscapeToolTargetTypeFlags::All);
+	LANDSCAPE_API void ClearEditLayer(const FGuid& InLayerGuid, TSet<TObjectPtr<ULandscapeComponent>>* InComponents = nullptr, ELandscapeToolTargetTypeFlags InClearMode = ELandscapeToolTargetTypeFlags::All, bool bMarkPackageDirty = true);
+	
 	LANDSCAPE_API bool DeleteLayer(int32 InLayerIndex);
 	LANDSCAPE_API void CollapseLayer(int32 InLayerIndex);
 	LANDSCAPE_API void DeleteLayers();
diff --git a/Engine/Source/Runtime/Landscape/Classes/LandscapeComponent.h b/Engine/Source/Runtime/Landscape/Classes/LandscapeComponent.h
--- a/Engine/Source/Runtime/Landscape/Classes/LandscapeComponent.h
+++ b/Engine/Source/Runtime/Landscape/Classes/LandscapeComponent.h
@@ -395,9 +395,8 @@ static const uint32 DefaultSplineHash = 0xFFFFFFFF;
 
 #endif
 
-// TODO [jonathan.bard] : Deprecate : this is now kinda redundant with ELandscapeToolTargetType
 UENUM()
-enum ELandscapeClearMode : int
+enum UE_DEPRECATED(5.7, "Use ELandscapeToolTargetTypeFlags instead.") ELandscapeClearMode : int
 {
 	Clear_Weightmap = 1 << 0 UMETA(DisplayName = "Paint"),
 	Clear_Heightmap = 1 << 1 UMETA(DisplayName = "Sculpt"),
diff --git a/Engine/Source/Runtime/Landscape/Private/LandscapeEditLayers.cpp b/Engine/Source/Runtime/Landscape/Private/LandscapeEditLayers.cpp
--- a/Engine/Source/Runtime/Landscape/Private/LandscapeEditLayers.cpp
+++ b/Engine/Source/Runtime/Landscape/Private/LandscapeEditLayers.cpp
@@ -12909,15 +12909,29 @@ void ALandscape::ClearPaintLayer(const FGuid& InLayerGuid, ULandscapeLayerInfoOb
 	});
 }
 
+PRAGMA_DISABLE_DEPRECATION_WARNINGS
+// Deprecated
 void ALandscape::ClearLayer(int32 InLayerIndex, TSet<TObjectPtr<ULandscapeComponent>>* InComponents, ELandscapeClearMode InClearMode)
+{
+	ClearEditLayer(InLayerIndex, InComponents, InClearMode == ELandscapeClearMode::Clear_Weightmap ? ELandscapeToolTargetTypeFlags::Weightmap : ELandscapeToolTargetTypeFlags::Heightmap);
+}
+
+// Deprecated
+void ALandscape::ClearLayer(const FGuid& InLayerGuid, TSet<TObjectPtr<ULandscapeComponent>>* InComponents, ELandscapeClearMode InClearMode, bool bMarkPackageDirty)
+{
+	ClearEditLayer(InLayerGuid, InComponents, InClearMode == ELandscapeClearMode::Clear_Weightmap ? ELandscapeToolTargetTypeFlags::Weightmap : ELandscapeToolTargetTypeFlags::Heightmap, bMarkPackageDirty);
+}
+PRAGMA_ENABLE_DEPRECATION_WARNINGS
+
+void ALandscape::ClearEditLayer(int32 InLayerIndex, TSet<TObjectPtr<ULandscapeComponent>>* InComponents, ELandscapeToolTargetTypeFlags InClearMode)
 {
 	if (const ULandscapeEditLayerBase* EditLayer = GetEditLayer(InLayerIndex))
 	{
-		ClearLayer(EditLayer->GetGuid(), InComponents, InClearMode);
+		ClearEditLayer(EditLayer->GetGuid(), InComponents, InClearMode);
 	}
 }
 
-void ALandscape::ClearLayer(const FGuid& InLayerGuid, TSet<TObjectPtr<ULandscapeComponent>>* InComponents, ELandscapeClearMode InClearMode, bool bMarkPackageDirty)
+void ALandscape::ClearEditLayer(const FGuid& InLayerGuid, TSet<TObjectPtr<ULandscapeComponent>>* InComponents, ELandscapeToolTargetTypeFlags InClearMode, bool bMarkPackageDirty)
 {
 	ensure(HasLayersContent());
 
@@ -12941,8 +12955,7 @@ void ALandscape::ClearLayer(const FGuid& InLayerGuid, TSet<TObjectPtr<ULandscape
 
 	TArray<uint16> NewHeightAlphaBlendData;
 	TArray<uint8> NewHeightFlagsData;
-
-	if (InClearMode & ELandscapeClearMode::Clear_Heightmap)
+	if (EnumHasAllFlags(InClearMode, ELandscapeToolTargetTypeFlags::Heightmap))
 	{
 		if (EditLayer->GetBlendMode() == LSBM_AlphaBlend)
 		{
@@ -12984,7 +12997,7 @@ void ALandscape::ClearLayer(const FGuid& InLayerGuid, TSet<TObjectPtr<ULandscape
 	FLandscapeDoNotDirtyScope DoNotDirtyScope(LandscapeEdit, !bMarkPackageDirty);
 	for (ULandscapeComponent* Component : Components)
 	{
-		if (InClearMode & ELandscapeClearMode::Clear_Heightmap)
+		if (EnumHasAllFlags(InClearMode, ELandscapeToolTargetTypeFlags::Heightmap))
 		{
 			int32 MinX = MAX_int32;
 			int32 MinY = MAX_int32;
@@ -12996,12 +13009,28 @@ void ALandscape::ClearLayer(const FGuid& InLayerGuid, TSet<TObjectPtr<ULandscape
 			LandscapeEdit.SetHeightData(MinX, MinY, MaxX, MaxY, NewHeightData.GetData(), 0, false, nullptr, NewHeightAlphaBlendData.GetData(), NewHeightFlagsData.GetData());
 		}
 
-		if (InClearMode & ELandscapeClearMode::Clear_Weightmap)
+		if (EnumHasAllFlags(InClearMode, ELandscapeToolTargetTypeFlags::Weightmap))
 		{
 			// Clear weight maps
 			for (FLandscapeInfoLayerSettings& LayerSettings : LandscapeInfo->Layers)
 			{
-				Component->DeleteLayer(LayerSettings.LayerInfoObj, LandscapeEdit);
+				if (LayerSettings.LayerInfoObj != ALandscapeProxy::VisibilityLayer)
+				{
+					Component->DeleteLayer(LayerSettings.LayerInfoObj, LandscapeEdit);
+				}
+			}
+		}
+
+		if (EnumHasAllFlags(InClearMode, ELandscapeToolTargetTypeFlags::Visibility))
+		{
+			// Clear visibility layer only
+			for (FLandscapeInfoLayerSettings& LayerSettings : LandscapeInfo->Layers)
+			{
+				if (LayerSettings.LayerInfoObj == ALandscapeProxy::VisibilityLayer)
+				{
+					Component->DeleteLayer(LayerSettings.LayerInfoObj, LandscapeEdit);
+					break;
+				}
 			}
 		}
 	}
@@ -13130,7 +13159,7 @@ void ALandscape::UpdateLandscapeSplines(const FGuid& InTargetLayer, bool bInUpda
 
 			// Clear layers without affecting weightmap allocations
 			const bool bMarkPackageDirty = false;
-			ClearLayer(SplinesEditLayer->GetGuid(), (!bInForceUpdateAllCompoments && LandscapeSplinesAffectedComponents.Num()) ? &LandscapeSplinesAffectedComponents : nullptr, Clear_All, bMarkPackageDirty);
+			ClearEditLayer(SplinesEditLayer->GetGuid(), (!bInForceUpdateAllCompoments && LandscapeSplinesAffectedComponents.Num()) ? &LandscapeSplinesAffectedComponents : nullptr, ELandscapeToolTargetTypeFlags::All, bMarkPackageDirty);
 			LandscapeSplinesAffectedComponents.Empty();
 			ModifiedComponent = &LandscapeSplinesAffectedComponents;
 			// For now, in Landscape Layer System Mode with a reserved layer for splines, we always update all the splines since we clear the whole layer first
