Fixed missing landscape spline tool localization

#rnx

diff --git a/Engine/Source/Editor/LandscapeEditor/Private/LandscapeSplineDetails.cpp b/Engine/Source/Editor/LandscapeEditor/Private/LandscapeSplineDetails.cpp
--- a/Engine/Source/Editor/LandscapeEditor/Private/LandscapeSplineDetails.cpp
+++ b/Engine/Source/Editor/LandscapeEditor/Private/LandscapeSplineDetails.cpp
@@ -309,7 +309,7 @@ FText FLandscapeSplineDetails::OnGetSplineOwningLandscapeText() const
 		}
 	}
 	
-	return FText::FromString("Owner: " + SplineOwnersStr);
+	return FText::Format(LOCTEXT("SplineOwner", "Owner: {0}"), FText::FromString(SplineOwnersStr));
 }
 
 FReply FLandscapeSplineDetails::OnSelectAdjacentLinearSplineElementButtonClicked(ESplineNavigationFlags Flags) const
