Landscape Editor - Updating Landscape Import/Export user experience and bug fixes
* Adds Reset to default button to heightmap, weightmap file paths
* Refactors CurrentGizmoActor into LandscapeEditorObject and exposes gizmo location in details view panel
* Removes unused GizmoHistory struct

diff --git a/Engine/Source/Editor/LandscapeEditor/Private/LandscapeEdMode.cpp b/Engine/Source/Editor/LandscapeEditor/Private/LandscapeEdMode.cpp
--- a/Engine/Source/Editor/LandscapeEditor/Private/LandscapeEdMode.cpp
+++ b/Engine/Source/Editor/LandscapeEditor/Private/LandscapeEdMode.cpp
@@ -108,7 +108,6 @@ FEdModeLandscape::FEdModeLandscape()
 	, PendingRenameTargetLayerIndex(INDEX_NONE)
 	, NewLandscapePreviewMode(ENewLandscapePreviewMode::None)
 	, ImportExportMode(EImportExportMode::Import)
-	, CurrentGizmoActor(nullptr)
 	, CopyPasteTool(nullptr)
 	, SplinesTool(nullptr)
 	, TargetLayerStartingIndex(0)
@@ -466,17 +465,18 @@ void FEdModeLandscape::Enter()
 		ALandscapeGizmoActiveActor* GizmoActor = *It;
 		if (GizmoActor->HasAnyFlags(RF_Transient))
 		{
-			CurrentGizmoActor = *It;
+			UISettings->CurrentGizmoActor = *It;
 			break;
 		}
 	}
 
-	if (!CurrentGizmoActor.IsValid())
+
+	if (!UISettings->CurrentGizmoActor.IsValid())
 	{
 		FActorSpawnParameters SpawnParams;
 		SpawnParams.ObjectFlags |= RF_Transient;
-		CurrentGizmoActor = GetWorld()->SpawnActor<ALandscapeGizmoActiveActor>(SpawnParams);
-		CurrentGizmoActor->ImportFromClipboard();
+		UISettings->CurrentGizmoActor = GetWorld()->SpawnActor<ALandscapeGizmoActiveActor>(SpawnParams);
+		UISettings->CurrentGizmoActor->ImportFromClipboard();
 	}
 
 	// Update list of landscapes and layers
@@ -506,6 +506,7 @@ void FEdModeLandscape::Enter()
 		}
 	}
 
+	TWeakObjectPtr<ALandscapeGizmoActiveActor> CurrentGizmoActor = UISettings->CurrentGizmoActor;
 	if (CurrentGizmoActor.IsValid())
 	{
 		CurrentGizmoActor->SetTargetLandscape(CurrentToolTarget.LandscapeInfo.Get());
@@ -748,7 +749,7 @@ void FEdModeLandscape::Exit()
 	GLandscapeEditRenderMode = ELandscapeEditRenderMode::None;
 	GLandscapeEditModeActive = false;
 
-	CurrentGizmoActor = nullptr;
+	UISettings->CurrentGizmoActor = nullptr;
 
 	GEditor->SelectNone(false, true);
 
@@ -860,7 +861,7 @@ void FEdModeLandscape::Tick(FEditorViewportClient* ViewportClient, float DeltaTi
 					CurrentBrush->Tick(ViewportClient, DeltaTime);
 				}
 			
-				if (CurrentBrush != GizmoBrush && CurrentGizmoActor.IsValid() && GizmoBrush && (GLandscapeEditRenderMode & ELandscapeEditRenderMode::Gizmo))
+				if (CurrentBrush != GizmoBrush && UISettings->CurrentGizmoActor.IsValid() && GizmoBrush && (GLandscapeEditRenderMode & ELandscapeEditRenderMode::Gizmo))
 				{
 					GizmoBrush->Tick(ViewportClient, DeltaTime);
 				}
@@ -1107,7 +1108,7 @@ bool FEdModeLandscape::StartTracking(FEditorViewportClient* InViewportClient, FV
 {
 	using namespace UE::Landscape::Editor::Private;
 
-	if (CurrentGizmoActor.IsValid() && CurrentGizmoActor->IsSelected() && GLandscapeEditRenderMode & ELandscapeEditRenderMode::Gizmo)
+	if (UISettings->CurrentGizmoActor.IsValid() && UISettings->CurrentGizmoActor->IsSelected() && GLandscapeEditRenderMode & ELandscapeEditRenderMode::Gizmo)
 	{
 		GIsGizmoDragging = true;
 		return true;
@@ -1453,13 +1454,13 @@ EEditAction::Type FEdModeLandscape::GetActionEditDelete()
 		if (Result == EEditAction::Skip)
 		{
 			// Prevent deleting Gizmo during LandscapeEdMode
-			if (CurrentGizmoActor.IsValid())
+			if (UISettings->CurrentGizmoActor.IsValid())
 			{
-				if (CurrentGizmoActor->IsSelected())
+				if (UISettings->CurrentGizmoActor->IsSelected())
 				{
 					if (GEditor->GetSelectedActors()->Num() > 1)
 					{
-						GEditor->GetSelectedActors()->Deselect(CurrentGizmoActor.Get());
+						GEditor->GetSelectedActors()->Deselect(UISettings->CurrentGizmoActor.Get());
 						Result = EEditAction::Skip;
 					}
 					else
@@ -1511,7 +1512,7 @@ EEditAction::Type FEdModeLandscape::GetActionEditCopy()
 		{
 			if (GLandscapeEditRenderMode & ELandscapeEditRenderMode::Gizmo || GLandscapeEditRenderMode & (ELandscapeEditRenderMode::Select))
 			{
-				if (CurrentGizmoActor.IsValid() && GizmoBrush && CurrentGizmoActor->TargetLandscapeInfo)
+				if (UISettings->CurrentGizmoActor.IsValid() && GizmoBrush && UISettings->CurrentGizmoActor->TargetLandscapeInfo)
 				{
 					Result = EEditAction::Process;
 				}
@@ -1537,7 +1538,7 @@ EEditAction::Type FEdModeLandscape::GetActionEditPaste()
 		{
 			if (GLandscapeEditRenderMode & ELandscapeEditRenderMode::Gizmo || GLandscapeEditRenderMode & (ELandscapeEditRenderMode::Select))
 			{
-				if (CurrentGizmoActor.IsValid() && GizmoBrush && CurrentGizmoActor->TargetLandscapeInfo)
+				if (UISettings->CurrentGizmoActor.IsValid() && GizmoBrush && UISettings->CurrentGizmoActor->TargetLandscapeInfo)
 				{
 					Result = EEditAction::Process;
 				}
@@ -1638,15 +1639,15 @@ bool FEdModeLandscape::ProcessEditCopy()
 
 			if (!CurrentlySelectedBPBrush)
 			{
-				bool IsSlowTask = IsSlowSelect(CurrentGizmoActor->TargetLandscapeInfo);
+				bool IsSlowTask = IsSlowSelect(UISettings->CurrentGizmoActor->TargetLandscapeInfo);
 				if (IsSlowTask)
 				{
 					GWarn->BeginSlowTask(LOCTEXT("BeginFitGizmoAndCopy", "Fit Gizmo to Selected Region and Copy Data..."), true);
 				}
 
 				FScopedTransaction Transaction(LOCTEXT("LandscapeGizmo_Copy", "Copy landscape data to Gizmo"));
-				CurrentGizmoActor->Modify();
-				CurrentGizmoActor->FitToSelection();
+				UISettings->CurrentGizmoActor->Modify();
+				UISettings->CurrentGizmoActor->FitToSelection();
 				CopyDataToGizmo();
 				SetCurrentTool(FName("CopyPaste"));
 
@@ -1707,7 +1708,7 @@ bool FEdModeLandscape::ProcessEditPaste()
 
 			if (!CurrentlySelectedBPBrush)
 			{
-				bool IsSlowTask = IsSlowSelect(CurrentGizmoActor->TargetLandscapeInfo);
+				bool IsSlowTask = IsSlowSelect(UISettings->CurrentGizmoActor->TargetLandscapeInfo);
 				if (IsSlowTask)
 				{
 					GWarn->BeginSlowTask(LOCTEXT("BeginPasteGizmoDataTask", "Paste Gizmo Data..."), true);
@@ -2004,7 +2005,7 @@ bool FEdModeLandscape::InputKey(FEditorViewportClient* ViewportClient, FViewport
 				return true;
 			}
 
-			if (CurrentGizmoActor.IsValid() && CurrentGizmoActor->IsSelected() && (GLandscapeEditRenderMode & ELandscapeEditRenderMode::Gizmo) && IsAltDown(Viewport))
+			if (UISettings->CurrentGizmoActor.IsValid() && UISettings->CurrentGizmoActor->IsSelected() && (GLandscapeEditRenderMode & ELandscapeEditRenderMode::Gizmo) && IsAltDown(Viewport))
 			{ 
 				// When manipulating the gizmo with Alt down, don't do anything but return true to indicate it's on purpose and prevent the editor from doing the usual thing
 				//  (alt-drag is usually "duplicate actor" and we want to prevent duplicating the gizmo actor) 
@@ -2250,9 +2251,9 @@ void FEdModeLandscape::SetCurrentTool(int32 ToolIndex, FName TargetLayerName)
 	}
 
 	// Update GizmoActor Landscape Target (is this necessary?)
-	if (CurrentGizmoActor.IsValid() && CurrentToolTarget.LandscapeInfo.IsValid())
+	if (UISettings->CurrentGizmoActor.IsValid() && CurrentToolTarget.LandscapeInfo.IsValid())
 	{
-		CurrentGizmoActor->SetTargetLandscape(CurrentToolTarget.LandscapeInfo.Get());
+		UISettings->CurrentGizmoActor->SetTargetLandscape(CurrentToolTarget.LandscapeInfo.Get());
 	}
 
 	if (Toolkit.IsValid())
@@ -2390,7 +2391,7 @@ int32 FEdModeLandscape::UpdateLandscapeList()
 		return FirstLevelEditor && Actor && FirstLevelEditor->GetEditorModeManager().IsEditingDisallowed(Actor);
 	};
 
-	if (!CurrentGizmoActor.IsValid())
+	if (!UISettings->CurrentGizmoActor.IsValid())
 	{
 		ALandscapeGizmoActiveActor* GizmoActor = nullptr;
 		if (TActorIterator<ALandscapeGizmoActiveActor> It(GetWorld()); It)
@@ -2422,9 +2423,9 @@ int32 FEdModeLandscape::UpdateLandscapeList()
 					CurrentIndex = Index;
 
 					// Update GizmoActor Landscape Target (is this necessary?)
-					if (CurrentGizmoActor.IsValid())
+					if (UISettings->CurrentGizmoActor.IsValid())
 					{
-						CurrentGizmoActor->SetTargetLandscape(LandscapeInfo);
+						UISettings->CurrentGizmoActor->SetTargetLandscape(LandscapeInfo);
 					}
 				}
 
@@ -2545,9 +2546,9 @@ void FEdModeLandscape::SetTargetLandscape(const TWeakObjectPtr<ULandscapeInfo>&
 	SetLandscapeInfo(InLandscapeInfo.Get());
 	// force a Leave and Enter the current tool, in case it has something about the current landscape cached
 	SetCurrentTool(CurrentToolIndex);
-	if (CurrentGizmoActor.IsValid())
+	if (UISettings->CurrentGizmoActor.IsValid())
 	{
-		CurrentGizmoActor->SetTargetLandscape(CurrentToolTarget.LandscapeInfo.Get());
+		UISettings->CurrentGizmoActor->SetTargetLandscape(CurrentToolTarget.LandscapeInfo.Get());
 	}
 
 	// register to new one
@@ -3034,7 +3035,7 @@ bool FEdModeLandscape::UsesTransformWidget() const
 		return true;
 	}
 
-	return (CurrentGizmoActor.IsValid() && CurrentGizmoActor->IsSelected() && (GLandscapeEditRenderMode & ELandscapeEditRenderMode::Gizmo));
+	return (UISettings->CurrentGizmoActor.IsValid() && UISettings->CurrentGizmoActor->IsSelected() && (GLandscapeEditRenderMode & ELandscapeEditRenderMode::Gizmo));
 }
 
 bool FEdModeLandscape::ShouldDrawWidget() const
@@ -3073,16 +3074,16 @@ FVector FEdModeLandscape::GetWidgetLocation() const
 		return UISettings->NewLandscape_Location;
 	}
 
-	if (CurrentGizmoActor.IsValid() && (GLandscapeEditRenderMode & ELandscapeEditRenderMode::Gizmo) && CurrentGizmoActor->IsSelected())
+	if (UISettings->CurrentGizmoActor.IsValid() && (GLandscapeEditRenderMode & ELandscapeEditRenderMode::Gizmo) && UISettings->CurrentGizmoActor->IsSelected())
 	{
-		if (CurrentGizmoActor->TargetLandscapeInfo && CurrentGizmoActor->TargetLandscapeInfo->GetLandscapeProxy())
+		if (UISettings->CurrentGizmoActor->TargetLandscapeInfo && UISettings->CurrentGizmoActor->TargetLandscapeInfo->GetLandscapeProxy())
 		{
 			// Apply Landscape transformation when it is available
-			ULandscapeInfo* LandscapeInfo = CurrentGizmoActor->TargetLandscapeInfo;
-			return CurrentGizmoActor->GetActorLocation()
-				+ FQuatRotationMatrix(LandscapeInfo->GetLandscapeProxy()->GetActorQuat()).TransformPosition(FVector(0, 0, CurrentGizmoActor->GetLength()));
+			ULandscapeInfo* LandscapeInfo = UISettings->CurrentGizmoActor->TargetLandscapeInfo;
+			return UISettings->CurrentGizmoActor->GetActorLocation()
+				+ FQuatRotationMatrix(LandscapeInfo->GetLandscapeProxy()->GetActorQuat()).TransformPosition(FVector(0, 0, UISettings->CurrentGizmoActor->GetLength()));
 		}
-		return CurrentGizmoActor->GetActorLocation();
+		return UISettings->CurrentGizmoActor->GetActorLocation();
 	}
 
 	// Override Widget for Splines Tool
@@ -3195,10 +3196,10 @@ bool FEdModeLandscape::IsSelectionAllowed(AActor* InActor, bool bInSelection) co
 /** FEdMode: Called when the currently selected actor has changed */
 void FEdModeLandscape::ActorSelectionChangeNotify()
 {
-	if (CurrentGizmoActor.IsValid() && CurrentGizmoActor->IsSelected() && (GEditor->GetSelectedActors()->CountSelections<AActor>() != 1))
+	if (UISettings->CurrentGizmoActor.IsValid() && UISettings->CurrentGizmoActor->IsSelected() && (GEditor->GetSelectedActors()->CountSelections<AActor>() != 1))
 	{
 		GEditor->SelectNone(false, true);
-		GEditor->SelectActor(CurrentGizmoActor.Get(), true, false, true);
+		GEditor->SelectActor(UISettings->CurrentGizmoActor.Get(), true, false, true);
 	}
 }
 
diff --git a/Engine/Source/Editor/LandscapeEditor/Private/LandscapeEdMode.h b/Engine/Source/Editor/LandscapeEditor/Private/LandscapeEdMode.h
--- a/Engine/Source/Editor/LandscapeEditor/Private/LandscapeEdMode.h
+++ b/Engine/Source/Editor/LandscapeEditor/Private/LandscapeEdMode.h
@@ -233,26 +233,6 @@ struct FLandscapeListInfo
 	}
 };
 
-struct FGizmoHistory
-{
-	ALandscapeGizmoActor* Gizmo;
-	FString GizmoName;
-
-	FGizmoHistory(ALandscapeGizmoActor* InGizmo)
-		: Gizmo(InGizmo)
-	{
-		GizmoName = Gizmo->GetPathName();
-	}
-
-	FGizmoHistory(ALandscapeGizmoActiveActor* InGizmo)
-	{
-		// handle for ALandscapeGizmoActiveActor -> ALandscapeGizmoActor
-		// ALandscapeGizmoActor is only for history, so it has limited data
-		Gizmo = InGizmo->SpawnGizmoActor();
-		GizmoName = Gizmo->GetPathName();
-	}
-};
-
 enum class ENewLandscapePreviewMode : uint8
 {
 	None,
@@ -306,7 +286,6 @@ public:
 	ENewLandscapePreviewMode NewLandscapePreviewMode;
 	EImportExportMode ImportExportMode;
 
-	TWeakObjectPtr<ALandscapeGizmoActiveActor> CurrentGizmoActor;
 	// UI callbacks for copy/paste tool
 	FLandscapeToolCopyPaste<FHeightmapToolTarget>* CopyPasteTool;
 	void CopyDataToGizmo();
diff --git a/Engine/Source/Editor/LandscapeEditor/Private/LandscapeEdModeBrushes.cpp b/Engine/Source/Editor/LandscapeEditor/Private/LandscapeEdModeBrushes.cpp
--- a/Engine/Source/Editor/LandscapeEditor/Private/LandscapeEdModeBrushes.cpp
+++ b/Engine/Source/Editor/LandscapeEditor/Private/LandscapeEdModeBrushes.cpp
@@ -667,7 +667,7 @@ public:
 	virtual void EnterBrush() override
 	{
 		// Make sure gizmo actor is selected
-		ALandscapeGizmoActiveActor* Gizmo = EdMode->CurrentGizmoActor.Get();
+		ALandscapeGizmoActiveActor* Gizmo = EdMode->UISettings->CurrentGizmoActor.Get();
 		if (Gizmo)
 		{
 			GEditor->SelectNone(false, true);
@@ -692,7 +692,7 @@ public:
 	{
 		if (GLandscapeEditRenderMode & ELandscapeEditRenderMode::Gizmo || GLandscapeEditRenderMode & ELandscapeEditRenderMode::Select)
 		{
-			ALandscapeGizmoActiveActor* Gizmo = EdMode->CurrentGizmoActor.Get();
+			ALandscapeGizmoActiveActor* Gizmo = EdMode->UISettings->CurrentGizmoActor.Get();
 
 			if (Gizmo && Gizmo->TargetLandscapeInfo && (Gizmo->TargetLandscapeInfo == EdMode->CurrentToolTarget.LandscapeInfo.Get()) && Gizmo->GizmoTexture && Gizmo->GetRootComponent())
 			{
@@ -730,8 +730,8 @@ public:
 						);
 					BrushMaterial->SetVectorParameterValue(FName(TEXT("AlphaScaleBias")), AlphaScaleBias);
 
-					float Angle = static_cast<float>(-EdMode->CurrentGizmoActor->GetActorRotation().Euler().Z * PI / 180.0f);
-					FVector4 LandscapeLocation(EdMode->CurrentGizmoActor->GetActorLocation().X, EdMode->CurrentGizmoActor->GetActorLocation().Y, EdMode->CurrentGizmoActor->GetActorLocation().Z, Angle);
+					float Angle = static_cast<float>(-EdMode->UISettings->CurrentGizmoActor->GetActorRotation().Euler().Z * PI / 180.0f);
+					FVector4 LandscapeLocation(EdMode->UISettings->CurrentGizmoActor->GetActorLocation().X, EdMode->UISettings->CurrentGizmoActor->GetActorLocation().Y, EdMode->UISettings->CurrentGizmoActor->GetActorLocation().Z, Angle);
 
 					BrushMaterial->SetDoubleVectorParameterValue(FName(TEXT("LandscapeLocation")), LandscapeLocation);
 					BrushMaterial->SetTextureParameterValue(FName(TEXT("AlphaTexture")), DataTexture);
@@ -788,7 +788,7 @@ public:
 	virtual FLandscapeBrushData ApplyBrush(const TArray<FLandscapeToolInteractorPosition>& InteractorPositions) override
 	{
 		// Selection Brush only works for 
-		ALandscapeGizmoActiveActor* Gizmo = EdMode->CurrentGizmoActor.Get();
+		ALandscapeGizmoActiveActor* Gizmo = EdMode->UISettings->CurrentGizmoActor.Get();
 		ULandscapeInfo* LandscapeInfo = EdMode->CurrentToolTarget.LandscapeInfo.Get();
 
 		if (!Gizmo || !Gizmo->GetRootComponent())
diff --git a/Engine/Source/Editor/LandscapeEditor/Private/LandscapeEdModeComponentTools.cpp b/Engine/Source/Editor/LandscapeEditor/Private/LandscapeEdModeComponentTools.cpp
--- a/Engine/Source/Editor/LandscapeEditor/Private/LandscapeEdModeComponentTools.cpp
+++ b/Engine/Source/Editor/LandscapeEditor/Private/LandscapeEdModeComponentTools.cpp
@@ -1020,7 +1020,7 @@ public:
 		TRACE_CPUPROFILER_EVENT_SCOPE(FLandscapeToolStrokeCopy_Apply);
 
 		//ULandscapeInfo* LandscapeInfo = EdMode->CurrentToolTarget.LandscapeInfo;
-		ALandscapeGizmoActiveActor* Gizmo = EdMode->CurrentGizmoActor.Get();
+		ALandscapeGizmoActiveActor* Gizmo = EdMode->UISettings->CurrentGizmoActor.Get();
 		if (LandscapeInfo && Gizmo && Gizmo->GizmoTexture && Gizmo->GetRootComponent())
 		{
 			Gizmo->TargetLandscapeInfo = LandscapeInfo;
@@ -1377,7 +1377,7 @@ public:
 		TRACE_CPUPROFILER_EVENT_SCOPE(FLandscapeToolStrokePaste_Apply);
 
 		//ULandscapeInfo* LandscapeInfo = EdMode->CurrentToolTarget.LandscapeInfo;
-		ALandscapeGizmoActiveActor* Gizmo = EdMode->CurrentGizmoActor.Get();
+		ALandscapeGizmoActiveActor* Gizmo = EdMode->UISettings->CurrentGizmoActor.Get();
 		// Cache and copy in Gizmo's region...
 		if (LandscapeInfo && Gizmo && Gizmo->GetRootComponent())
 		{
@@ -1731,7 +1731,7 @@ public:
 	virtual void EnterTool() override
 	{
 		// Make sure gizmo actor is selected
-		ALandscapeGizmoActiveActor* Gizmo = this->EdMode->CurrentGizmoActor.Get();
+		ALandscapeGizmoActiveActor* Gizmo = this->EdMode->UISettings->CurrentGizmoActor.Get();
 		if (Gizmo)
 		{
 			Gizmo->bFollowTerrainHeight = false;
@@ -1742,7 +1742,7 @@ public:
 
 	virtual void ExitTool() override
 	{
-		if (ALandscapeGizmoActiveActor* Gizmo = this->EdMode->CurrentGizmoActor.Get())
+		if (ALandscapeGizmoActiveActor* Gizmo = this->EdMode->UISettings->CurrentGizmoActor.Get())
 		{
 			Gizmo->bFollowTerrainHeight = true;
 		}
@@ -1774,10 +1774,10 @@ void FEdModeLandscape::CopyDataToGizmo()
 	{
 		CopyPasteTool->Copy();
 	}
-	if (CurrentGizmoActor.IsValid())
+	if (UISettings->CurrentGizmoActor.IsValid())
 	{
 		GEditor->SelectNone(false, true);
-		GEditor->SelectActor(CurrentGizmoActor.Get(), true, true, true);
+		GEditor->SelectActor(UISettings->CurrentGizmoActor.Get(), true, true, true);
 	}
 }
 
@@ -1788,10 +1788,10 @@ void FEdModeLandscape::PasteDataFromGizmo()
 	{
 		CopyPasteTool->Paste();
 	}
-	if (CurrentGizmoActor.IsValid())
+	if (UISettings->CurrentGizmoActor.IsValid())
 	{
 		GEditor->SelectNone(false, true);
-		GEditor->SelectActor(CurrentGizmoActor.Get(), true, true, true);
+		GEditor->SelectActor(UISettings->CurrentGizmoActor.Get(), true, true, true);
 	}
 }
 
@@ -2408,9 +2408,9 @@ public:
 		
 		if (EdMode->UISettings->ImportType != ELandscapeImportTransformType::Resample)
 		{
-			if (ALandscapeGizmoActiveActor* GizmoActor = EdMode->CurrentGizmoActor.Get())
+			if (ALandscapeGizmoActiveActor* GizmoActor = EdMode->UISettings->CurrentGizmoActor.Get())
 			{
-				return EdMode->CurrentGizmoActor->GetActorLocation();
+				return EdMode->UISettings->CurrentGizmoActor->GetActorLocation();
 			}	
 		}
 		return MinLocation;
@@ -2443,7 +2443,7 @@ public:
 				FVector LocalPosition(EdMode->UISettings->ImportLandscape_GizmoLocalPosition.X, EdMode->UISettings->ImportLandscape_GizmoLocalPosition.Y, 0.0f);
 				FVector GizmoPosition = LandscapeInfo->GetLandscapeProxy()->LandscapeActorToWorld().TransformPosition(LocalPosition);
 				
-				if (ALandscapeGizmoActiveActor* GizmoActor = EdMode->CurrentGizmoActor.Get())
+				if (ALandscapeGizmoActiveActor* GizmoActor = EdMode->UISettings->CurrentGizmoActor.Get())
 				{
 					GizmoActor->SetActorLocation(GizmoPosition);
 				}
@@ -2480,7 +2480,7 @@ public:
 			// Resample Gizmo can't move
 			if (EdMode && EdMode->UISettings->ImportType != ELandscapeImportTransformType::Resample)
 			{
-				if (ALandscapeGizmoActiveActor* GizmoActor = EdMode->CurrentGizmoActor.Get())
+				if (ALandscapeGizmoActiveActor* GizmoActor = EdMode->UISettings->CurrentGizmoActor.Get())
 				{
 					GEditor->ApplyDeltaToActor(
 						GizmoActor,
diff --git a/Engine/Source/Editor/LandscapeEditor/Private/LandscapeEditorDetailCustomization_CopyPaste.cpp b/Engine/Source/Editor/LandscapeEditor/Private/LandscapeEditorDetailCustomization_CopyPaste.cpp
--- a/Engine/Source/Editor/LandscapeEditor/Private/LandscapeEditorDetailCustomization_CopyPaste.cpp
+++ b/Engine/Source/Editor/LandscapeEditor/Private/LandscapeEditorDetailCustomization_CopyPaste.cpp
@@ -259,7 +259,7 @@ FReply FLandscapeEditorDetailCustomization_CopyPaste::OnFitGizmoToSelectionButto
 	FEdModeLandscape* LandscapeEdMode = GetEditorMode();
 	if (LandscapeEdMode != NULL)
 	{
-		ALandscapeGizmoActiveActor* Gizmo = LandscapeEdMode->CurrentGizmoActor.Get();
+		ALandscapeGizmoActiveActor* Gizmo = LandscapeEdMode->UISettings->CurrentGizmoActor.Get();
 		if (Gizmo && Gizmo->TargetLandscapeInfo)
 		{
 			FScopedTransaction Transaction(LOCTEXT("LandscapeGizmo_FitToSelection", "Fit gizmo size to selection"));
@@ -276,7 +276,7 @@ FReply FLandscapeEditorDetailCustomization_CopyPaste::OnFitHeightsToGizmoButtonC
 	FEdModeLandscape* LandscapeEdMode = GetEditorMode();
 	if (LandscapeEdMode != NULL)
 	{
-		ALandscapeGizmoActiveActor* Gizmo = LandscapeEdMode->CurrentGizmoActor.Get();
+		ALandscapeGizmoActiveActor* Gizmo = LandscapeEdMode->UISettings->CurrentGizmoActor.Get();
 		if (Gizmo && Gizmo->TargetLandscapeInfo)
 		{
 			FScopedTransaction Transaction(LOCTEXT("LandscapeGizmo_FitMinMaxHeight", "Set gizmo height to fix contained data"));
@@ -293,7 +293,7 @@ FReply FLandscapeEditorDetailCustomization_CopyPaste::OnClearGizmoDataButtonClic
 	FEdModeLandscape* LandscapeEdMode = GetEditorMode();
 	if (LandscapeEdMode != NULL)
 	{
-		ALandscapeGizmoActiveActor* Gizmo = LandscapeEdMode->CurrentGizmoActor.Get();
+		ALandscapeGizmoActiveActor* Gizmo = LandscapeEdMode->UISettings->CurrentGizmoActor.Get();
 		if (Gizmo && Gizmo->TargetLandscapeInfo)
 		{
 			Gizmo->ClearGizmoData();
@@ -462,7 +462,7 @@ FReply FLandscapeEditorDetailCustomization_CopyPaste::OnGizmoImportButtonClicked
 	FEdModeLandscape* LandscapeEdMode = GetEditorMode();
 	if (LandscapeEdMode != NULL)
 	{
-		ALandscapeGizmoActiveActor* Gizmo = LandscapeEdMode->CurrentGizmoActor.Get();
+		ALandscapeGizmoActiveActor* Gizmo = LandscapeEdMode->UISettings->CurrentGizmoActor.Get();
 		if (Gizmo)
 		{
 			TArray<uint8> Data;
@@ -532,7 +532,7 @@ FReply FLandscapeEditorDetailCustomization_CopyPaste::OnGizmoExportButtonClicked
 	FEdModeLandscape* LandscapeEdMode = GetEditorMode();
 	if (LandscapeEdMode != NULL)
 	{
-		ALandscapeGizmoActiveActor* Gizmo = LandscapeEdMode->CurrentGizmoActor.Get();
+		ALandscapeGizmoActiveActor* Gizmo = LandscapeEdMode->UISettings->CurrentGizmoActor.Get();
 		if (Gizmo && Gizmo->TargetLandscapeInfo && Gizmo->SelectedData.Num())
 		{
 			int32 TargetIndex = -1;
diff --git a/Engine/Source/Editor/LandscapeEditor/Private/LandscapeEditorDetailCustomization_ImportExport.cpp b/Engine/Source/Editor/LandscapeEditor/Private/LandscapeEditorDetailCustomization_ImportExport.cpp
--- a/Engine/Source/Editor/LandscapeEditor/Private/LandscapeEditorDetailCustomization_ImportExport.cpp
+++ b/Engine/Source/Editor/LandscapeEditor/Private/LandscapeEditorDetailCustomization_ImportExport.cpp
@@ -171,7 +171,36 @@ void FLandscapeEditorDetailCustomization_ImportExport::CustomizeDetails(IDetailL
 					.Text(NSLOCTEXT("UnrealEd", "GenericOpenDialog", "..."))
 					.OnClicked_Static(&OnBrowseFilenameButtonClicked, PropertyHandle_HeightmapFilename)
 				]
-			];
+			]
+			// Override the ResetToDefault so we can clear the file path and the check box on reset
+			.OverrideResetToDefault(FResetToDefaultOverride::Create(
+				FIsResetToDefaultVisible::CreateLambda([this, PropertyHandle_HeightmapFilename](TSharedPtr<IPropertyHandle> InProperty)
+				{
+					FString FilePathValue;
+					PropertyHandle_HeightmapFilename->GetValue(FilePathValue);
+
+					FEdModeLandscape* LandscapeEdMode = GetEditorMode();
+					if (LandscapeEdMode && LandscapeEdMode->UISettings)
+					{
+						// If the checkbox is selected or the file path is set show the reset button
+						return !FilePathValue.IsEmpty() || LandscapeEdMode->UISettings->bHeightmapSelected;
+					}
+
+					return false;
+				}),
+				FResetToDefaultHandler::CreateLambda([this, PropertyHandle_HeightmapFilename](TSharedPtr<IPropertyHandle> InProperty)
+				{
+					// Reset the file path 
+					PropertyHandle_HeightmapFilename->SetValue(FString());
+
+					// Reset the selected check box
+					FEdModeLandscape* LandscapeEdMode = GetEditorMode();
+					if (LandscapeEdMode && LandscapeEdMode->UISettings)
+					{
+						LandscapeEdMode->UISettings->bHeightmapSelected = false;
+					}
+				})
+			));
 	};
 
 	AddHeightmapFileName(PropertyHandle_ExportHeightmapFilename, MakeAttributeLambda([]() { return FLandscapeEditorDetailCustomization_ImportExport::GetImportExportVisibility(false); }));
@@ -287,21 +316,6 @@ void FLandscapeEditorDetailCustomization_ImportExport::CustomizeDetails(IDetailL
 	TSharedRef<IPropertyHandle> PropertyHandle_Layers = DetailBuilder.GetProperty(GET_MEMBER_NAME_CHECKED(ULandscapeEditorObject, ImportLandscape_Layers));
 	ImportExportCategory.AddProperty(PropertyHandle_Layers);
 
-	ImportExportCategory.AddCustomRow(FText::GetEmpty())
-	[
-		SNew(SHorizontalBox)
-		+ SHorizontalBox::Slot()
-		.FillWidth(1)
-		+ SHorizontalBox::Slot()
-		.AutoWidth()
-		[
-			SNew(SButton)
-			.Text(LOCTEXT("OriginReset", "Reset Import Location"))
-			.OnClicked_Static(&OnOriginResetButtonClicked)
-		]
-	]
-	.Visibility(TAttribute<EVisibility>(this, &FLandscapeEditorDetailCustomization_ImportExport::GetImportingVisibility));
-	
 	ImportExportCategory.AddCustomRow(FText::GetEmpty())
 	[
 		SNew(SHorizontalBox)
@@ -408,7 +422,7 @@ void FLandscapeEditorDetailCustomization_ImportExport::OnModeChanged(ECheckBoxSt
 			{
 				FVector LocalPosition(LandscapeEdMode->UISettings->ImportLandscape_GizmoLocalPosition.X, LandscapeEdMode->UISettings->ImportLandscape_GizmoLocalPosition.Y, 0.0f);
 				FVector GizmoPosition = LandscapeInfo->GetLandscapeProxy()->LandscapeActorToWorld().TransformPosition(LocalPosition);
-				if (ALandscapeGizmoActiveActor* GizmoActor = LandscapeEdMode->CurrentGizmoActor.Get())
+				if (ALandscapeGizmoActiveActor* GizmoActor = LandscapeEdMode->UISettings->CurrentGizmoActor.Get())
 				{
 					GizmoActor->SetActorLocation(GizmoPosition);
 				}
@@ -491,7 +505,8 @@ FReply FLandscapeEditorDetailCustomization_ImportExport::OnBrowseFilenameButtonC
 	const FString DialogTitle = bIsImporting ? LOCTEXT("ImportHeightmap", "Import Heightmap").ToString() : LOCTEXT("ExportHeightmap", "Export Heightmap").ToString();
 
 	TOptional<FString> OptionalFilename = LandscapeEditorUtils::GetImportExportFilename(DialogTitle, LandscapeEdMode->UISettings->LastImportPath, DialogType, bIsImporting);
-	if (OptionalFilename.IsSet())
+	// If the user canceled the dialog (Filename is empty), don't clear the path
+	if (OptionalFilename.IsSet() && !OptionalFilename->IsEmpty())
 	{
 		const FString& Filename = OptionalFilename.GetValue();
 
@@ -546,32 +561,6 @@ FText FLandscapeEditorDetailCustomization_ImportExport::GetImportExportButtonTex
 	return LOCTEXT("Export", "Export");
 }
 
-FReply FLandscapeEditorDetailCustomization_ImportExport::OnOriginResetButtonClicked()
-{
-	FEdModeLandscape* LandscapeEdMode = GetEditorMode();	
-	if (!LandscapeEdMode)
-	{
-		return FReply::Handled();
-	}
-
-	FScopedTransaction Transaction(LOCTEXT("LandscapeEdResetImportLocation", "Reset Import Location"));
-		
-	TWeakObjectPtr<ULandscapeInfo> LandscapeInfo = LandscapeEdMode->CurrentToolTarget.LandscapeInfo;
-	TWeakObjectPtr<ALandscapeGizmoActiveActor> CurrentGizmoActor = LandscapeEdMode->CurrentGizmoActor;
-
-	check(LandscapeInfo.Get());
-	check(CurrentGizmoActor.Get());
-	
-	FIntRect LandscapeExtent; 
-	LandscapeInfo->GetLandscapeExtent(LandscapeExtent);
-	const FTransform LandscapeToWorldTransform = LandscapeInfo->GetLandscapeProxy()->LandscapeActorToWorld();
-	const FVector LandscapeOrigin = LandscapeToWorldTransform.TransformPosition(FVector(LandscapeExtent.Min.X,LandscapeExtent.Min.Y,0.0f));
-
-	CurrentGizmoActor->SetActorLocation(LandscapeOrigin, false);
-	
-	return FReply::Handled();
-}
-
 FReply FLandscapeEditorDetailCustomization_ImportExport::OnImportExportButtonClicked()
 {
 	FEdModeLandscape* LandscapeEdMode = GetEditorMode();	
@@ -596,7 +585,7 @@ FReply FLandscapeEditorDetailCustomization_ImportExport::OnImportExportButtonCli
 						
 			ELandscapeImportTransformType TransformType = LandscapeEdMode->UISettings->ImportType;
 									
-			FVector LocalGizmoPosition = LandscapeInfo->GetLandscapeProxy()->LandscapeActorToWorld().InverseTransformPosition(LandscapeEdMode->CurrentGizmoActor->GetActorLocation());
+			FVector LocalGizmoPosition = LandscapeInfo->GetLandscapeProxy()->LandscapeActorToWorld().InverseTransformPosition(LandscapeEdMode->UISettings->CurrentGizmoActor->GetActorLocation());
 			FIntPoint LocalGizmoPoint = FIntPoint(FMath::FloorToInt32(LocalGizmoPosition.X), FMath::FloorToInt32(LocalGizmoPosition.Y));
 			
 			// Update Gizmo Position if we exit and comeback into tool
diff --git a/Engine/Source/Editor/LandscapeEditor/Private/LandscapeEditorDetailCustomization_ImportExport.h b/Engine/Source/Editor/LandscapeEditor/Private/LandscapeEditorDetailCustomization_ImportExport.h
--- a/Engine/Source/Editor/LandscapeEditor/Private/LandscapeEditorDetailCustomization_ImportExport.h
+++ b/Engine/Source/Editor/LandscapeEditor/Private/LandscapeEditorDetailCustomization_ImportExport.h
@@ -57,8 +57,6 @@ private:
 	static FText GetImportExportButtonText();
 	static FReply OnImportExportButtonClicked();
 	static bool GetImportExportButtonIsEnabled();
-	
-	static FReply OnOriginResetButtonClicked();
 
 	EVisibility GetImportExportLandscapeErrorVisibility() const;
 	EVisibility GetImportingVisibility() const;
diff --git a/Engine/Source/Editor/LandscapeEditor/Private/LandscapeEditorDetailCustomization_ImportLayers.cpp b/Engine/Source/Editor/LandscapeEditor/Private/LandscapeEditorDetailCustomization_ImportLayers.cpp
--- a/Engine/Source/Editor/LandscapeEditor/Private/LandscapeEditorDetailCustomization_ImportLayers.cpp
+++ b/Engine/Source/Editor/LandscapeEditor/Private/LandscapeEditorDetailCustomization_ImportLayers.cpp
@@ -36,14 +36,23 @@ TSharedRef<IPropertyTypeCustomization> FLandscapeEditorStructCustomization_FLand
 
 bool FLandscapeEditorStructCustomization_FLandscapeImportLayer::IsImporting()
 {
-	if (IsToolActive("ImportExport"))
+	if (FEdModeLandscape* LandscapeEdMode = GetEditorMode())
 	{
-		FEdModeLandscape* LandscapeEdMode = GetEditorMode();
-		check(LandscapeEdMode != nullptr);
-		return LandscapeEdMode->ImportExportMode == EImportExportMode::Import;
+		if (IsToolActive("ImportExport"))
+		{
+			return LandscapeEdMode->ImportExportMode == EImportExportMode::Import;
+		}
+
+		if (IsToolActive("NewLandscape"))
+		{
+			if (LandscapeEdMode->NewLandscapePreviewMode == ENewLandscapePreviewMode::ImportLandscape)
+			{
+				return true;
+			}
+		}
 	}
 
-	return true;
+	return false;
 }
 
 void FLandscapeEditorStructCustomization_FLandscapeImportLayer::CustomizeHeader(TSharedRef<IPropertyHandle> StructPropertyHandle, FDetailWidgetRow& HeaderRow, IPropertyTypeCustomizationUtils& StructCustomizationUtils)
@@ -87,135 +96,210 @@ void FLandscapeEditorStructCustomization_FLandscapeImportLayer::CustomizeChildre
 	Result = PropertyHandle_ThumbnailMIC->GetValue(ThumbnailMIC);
 	checkSlow(Result == FPropertyAccess::Success);
 
-	auto CreateFilenameWidget = [&](TSharedRef<IPropertyHandle> PropertyHandle_Filepath, bool bImport)
+	auto CreateFilePathNameContentWidget = [&](bool bImport)
 	{
 		return SNew(SHorizontalBox)
 			.Visibility_Static(&GetImportExportVisibility, bImport)
 			+ SHorizontalBox::Slot()
 			.AutoWidth()
-			.Padding(0, 0, 2, 0)
+			.VAlign(VAlign_Center)
+			.HAlign(HAlign_Left)
+			.Padding(FMargin(2))
 			[
-				SNew(SErrorText)
-				.Visibility_Static(&GetErrorVisibility, PropertyHandle_ImportResult)
-				.BackgroundColor_Static(&GetErrorColor, PropertyHandle_ImportResult)
-				.ErrorText(NSLOCTEXT("UnrealEd", "Error", "!"))
-				.ToolTip(
-					SNew(SToolTip)
-					.Text_Static(&GetErrorText, PropertyHandle_ErrorMessage)
-				)
+				SNew(SCheckBox)
+					.Visibility_Static(GetImportLayerSelectionVisibility)
+					.IsEnabled_Static(&FLandscapeEditorStructCustomization_FLandscapeImportLayer::IsValidLayerInfo, PropertyHandle_LayerInfo)
+					.IsChecked_Static(&FLandscapeEditorStructCustomization_FLandscapeImportLayer::GetImportLayerSelectedCheckState, PropertyHandle_Selected, PropertyHandle_LayerInfo)
+					.OnCheckStateChanged_Static(&FLandscapeEditorStructCustomization_FLandscapeImportLayer::OnImportLayerSelectedCheckStateChanged, PropertyHandle_Selected)
+					.ToolTipText_Static(&FLandscapeEditorStructCustomization_FLandscapeImportLayer::GetImportLayerSelectedToolTip, PropertyHandle_Selected, PropertyHandle_LayerInfo)
 			]
 			+ SHorizontalBox::Slot()
+			.AutoWidth()
+			.VAlign(VAlign_Center)
+			.Padding(FMargin(2))
 			[
-				PropertyHandle_Filepath->CreatePropertyValueWidget()
+				SNew(SLandscapeAssetThumbnail, ThumbnailMIC, StructCustomizationUtils.GetThumbnailPool().ToSharedRef(), FName("LandscapeEditor.Target_Weightmap"))
+					.ThumbnailSize(FIntPoint(48, 48))
 			]
 			+ SHorizontalBox::Slot()
-			.AutoWidth()
-			.Padding(1, 0, 0, 0)
+			.FillWidth(1)
+			.VAlign(VAlign_Center)
+			.Padding(FMargin(2))
 			[
-				SNew(SButton)
-				.ContentPadding(FMargin(4, 0))
-				.Text(NSLOCTEXT("UnrealEd", "GenericOpenDialog", "..."))
-				.OnClicked_Static(&FLandscapeEditorStructCustomization_FLandscapeImportLayer::OnLayerFilenameButtonClicked, PropertyHandle_Filepath)
+				SNew(STextBlock)
+					.Font(StructCustomizationUtils.GetRegularFont())
+					.Text(LayerNameText)
 			];
 	};
 
-	ChildBuilder.AddCustomRow(LayerNameText)
-	.NameContent()
-	[
-		SNew(SHorizontalBox)
-		+ SHorizontalBox::Slot()
-		.AutoWidth()
-		.VAlign(VAlign_Center)
-		.HAlign(HAlign_Left)
-		.Padding(FMargin(2))
+	auto CreateFilePathContentWidget = [&](TSharedRef<IPropertyHandle> PropertyHandle_Filepath, bool bImport)
+	{
+		return SNew(SBox)
+			.VAlign(VAlign_Center)
+			.Visibility_Static(&GetImportExportVisibility, bImport)
+			.IsEnabled_Static(&FLandscapeEditorStructCustomization_FLandscapeImportLayer::IsImportLayerSelected, PropertyHandle_Selected, PropertyHandle_LayerInfo)
+			[
+				SNew(SVerticalBox)
+					+ SVerticalBox::Slot()
+					.AutoHeight()
+					[
+						SNew(SHorizontalBox)
+							.Visibility_Static(&FLandscapeEditorStructCustomization_FLandscapeImportLayer::GetLayerInfoAssignVisibility)
+							+ SHorizontalBox::Slot()
+							[
+								SNew(SObjectPropertyEntryBox)
+									.AllowedClass(ULandscapeLayerInfoObject::StaticClass())
+									.PropertyHandle(PropertyHandle_LayerInfo)
+									.OnShouldFilterAsset_Static(&ShouldFilterLayerInfo, LayerName)
+									.AllowCreate(false)
+							]
+							+ SHorizontalBox::Slot()
+							.AutoWidth()
+							.VAlign(VAlign_Center)
+							[
+								SNew(SComboButton)
+									.ButtonStyle(FAppStyle::Get(), "HoverHintOnly")
+									.HasDownArrow(false)
+									.ContentPadding(4.0f)
+									.ForegroundColor(FSlateColor::UseForeground())
+									.IsFocusable(false)
+									.ToolTipText(LOCTEXT("Target_Create", "Create Layer Info"))
+									.Visibility_Static(&GetImportLayerCreateVisibility, PropertyHandle_LayerInfo)
+									.OnGetMenuContent_Static(&OnGetImportLayerCreateMenu, PropertyHandle_LayerInfo, LayerName)
+									.ButtonContent()
+									[
+										SNew(SImage)
+											.Image(FAppStyle::GetBrush("LandscapeEditor.Target_Create"))
+											.ColorAndOpacity(FSlateColor::UseForeground())
+									]
+							]
+					]
+					+ SVerticalBox::Slot()
+					.AutoHeight()
+					[
+						SNew(SVerticalBox)
+							.Visibility_Static(&FLandscapeEditorStructCustomization_FLandscapeImportLayer::GetImportLayerVisibility)
+							.IsEnabled_Static(&FLandscapeEditorStructCustomization_FLandscapeImportLayer::IsImportLayerSelected, PropertyHandle_Selected, PropertyHandle_LayerInfo)
+							+ SVerticalBox::Slot()
+							.AutoHeight()
+							[
+								SNew(SHorizontalBox)
+									+ SHorizontalBox::Slot()
+									.AutoWidth()
+									.Padding(0, 0, 2, 0)
+									[
+										SNew(SErrorText)
+											.Visibility_Static(&GetErrorVisibility, PropertyHandle_ImportResult)
+											.BackgroundColor_Static(&GetErrorColor, PropertyHandle_ImportResult)
+											.ErrorText(NSLOCTEXT("UnrealEd", "Error", "!"))
+											.ToolTip(
+												SNew(SToolTip)
+												.Text_Static(&GetErrorText, PropertyHandle_ErrorMessage)
+											)
+									]
+									+ SHorizontalBox::Slot()
+									[
+										PropertyHandle_Filepath->CreatePropertyValueWidget()
+									]
+									+ SHorizontalBox::Slot()
+									.AutoWidth()
+									.Padding(1, 0, 0, 0)
+									[
+										SNew(SButton)
+											.ContentPadding(FMargin(4, 0))
+											.Text(NSLOCTEXT("UnrealEd", "GenericOpenDialog", "..."))
+											.OnClicked_Static(&FLandscapeEditorStructCustomization_FLandscapeImportLayer::OnLayerFilenameButtonClicked, PropertyHandle_Filepath)
+									]
+							]
+					]
+			];
+	};
+
+	// Here we add both the Source (Import) and export file paths, one is always hidden by the visibility lambda
+	// Each property handle needs to be registered to a custom row for the Reset To Default arrows to be shown
+	IDetailPropertyRow& ImportFilePathRow = ChildBuilder.AddProperty(PropertyHandle_SourceFilePath)
+		.Visibility(TAttribute<EVisibility>::CreateLambda([PropertyHandle_Selected]() {
+			return GetImportExportVisibility(/*bImport=*/true);
+		}));
+	
+	ImportFilePathRow.CustomWidget()
+		.NameContent()
 		[
-			SNew(SCheckBox)
-			.Visibility_Static(GetImportLayerSelectionVisibility)
-			.IsEnabled_Static(&FLandscapeEditorStructCustomization_FLandscapeImportLayer::IsValidLayerInfo, PropertyHandle_LayerInfo)
-			.IsChecked_Static(&FLandscapeEditorStructCustomization_FLandscapeImportLayer::GetImportLayerSelectedCheckState, PropertyHandle_Selected, PropertyHandle_LayerInfo)
-			.OnCheckStateChanged_Static(&FLandscapeEditorStructCustomization_FLandscapeImportLayer::OnImportLayerSelectedCheckStateChanged, PropertyHandle_Selected)
-			.ToolTipText_Static(&FLandscapeEditorStructCustomization_FLandscapeImportLayer::GetImportLayerSelectedToolTip, PropertyHandle_Selected, PropertyHandle_LayerInfo)
+			CreateFilePathNameContentWidget(/*bImport=*/true)
 		]
-		+ SHorizontalBox::Slot()
-		.AutoWidth()
-		.VAlign(VAlign_Center)
-		.Padding(FMargin(2))
+		.ValueContent()
+		.MinDesiredWidth(250.0f) // copied from SPropertyEditorAsset::GetDesiredWidth
+		.MaxDesiredWidth(0)
 		[
-			SNew(SLandscapeAssetThumbnail, ThumbnailMIC, StructCustomizationUtils.GetThumbnailPool().ToSharedRef(), FName("LandscapeEditor.Target_Weightmap"))
-			.ThumbnailSize(FIntPoint(48, 48))
+			CreateFilePathContentWidget(PropertyHandle_SourceFilePath, /*bImport=*/true)
 		]
-		+ SHorizontalBox::Slot()
-		.FillWidth(1)
-		.VAlign(VAlign_Center)
-		.Padding(FMargin(2))
+		.OverrideResetToDefault(FResetToDefaultOverride::Create(
+			FIsResetToDefaultVisible::CreateLambda([this, PropertyHandle_Selected, PropertyHandle_SourceFilePath, PropertyHandle_ExportFilePath](TSharedPtr<IPropertyHandle> InProperty)
+			{
+				return IsResetFilePathButtonShown(PropertyHandle_Selected, PropertyHandle_SourceFilePath, PropertyHandle_ExportFilePath);
+			}),
+			FResetToDefaultHandler::CreateLambda([this, PropertyHandle_LayerInfo, PropertyHandle_Selected, PropertyHandle_SourceFilePath, PropertyHandle_ExportFilePath](TSharedPtr<IPropertyHandle> InProperty)
+			{
+				OnResetFilePathButtonClicked(PropertyHandle_LayerInfo, PropertyHandle_Selected, PropertyHandle_SourceFilePath, PropertyHandle_ExportFilePath);
+			})
+		));
+
+	// Export file path (row hidden while in import mode)
+	IDetailPropertyRow& ExportFilePathRow = ChildBuilder.AddProperty(PropertyHandle_ExportFilePath)
+		.Visibility(TAttribute<EVisibility>::CreateLambda([PropertyHandle_Selected]() {
+			return GetImportExportVisibility(/*bImport=*/false);
+		}));
+	ExportFilePathRow.CustomWidget()
+		.NameContent()
 		[
-			SNew(STextBlock)
-			.Font(StructCustomizationUtils.GetRegularFont())
-			.Text(LayerNameText)
+			CreateFilePathNameContentWidget(/*bImport=*/false)
 		]
-	]
-	.ValueContent()
-	.MinDesiredWidth(250.0f) // copied from SPropertyEditorAsset::GetDesiredWidth
-	.MaxDesiredWidth(0)
-	[
-		SNew(SBox)
-		.VAlign(VAlign_Center)
-		.Padding(FMargin(0, 0, 12, 0)) // Line up with the other properties due to having no reset to default button
+		.ValueContent()
+		.MinDesiredWidth(250.0f) // copied from SPropertyEditorAsset::GetDesiredWidth
+		.MaxDesiredWidth(0)
 		[
-			SNew(SVerticalBox)
-			+ SVerticalBox::Slot()
-			.AutoHeight()
-			[
-				SNew(SHorizontalBox)
-				.Visibility_Static(&FLandscapeEditorStructCustomization_FLandscapeImportLayer::GetLayerInfoAssignVisibility)
-				+ SHorizontalBox::Slot()
-				[
-					SNew(SObjectPropertyEntryBox)
-					.AllowedClass(ULandscapeLayerInfoObject::StaticClass())
-					.PropertyHandle(PropertyHandle_LayerInfo)
-					.OnShouldFilterAsset_Static(&ShouldFilterLayerInfo, LayerName)
-					.AllowCreate(false)
-				]
-				+ SHorizontalBox::Slot()
-				.AutoWidth()
-				.VAlign(VAlign_Center)
-				[
-					SNew(SComboButton)
-					.ButtonStyle(FAppStyle::Get(), "HoverHintOnly")
-					.HasDownArrow(false)
-					.ContentPadding(4.0f)
-					.ForegroundColor(FSlateColor::UseForeground())
-					.IsFocusable(false)
-					.ToolTipText(LOCTEXT("Target_Create", "Create Layer Info"))
-					.Visibility_Static(&GetImportLayerCreateVisibility, PropertyHandle_LayerInfo)
-					.OnGetMenuContent_Static(&OnGetImportLayerCreateMenu, PropertyHandle_LayerInfo, LayerName)
-					.ButtonContent()
-					[
-						SNew(SImage)
-						.Image(FAppStyle::GetBrush("LandscapeEditor.Target_Create"))
-						.ColorAndOpacity(FSlateColor::UseForeground())
-					]
-				]
-			]
-			+ SVerticalBox::Slot()
-			.AutoHeight()
-			[
-				SNew(SVerticalBox)
-				.Visibility_Static(&FLandscapeEditorStructCustomization_FLandscapeImportLayer::GetImportLayerVisibility)
-				.IsEnabled_Static(&FLandscapeEditorStructCustomization_FLandscapeImportLayer::IsImportLayerSelected, PropertyHandle_Selected, PropertyHandle_LayerInfo)
-				+ SVerticalBox::Slot()
-				.AutoHeight()
-				[
-					CreateFilenameWidget(PropertyHandle_SourceFilePath, true)
-				]
-				+ SVerticalBox::Slot()
-				.AutoHeight()
-				[
-					CreateFilenameWidget(PropertyHandle_ExportFilePath, false)
-				]
-			]
+			CreateFilePathContentWidget(PropertyHandle_ExportFilePath, /*bImport=*/false)
 		]
-	];
+		.OverrideResetToDefault(FResetToDefaultOverride::Create(
+			FIsResetToDefaultVisible::CreateLambda([this, PropertyHandle_Selected, PropertyHandle_SourceFilePath, PropertyHandle_ExportFilePath](TSharedPtr<IPropertyHandle> InProperty)
+			{
+				return IsResetFilePathButtonShown(PropertyHandle_Selected, PropertyHandle_SourceFilePath, PropertyHandle_ExportFilePath);
+			}),
+			FResetToDefaultHandler::CreateLambda([this, PropertyHandle_LayerInfo, PropertyHandle_Selected, PropertyHandle_SourceFilePath, PropertyHandle_ExportFilePath](TSharedPtr<IPropertyHandle> InProperty)
+			{
+				OnResetFilePathButtonClicked(PropertyHandle_LayerInfo, PropertyHandle_Selected, PropertyHandle_SourceFilePath, PropertyHandle_ExportFilePath);
+			})
+		));
+}
+
+void FLandscapeEditorStructCustomization_FLandscapeImportLayer::OnResetFilePathButtonClicked(TSharedRef<IPropertyHandle> PropertyHandle_LayerInfo, TSharedRef<IPropertyHandle> PropertyHandle_Selected, TSharedRef<IPropertyHandle> PropertyHandle_SourceFilePath, TSharedRef<IPropertyHandle> PropertyHandle_ExportFilePath)
+{
+	// Reset the file path 
+	IsImporting() ? PropertyHandle_SourceFilePath->SetValue(FString()) : PropertyHandle_ExportFilePath->SetValue(FString());
+
+	// Reset the selected check box
+	PropertyHandle_Selected->SetValue(false);
+
+	if (IsValidLayerInfo(PropertyHandle_LayerInfo))
+	{
+		// In New Landscape Import mode, each row has an optional layer info selector that should be reset
+		ULandscapeLayerInfoObject* EmptyLayerInfo = nullptr;
+		PropertyHandle_LayerInfo->SetValue(EmptyLayerInfo);
+	}
+}
+
+bool FLandscapeEditorStructCustomization_FLandscapeImportLayer::IsResetFilePathButtonShown(TSharedRef<IPropertyHandle> PropertyHandle_Selected, TSharedRef<IPropertyHandle> PropertyHandle_SourceFilePath, TSharedRef<IPropertyHandle> PropertyHandle_ExportFilePath)
+{
+	TSharedRef<IPropertyHandle> CurrentFilePathHandle = IsImporting() ? PropertyHandle_SourceFilePath : PropertyHandle_ExportFilePath;
+
+	FString FilePathValue;
+	CurrentFilePathHandle->GetValue(FilePathValue);
+
+	bool bIsSelected = false;
+	PropertyHandle_Selected->GetValue(bIsSelected);
+
+	// If the checkbox is selected or the file path is set show the reset button
+	return bIsSelected || !FilePathValue.IsEmpty();
 }
 
 void FLandscapeEditorStructCustomization_FLandscapeImportLayer::OnImportWeightmapFilenameChanged()
@@ -241,7 +325,8 @@ FReply FLandscapeEditorStructCustomization_FLandscapeImportLayer::OnLayerFilenam
 	const FString DialogTitle = bIsImporting ? LOCTEXT("ImportLayer", "Import Layer").ToString() : LOCTEXT("ExportLayer", "Export Layer").ToString();
 
 	TOptional<FString> OptionalExportImportPath = LandscapeEditorUtils::GetImportExportFilename(DialogTitle, LandscapeEdMode->UISettings->LastImportPath, DialogTypeString, bIsImporting);
-	if (OptionalExportImportPath.IsSet())
+	// If the user canceled the dialog (Filename is empty), don't clear the path
+	if (OptionalExportImportPath.IsSet() && !OptionalExportImportPath->IsEmpty())
 	{
 		const FString& Filename = OptionalExportImportPath.GetValue();
 
diff --git a/Engine/Source/Editor/LandscapeEditor/Private/LandscapeEditorDetailCustomization_ImportLayers.h b/Engine/Source/Editor/LandscapeEditor/Private/LandscapeEditorDetailCustomization_ImportLayers.h
--- a/Engine/Source/Editor/LandscapeEditor/Private/LandscapeEditorDetailCustomization_ImportLayers.h
+++ b/Engine/Source/Editor/LandscapeEditor/Private/LandscapeEditorDetailCustomization_ImportLayers.h
@@ -19,7 +19,7 @@ public:
 	virtual void CustomizeHeader(TSharedRef<IPropertyHandle> StructPropertyHandle, class FDetailWidgetRow& HeaderRow, IPropertyTypeCustomizationUtils& StructCustomizationUtils) override;
 	virtual void CustomizeChildren(TSharedRef<IPropertyHandle> StructPropertyHandle, class IDetailChildrenBuilder& ChildBuilder, IPropertyTypeCustomizationUtils& StructCustomizationUtils) override;
 
-public:
+private:
 	static void OnImportWeightmapFilenameChanged();
 	static FReply OnLayerFilenameButtonClicked(TSharedRef<IPropertyHandle> PropertyHandle_LayerFilename);
 	static bool ShouldFilterLayerInfo(const struct FAssetData& AssetData, FName LayerName);
@@ -36,6 +36,8 @@ public:
 	static FSlateColor GetErrorColor(TSharedRef<IPropertyHandle> PropertyHandle_ImportResult);
 	static FText GetErrorText(TSharedRef<IPropertyHandle> PropertyHandle_ErrorMessage);
 
+	static void OnResetFilePathButtonClicked(TSharedRef<IPropertyHandle> PropertyHandle_LayerInfo, TSharedRef<IPropertyHandle> PropertyHandle_Selected, TSharedRef<IPropertyHandle> PropertyHandle_SourceFilePath, TSharedRef<IPropertyHandle> PropertyHandle_ExportFilePath);
+	static bool IsResetFilePathButtonShown(TSharedRef<IPropertyHandle> PropertyHandle_Selected, TSharedRef<IPropertyHandle> PropertyHandle_SourceFilePath, TSharedRef<IPropertyHandle> PropertyHandle_ExportFilePath);
 	static bool IsImportLayerSelected(TSharedRef<IPropertyHandle> PropertyHandle_Selected, TSharedRef<IPropertyHandle> PropertyHandle_LayerInfo);
 	static ECheckBoxState GetImportLayerSelectedCheckState(TSharedRef<IPropertyHandle> PropertyHandle_Selected, TSharedRef<IPropertyHandle> PropertyHandle_LayerInfo);
 	static void OnImportLayerSelectedCheckStateChanged(ECheckBoxState CheckState, TSharedRef<IPropertyHandle> PropertyHandle_Selected);
diff --git a/Engine/Source/Editor/LandscapeEditor/Private/LandscapeEditorDetailCustomization_MiscTools.cpp b/Engine/Source/Editor/LandscapeEditor/Private/LandscapeEditorDetailCustomization_MiscTools.cpp
--- a/Engine/Source/Editor/LandscapeEditor/Private/LandscapeEditorDetailCustomization_MiscTools.cpp
+++ b/Engine/Source/Editor/LandscapeEditor/Private/LandscapeEditorDetailCustomization_MiscTools.cpp
@@ -8,6 +8,7 @@
 #include "Widgets/Layout/SBox.h"
 #include "Widgets/Input/SButton.h"
 #include "Widgets/Input/SCheckBox.h"
+#include "Widgets/Input/SVectorInputBox.h"
 #include "LandscapeEditorObject.h"
 
 #include "DetailLayoutBuilder.h"
@@ -31,6 +32,11 @@ BEGIN_SLATE_FUNCTION_BUILD_OPTIMIZATION
 void FLandscapeEditorDetailCustomization_MiscTools::CustomizeDetails(IDetailLayoutBuilder& DetailBuilder)
 {
 	IDetailCategoryBuilder& ToolsCategory = DetailBuilder.EditCategory("Tool Settings");
+	FEdModeLandscape* LandscapeEdMode = GetEditorMode();
+	if (!LandscapeEdMode)
+	{
+		return;
+	}
 
 	if (IsBrushSetActive("BrushSet_Component"))
 	{
@@ -139,8 +145,6 @@ void FLandscapeEditorDetailCustomization_MiscTools::CustomizeDetails(IDetailLayo
 			]
 		];
 
-		const FEdModeLandscape* LandscapeEdMode = GetEditorMode();
-
 		// Once a Splines Edit Layer exists, disable the buttons since spline updates happen automatically
 		ToolsCategory.AddCustomRow(LOCTEXT("ApplySplinesLabel", "Apply Splines"))
 		[
@@ -329,9 +333,118 @@ void FLandscapeEditorDetailCustomization_MiscTools::CustomizeDetails(IDetailLayo
 		]
 		.Visibility(TAttribute<EVisibility>(this, &FLandscapeEditorDetailCustomization_MiscTools::GetMiscLandscapeErrorVisibility));
 	}
+
+	// Display the Current Gizmo location (Vector3) 
+	// Only revealing a single ALandscapeGizmoActiveActor property - if more fields need to be exposed later, consider registering a struct customization instead
+	TSharedRef<IPropertyHandle> CurrentGizmoHandle = DetailBuilder.GetProperty(GET_MEMBER_NAME_CHECKED(ULandscapeEditorObject, CurrentGizmoActor));
+	TWeakObjectPtr<ALandscapeGizmoActiveActor> GizmoActor = LandscapeEdMode->UISettings->CurrentGizmoActor.Get();
+	if (GizmoActor.IsValid())
+	{
+		FSlateFontInfo FontInfo = IDetailLayoutBuilder::GetDetailFont();
+
+		ToolsCategory.AddProperty(CurrentGizmoHandle)
+			.CustomWidget()
+			.NameContent()
+			[
+				SNew(STextBlock)
+					.Text(LOCTEXT("LandscapeImportGizmoLocation", "Import Gizmo Location"))
+					.Font(IDetailLayoutBuilder::GetDetailFont())
+			]
+			.ValueContent()
+			.MinDesiredWidth(125.0f * 3.0f)
+			.MaxDesiredWidth(125.0f * 3.0f)
+			.VAlign(VAlign_Center)
+			[
+				// Since FComponentTransformDetails::GenerateChildContent is not exposed, create a minimal vector component to display location
+				// Drag features intentionally omitted to keep this widget minimal
+				SNew(SVectorInputBox)
+					.bColorAxisLabels(true)
+					.Font(FontInfo)
+					.X(TAttribute<TOptional<float>>::Create([GizmoActor]() { return GizmoActor.IsValid() ? GizmoActor->GetActorLocation().X : 0; }))
+					.Y(TAttribute<TOptional<float>>::Create([GizmoActor]() { return GizmoActor.IsValid() ? GizmoActor->GetActorLocation().Y : 0; }))
+					.Z(TAttribute<TOptional<float>>::Create([GizmoActor]() { return GizmoActor.IsValid() ? GizmoActor->GetActorLocation().Z : 0; }))
+					.OnXCommitted_Lambda([GizmoActor](float NewX, ETextCommit::Type) 
+					{
+						if (GizmoActor.IsValid())
+						{
+							FScopedTransaction Transaction(LOCTEXT("LandscapeUpdateGizmoX", "Landscape Update Gizmo X Location"));
+							GizmoActor->Modify();
+							// Apply delta from the current location
+							GizmoActor->EditorApplyTranslation(FVector(NewX - GizmoActor->GetActorLocation().X, 0, 0), /*bAltDown =*/false, /*bShiftDown =*/false, /*bCtrlDown =*/false);
+						}
+					})
+					.OnYCommitted_Lambda([GizmoActor](float NewY, ETextCommit::Type) 
+					{
+						if (GizmoActor.IsValid())
+						{
+							FScopedTransaction Transaction(LOCTEXT("LandscapeUpdateGizmoY", "Landscape Update Gizmo Y Location"));
+							GizmoActor->Modify();
+							// Apply delta from the current location
+							GizmoActor->EditorApplyTranslation(FVector(0, NewY - GizmoActor->GetActorLocation().Y, 0), /*bAltDown =*/false, /*bShiftDown =*/false, /*bCtrlDown =*/false);
+						}
+					})
+					.OnZCommitted_Lambda([GizmoActor](float NewZ, ETextCommit::Type) 
+					{
+						if (GizmoActor.IsValid())
+						{
+							FScopedTransaction Transaction(LOCTEXT("LandscapeUpdateGizmoZ", "Landscape Update Gizmo Z Location"));
+							GizmoActor->Modify();
+							// Apply delta from the current location
+							GizmoActor->EditorApplyTranslation(FVector(0, 0, NewZ - GizmoActor->GetActorLocation().Z), /*bAltDown =*/false, /*bShiftDown =*/false, /*bCtrlDown =*/false);
+						}
+					})
+			]
+			.OverrideResetToDefault(FResetToDefaultOverride::Create(
+				FIsResetToDefaultVisible::CreateStatic(&FLandscapeEditorDetailCustomization_MiscTools::IsResetGizmoButtonShown),
+				FResetToDefaultHandler::CreateStatic(&FLandscapeEditorDetailCustomization_MiscTools::ResetGizmoToOrigin)
+			));
+	}
 }
 END_SLATE_FUNCTION_BUILD_OPTIMIZATION
 
+bool FLandscapeEditorDetailCustomization_MiscTools::IsResetGizmoButtonShown(TSharedPtr<IPropertyHandle> InProperty)
+{
+	if (FEdModeLandscape * LandscapeEdMode = GetEditorMode())
+	{
+		TWeakObjectPtr<ULandscapeInfo> LandscapeInfo = LandscapeEdMode->CurrentToolTarget.LandscapeInfo;
+		TWeakObjectPtr<ALandscapeGizmoActiveActor> CurrentGizmoActor = LandscapeEdMode->UISettings->CurrentGizmoActor;
+
+		check(LandscapeInfo.Get());
+		check(CurrentGizmoActor.Get());
+
+		FIntRect LandscapeExtent;
+		LandscapeInfo->GetLandscapeExtent(LandscapeExtent);
+		const FTransform LandscapeToWorldTransform = LandscapeInfo->GetLandscapeProxy()->LandscapeActorToWorld();
+		const FVector LandscapeOrigin = LandscapeToWorldTransform.TransformPosition(FVector(LandscapeExtent.Min.X, LandscapeExtent.Min.Y, 0.0f));
+
+		return CurrentGizmoActor->GetActorLocation() != LandscapeOrigin;
+	}
+	return false;
+}
+
+void FLandscapeEditorDetailCustomization_MiscTools::ResetGizmoToOrigin(TSharedPtr<IPropertyHandle> InProperty)
+{
+	FEdModeLandscape* LandscapeEdMode = GetEditorMode();
+	if (!LandscapeEdMode)
+	{
+		return;
+	}
+
+	TWeakObjectPtr<ULandscapeInfo> LandscapeInfo = LandscapeEdMode->CurrentToolTarget.LandscapeInfo;
+	TWeakObjectPtr<ALandscapeGizmoActiveActor> CurrentGizmoActor = LandscapeEdMode->UISettings->CurrentGizmoActor;
+
+	check(LandscapeInfo.Get());
+	check(CurrentGizmoActor.Get());
+
+	FIntRect LandscapeExtent;
+	LandscapeInfo->GetLandscapeExtent(LandscapeExtent);
+	const FTransform LandscapeToWorldTransform = LandscapeInfo->GetLandscapeProxy()->LandscapeActorToWorld();
+	const FVector LandscapeOrigin = LandscapeToWorldTransform.TransformPosition(FVector(LandscapeExtent.Min.X, LandscapeExtent.Min.Y, 0.0f));
+
+	CurrentGizmoActor->Modify();
+	CurrentGizmoActor->EditorApplyTranslation(LandscapeOrigin - CurrentGizmoActor->GetActorLocation(), /*bAltDown =*/false, /*bShiftDown =*/false, /*bCtrlDown =*/false);
+}
+
 EVisibility FLandscapeEditorDetailCustomization_MiscTools::GetClearComponentSelectionVisibility()
 {
 	FEdModeLandscape* LandscapeEdMode = GetEditorMode();
diff --git a/Engine/Source/Editor/LandscapeEditor/Private/LandscapeEditorDetailCustomization_MiscTools.h b/Engine/Source/Editor/LandscapeEditor/Private/LandscapeEditorDetailCustomization_MiscTools.h
--- a/Engine/Source/Editor/LandscapeEditor/Private/LandscapeEditorDetailCustomization_MiscTools.h
+++ b/Engine/Source/Editor/LandscapeEditor/Private/LandscapeEditorDetailCustomization_MiscTools.h
@@ -10,12 +10,20 @@
 
 class IDetailLayoutBuilder;
 
+namespace LandscapeTestUtils
+{
+	struct LandscapeTestCommands;
+} //namespace LandscapeTestUtils
+
 /**
  * Slate widgets customizer for the smaller tools that need little customization
  */
 
 class FLandscapeEditorDetailCustomization_MiscTools : public FLandscapeEditorDetailCustomization_Base
 {
+	// Give tests access to private members
+	friend struct LandscapeTestUtils::LandscapeTestCommands;
+
 public:
 	/** Makes a new instance of this detail layout class for a specific detail view requesting it */
 	static TSharedRef<IDetailCustomization> MakeInstance();
@@ -64,4 +72,8 @@ protected:
 	// Common Error handling
 	EVisibility GetMiscLandscapeErrorVisibility() const;
 	FText GetMiscLandscapeErrorText() const;
+
+	// Current Gizmo customization (Import, Copy tools)
+	static bool IsResetGizmoButtonShown(TSharedPtr<IPropertyHandle> InProperty = nullptr);
+	static void ResetGizmoToOrigin(TSharedPtr<IPropertyHandle> InProperty = nullptr);
 };
diff --git a/Engine/Source/Editor/LandscapeEditor/Private/LandscapeEditorObject.cpp b/Engine/Source/Editor/LandscapeEditor/Private/LandscapeEditorObject.cpp
--- a/Engine/Source/Editor/LandscapeEditor/Private/LandscapeEditorObject.cpp
+++ b/Engine/Source/Editor/LandscapeEditor/Private/LandscapeEditorObject.cpp
@@ -340,17 +340,6 @@ void ULandscapeEditorObject::Load()
 	SetbUseSelectedRegion(bUseSelectedRegion);
 	SetbUseNegativeMask(bUseNegativeMask);
 
-	// Gizmo History (not saved!)
-	GizmoHistories.Empty();
-	for (TActorIterator<ALandscapeGizmoActor> It(ParentMode->GetWorld()); It; ++It)
-	{
-		ALandscapeGizmoActor* Gizmo = *It;
-		if (!Gizmo->IsEditable())
-		{
-			new(GizmoHistories) FGizmoHistory(Gizmo);
-		}
-	}
-
 	FString NewLandscapeMaterialName;
 
 	// If NewLandscape_Material is not null, we will try to use it
@@ -509,24 +498,24 @@ void ULandscapeEditorObject::SetGizmoSnapMode(ELandscapeGizmoSnapType InSnapMode
 {
 	SnapMode = InSnapMode;
 
-	if (ParentMode->CurrentGizmoActor.IsValid())
+	if (CurrentGizmoActor.IsValid())
 	{
-		ParentMode->CurrentGizmoActor->SnapType = SnapMode;
+		CurrentGizmoActor->SnapType = SnapMode;
 	}
 
 	if (SnapMode != ELandscapeGizmoSnapType::None)
 	{
 		// Ignore gizmo updates if landscape info is invalid (Ex. landscape is deleted after using Copy Tool)
-		if (ParentMode->CurrentGizmoActor.IsValid() && ParentMode->CurrentGizmoActor->TargetLandscapeInfo)
+		if (CurrentGizmoActor.IsValid() && CurrentGizmoActor->TargetLandscapeInfo)
 		{
-			const FVector WidgetLocation = ParentMode->CurrentGizmoActor->GetActorLocation();
-			const FRotator WidgetRotation = ParentMode->CurrentGizmoActor->GetActorRotation();
+			const FVector WidgetLocation = CurrentGizmoActor->GetActorLocation();
+			const FRotator WidgetRotation = CurrentGizmoActor->GetActorRotation();
 
-			const FVector SnappedWidgetLocation = ParentMode->CurrentGizmoActor->SnapToLandscapeGrid(WidgetLocation);
-			const FRotator SnappedWidgetRotation = ParentMode->CurrentGizmoActor->SnapToLandscapeGrid(WidgetRotation);
+			const FVector SnappedWidgetLocation = CurrentGizmoActor->SnapToLandscapeGrid(WidgetLocation);
+			const FRotator SnappedWidgetRotation = CurrentGizmoActor->SnapToLandscapeGrid(WidgetRotation);
 
-			ParentMode->CurrentGizmoActor->SetActorLocation(SnappedWidgetLocation, false);
-			ParentMode->CurrentGizmoActor->SetActorRotation(SnappedWidgetRotation);
+			CurrentGizmoActor->SetActorLocation(SnappedWidgetLocation, false);
+			CurrentGizmoActor->SetActorRotation(SnappedWidgetRotation);
 		}
 	}
 }
diff --git a/Engine/Source/Editor/LandscapeEditor/Private/Tests/LandscapeAutomationTests.cpp b/Engine/Source/Editor/LandscapeEditor/Private/Tests/LandscapeAutomationTests.cpp
--- a/Engine/Source/Editor/LandscapeEditor/Private/Tests/LandscapeAutomationTests.cpp
+++ b/Engine/Source/Editor/LandscapeEditor/Private/Tests/LandscapeAutomationTests.cpp
@@ -20,6 +20,7 @@
 
 #include "LandscapeEditorDetailCustomization_NewLandscape.h"
 #include "LandscapeEditorDetailCustomization_ImportExport.h"
+#include "LandscapeEditorDetailCustomization_MiscTools.h"
 
 #if WITH_DEV_AUTOMATION_TESTS
 
@@ -65,6 +66,7 @@ namespace LandscapeTestUtils
 		GLevelEditorModeTools().ActivateMode(FBuiltinEditorModes::EM_Landscape);
 		FEdModeLandscape* LandscapeEdMode = (FEdModeLandscape*)GLevelEditorModeTools().GetActiveMode(FBuiltinEditorModes::EM_Landscape);
 
+		TSharedPtr<FLandscapeEditorDetailCustomization_MiscTools> Customization_MiscTools = MakeShareable(new FLandscapeEditorDetailCustomization_MiscTools);
 		TSharedPtr<FLandscapeEditorDetailCustomization_ImportExport> Customization_ImportExportLandscape = MakeShareable(new FLandscapeEditorDetailCustomization_ImportExport);
 
 		LandscapeEdMode->UISettings->ImportLandscape_HeightmapFilename = HeightMapFilename;
@@ -72,8 +74,7 @@ namespace LandscapeTestUtils
 		LandscapeEdMode->UISettings->ImportExportMode = ELandscapeImportExportMode::All;
 		LandscapeEdMode->UISettings->ImportType = ELandscapeImportTransformType::Resample;
 
-		Customization_ImportExportLandscape->OnOriginResetButtonClicked();
-
+		Customization_MiscTools->ResetGizmoToOrigin();
 		Customization_ImportExportLandscape->OnImportExportButtonClicked();
 	}
 
diff --git a/Engine/Source/Editor/LandscapeEditor/Public/LandscapeEditorObject.h b/Engine/Source/Editor/LandscapeEditor/Public/LandscapeEditorObject.h
--- a/Engine/Source/Editor/LandscapeEditor/Public/LandscapeEditorObject.h
+++ b/Engine/Source/Editor/LandscapeEditor/Public/LandscapeEditorObject.h
@@ -480,6 +480,10 @@ class ULandscapeEditorObject : public UObject
 	UPROPERTY(Category="Tool Settings", EditAnywhere, NonTransactional, meta=(DisplayName="Gizmo copy/paste all layers", ShowForTools="CopyPaste"))
 	bool bApplyToAllTargets;
 
+	// The current gizmo actor. Only the gizmo location is revealed to the user interface
+	UPROPERTY(Category="Tool Settings", EditAnywhere, NonTransactional, Transient, meta=(ShowForTools="ImportExport,CopyPaste"))
+	TWeakObjectPtr<ALandscapeGizmoActiveActor> CurrentGizmoActor;
+
 	UPROPERTY(Category="Tool Settings", EditAnywhere, NonTransactional, meta=(DisplayName="Snap Gizmo to Landscape grid", ShowForTools="CopyPaste,ImportExport"))
 	ELandscapeGizmoSnapType SnapMode;
 
@@ -496,8 +500,6 @@ class ULandscapeEditorObject : public UObject
 	UPROPERTY(Category="Tool Settings", EditAnywhere, NonTransactional, AdvancedDisplay, meta=(DisplayName="Layers", ShowForTools="CopyPaste"))
 	TArray<FGizmoImportLayer> GizmoImportLayers;
 
-	TArray<FGizmoHistory> GizmoHistories;
-
 	// Mirror Tool
 
 	// Location of the mirror plane, defaults to the center of the landscape. Doesn't normally need to be changed!
