Landscape Editor - Moving SetSelectedTargetIndex to call before UpdateTargetList to avoid invalid index assert on level instance creation

diff --git a/Engine/Source/Editor/LandscapeEditor/Private/LandscapeEdMode.cpp b/Engine/Source/Editor/LandscapeEditor/Private/LandscapeEdMode.cpp
--- a/Engine/Source/Editor/LandscapeEditor/Private/LandscapeEdMode.cpp
+++ b/Engine/Source/Editor/LandscapeEditor/Private/LandscapeEdMode.cpp
@@ -339,12 +339,20 @@ bool FEdModeLandscape::IsLandscapeViewModeExclusiveToEditorMode(ELandscapeViewMo
 void FEdModeLandscape::SetLandscapeInfo(ULandscapeInfo* InLandscapeInfo)
 {
 	check(!InLandscapeInfo || InLandscapeInfo->SupportsLandscapeEditing());
+	const bool bIsInitialLandscapeInfo = CurrentToolTarget.LandscapeInfo == nullptr;
 	if (CurrentToolTarget.LandscapeInfo != InLandscapeInfo)
 	{
 		{
 			TGuardValue<bool> GuardFlag(bUpdatingLandscapeInfo, true);
 			CurrentToolTarget.LandscapeInfo = InLandscapeInfo;
 
+			// Only update the selected edit layer index for the initial landscape, 
+			// this avoids setting the last saved index to 0 when switching between existing landscapes
+			if (bIsInitialLandscapeInfo && CanHaveLandscapeLayersContent())
+			{
+				SetSelectedEditLayer(0);
+			}
+
 			UpdateTargetList();
 			UpdateToolModes();
 		}
@@ -2483,11 +2491,6 @@ int32 FEdModeLandscape::UpdateLandscapeList()
 			SetLandscapeInfo(LandscapeList[0].Info);
 			CurrentIndex = 0;
 
-			if (CanHaveLandscapeLayersContent())
-			{
-				SetSelectedEditLayer(0);
-			}
-
 			UpdateLayerUsageInformation();
 						
 			if (!CurrentToolName.IsNone())
