Fix assert when landscape smooth tool drifts outside landscape bounds, when using new selective layer render code.

diff --git a/Engine/Source/Editor/LandscapeEditor/Private/LandscapeEdModeBrushes.cpp b/Engine/Source/Editor/LandscapeEditor/Private/LandscapeEdModeBrushes.cpp
--- a/Engine/Source/Editor/LandscapeEditor/Private/LandscapeEdModeBrushes.cpp
+++ b/Engine/Source/Editor/LandscapeEditor/Private/LandscapeEdModeBrushes.cpp
@@ -326,59 +326,64 @@ public:
 			InteractorPositions = InInteractorPositions;
 		}
 
-		FIntRect InclusiveBounds;
+		FIntRect ExclusiveBounds;
 		for (const FLandscapeToolInteractorPosition& InteractorPosition : InteractorPositions)
 		{
-			FIntRect SpotInclusiveBounds;
-			SpotInclusiveBounds.Min.X = FMath::FloorToInt32(InteractorPosition.Position.X - TotalRadius);
-			SpotInclusiveBounds.Min.Y = FMath::FloorToInt32(InteractorPosition.Position.Y - TotalRadius);
-			SpotInclusiveBounds.Max.X = FMath::CeilToInt32(InteractorPosition.Position.X + TotalRadius);
-			SpotInclusiveBounds.Max.Y = FMath::CeilToInt32(InteractorPosition.Position.Y + TotalRadius);
+			FIntRect SpotExclusiveBounds;
+			SpotExclusiveBounds.Min.X = FMath::FloorToInt32(InteractorPosition.Position.X - TotalRadius);
+			SpotExclusiveBounds.Min.Y = FMath::FloorToInt32(InteractorPosition.Position.Y - TotalRadius);
+			SpotExclusiveBounds.Max.X = FMath::CeilToInt32(InteractorPosition.Position.X + TotalRadius) + 1;
+			SpotExclusiveBounds.Max.Y = FMath::CeilToInt32(InteractorPosition.Position.Y + TotalRadius) + 1;
 
-			if (InclusiveBounds.IsEmpty())
+			if (ExclusiveBounds.IsEmpty())
 			{
-				InclusiveBounds = SpotInclusiveBounds;
+				ExclusiveBounds = SpotExclusiveBounds;
 			}
 			else
 			{
-				InclusiveBounds.Min = InclusiveBounds.Min.ComponentMin(SpotInclusiveBounds.Min);
-				InclusiveBounds.Max = InclusiveBounds.Max.ComponentMax(SpotInclusiveBounds.Max);
+				ExclusiveBounds.Union(SpotExclusiveBounds);
 			}
 		}
 
-		// Clamp to landscape bounds
+		// Clamp to landscape bounds.  GetLandscapeExtent provides inclusive bounds.
 		int32 MinX, MaxX, MinY, MaxY;
 		if (!ensure(LandscapeInfo->GetLandscapeExtent(MinX, MinY, MaxX, MaxY)))
 		{
 			// Landscape has no components somehow
 			return FLandscapeBrushData();
 		}
-		InclusiveBounds.Clip(FIntRect(MinX, MinY, MaxX, MaxY));
+		// Clip using exclusive bounds so we can detect empty result.
+		ExclusiveBounds.Clip(FIntRect(MinX, MinY, MaxX + 1, MaxY + 1));
+		if (ExclusiveBounds.Area() <= 0)
+		{
+			// Brush is completely outside the landscape extent
+			return FLandscapeBrushData();
+		}
 
 		// check if region to be modified overlaps with an unloaded component
-		// TODO [chris.tchou] : The +1 shouldn't be necessary, but without it, edges can be modified by the brush. Possibly a bug in the brush apply code.
-		bool bHasUnloadedComponents = LandscapeInfo->HasUnloadedComponentsInRegion(InclusiveBounds.Min.X, InclusiveBounds.Min.Y, InclusiveBounds.Max.X+1, InclusiveBounds.Max.Y+1);
+		// TODO [chris.tchou] : The exclusive bounds here (InclusiveMax+1) shouldn't be necessary since the code looks like it expects inclusive values.  Without it, edges can be modified by the brush. Possibly a bug in the brush apply code.
+		bool bHasUnloadedComponents = LandscapeInfo->HasUnloadedComponentsInRegion(ExclusiveBounds.Min.X, ExclusiveBounds.Min.Y, ExclusiveBounds.Max.X, ExclusiveBounds.Max.Y);
 		if (bHasUnloadedComponents)
 		{
 			return FLandscapeBrushData();
 		}
 
-		// FLandscapeBrushData requires exclusive bounds : 
-		FLandscapeBrushData BrushData(FIntRect(InclusiveBounds.Min, InclusiveBounds.Max + 1));
+		FLandscapeBrushData BrushData(ExclusiveBounds);
 
 		for (const FLandscapeToolInteractorPosition& InteractorPosition : InteractorPositions)
 		{
-			FIntRect SpotInclusiveBounds;
-			SpotInclusiveBounds.Min.X = FMath::Max(FMath::FloorToInt32(InteractorPosition.Position.X - TotalRadius), InclusiveBounds.Min.X);
-			SpotInclusiveBounds.Min.Y = FMath::Max(FMath::FloorToInt32(InteractorPosition.Position.Y - TotalRadius), InclusiveBounds.Min.Y);
-			SpotInclusiveBounds.Max.X = FMath::Min(FMath::CeilToInt32(InteractorPosition.Position.X + TotalRadius), InclusiveBounds.Max.X);
-			SpotInclusiveBounds.Max.Y = FMath::Min(FMath::CeilToInt32(InteractorPosition.Position.Y + TotalRadius), InclusiveBounds.Max.Y);
-
-			for (int32 Y = SpotInclusiveBounds.Min.Y; Y <= SpotInclusiveBounds.Max.Y; Y++)
+			FIntRect SpotExclusiveBounds;
+			SpotExclusiveBounds.Min.X = FMath::FloorToInt32(InteractorPosition.Position.X - TotalRadius);
+			SpotExclusiveBounds.Min.Y = FMath::FloorToInt32(InteractorPosition.Position.Y - TotalRadius);
+			SpotExclusiveBounds.Max.X = FMath::CeilToInt32(InteractorPosition.Position.X + TotalRadius) + 1;
+			SpotExclusiveBounds.Max.Y = FMath::CeilToInt32(InteractorPosition.Position.Y + TotalRadius) + 1;
+			SpotExclusiveBounds.Clip(ExclusiveBounds);
+
+			for (int32 Y = SpotExclusiveBounds.Min.Y; Y < SpotExclusiveBounds.Max.Y; Y++)
 			{
 				float* Scanline = BrushData.GetDataPtr(FIntPoint(0, Y));
 
-				for (int32 X = SpotInclusiveBounds.Min.X; X <= SpotInclusiveBounds.Max.X; X++)
+				for (int32 X = SpotExclusiveBounds.Min.X; X < SpotExclusiveBounds.Max.X; X++)
 				{
 					float PrevAmount = Scanline[X];
 					if (PrevAmount < 1.0f)
@@ -1168,29 +1173,34 @@ public:
 		int32 SizeY = EdMode->UISettings->AlphaTextureSizeY;
 		check(EdMode->UISettings->HasValidAlphaTextureData() && (SizeX > 0) && (SizeY > 0)); // See CanPaint() above : if we can paint, these must all be valid
 
-		FIntRect InclusiveBounds;
-		InclusiveBounds.Min.X = FMath::FloorToInt(LastMousePosition.X - TotalRadius);
-		InclusiveBounds.Min.Y = FMath::FloorToInt(LastMousePosition.Y - TotalRadius);
-		InclusiveBounds.Max.X = FMath::CeilToInt(LastMousePosition.X + TotalRadius);
-		InclusiveBounds.Max.Y = FMath::CeilToInt(LastMousePosition.Y + TotalRadius);
+		FIntRect ExclusiveBounds;
+		ExclusiveBounds.Min.X = FMath::FloorToInt(LastMousePosition.X - TotalRadius);
+		ExclusiveBounds.Min.Y = FMath::FloorToInt(LastMousePosition.Y - TotalRadius);
+		ExclusiveBounds.Max.X = FMath::CeilToInt(LastMousePosition.X + TotalRadius) + 1;
+		ExclusiveBounds.Max.Y = FMath::CeilToInt(LastMousePosition.Y + TotalRadius) + 1;
 
-		// Clamp to landscape bounds
+		// Clamp to landscape bounds.  GetLandscapeExtent provides inclusive bounds.
 		int32 MinX, MaxX, MinY, MaxY;
 		if (!ensure(LandscapeInfo->GetLandscapeExtent(MinX, MinY, MaxX, MaxY)))
 		{
 			// Landscape has no components somehow
 			return FLandscapeBrushData();
 		}
-		InclusiveBounds.Clip(FIntRect(MinX, MinY, MaxX, MaxY));
+		// Clip using exclusive bounds so we can detect empty result.
+		ExclusiveBounds.Clip(FIntRect(MinX, MinY, MaxX + 1, MaxY + 1));
+		if (ExclusiveBounds.Area() <= 0)
+		{
+			// Brush is completely outside the landscape extent
+			return FLandscapeBrushData();
+		}
 
-		// FLandscapeBrushData requires exclusive bounds : 
-		FLandscapeBrushData BrushData(FIntRect(InclusiveBounds.Min, InclusiveBounds.Max + 1));
+		FLandscapeBrushData BrushData(ExclusiveBounds);
 
-		for (int32 Y = InclusiveBounds.Min.Y; Y <= InclusiveBounds.Max.Y; Y++)
+		for (int32 Y = ExclusiveBounds.Min.Y; Y < ExclusiveBounds.Max.Y; Y++)
 		{
 			float* Scanline = BrushData.GetDataPtr(FIntPoint(0, Y));
 
-			for (int32 X = InclusiveBounds.Min.X; X <= InclusiveBounds.Max.X; X++)
+			for (int32 X = ExclusiveBounds.Min.X; X < ExclusiveBounds.Max.X; X++)
 			{
 				float Angle;
 				FVector2D Scale;
@@ -1445,29 +1455,34 @@ public:
 			check(bIsValid && !FMath::IsNearlyZero(AlphaBrushScale) && !FMath::IsNearlyZero(Radius) && (SizeX > 0) && (SizeY > 0)); // See CanPaint() function above : if bCanPaint, AlphaBrushScale should be non-zero
 
 			const float BrushAngle = EdMode->UISettings->bAlphaBrushAutoRotate ? LastMouseAngle : FMath::DegreesToRadians(EdMode->UISettings->AlphaBrushRotation);
-			FIntRect InclusiveBounds;
-			InclusiveBounds.Min.X = FMath::FloorToInt(LastMousePosition.X - Radius);
-			InclusiveBounds.Min.Y = FMath::FloorToInt(LastMousePosition.Y - Radius);
-			InclusiveBounds.Max.X = FMath::CeilToInt(LastMousePosition.X + Radius);
-			InclusiveBounds.Max.Y = FMath::CeilToInt(LastMousePosition.Y + Radius);
+			FIntRect ExclusiveBounds;
+			ExclusiveBounds.Min.X = FMath::FloorToInt(LastMousePosition.X - Radius);
+			ExclusiveBounds.Min.Y = FMath::FloorToInt(LastMousePosition.Y - Radius);
+			ExclusiveBounds.Max.X = FMath::CeilToInt(LastMousePosition.X + Radius) + 1;
+			ExclusiveBounds.Max.Y = FMath::CeilToInt(LastMousePosition.Y + Radius) + 1;
 
-			// Clamp to landscape bounds
+			// Clamp to landscape bounds.  GetLandscapeExtent provides inclusive bounds.
 			int32 MinX, MaxX, MinY, MaxY;
 			if (!ensure(LandscapeInfo->GetLandscapeExtent(MinX, MinY, MaxX, MaxY)))
 			{
 				// Landscape has no components somehow
 				return FLandscapeBrushData();
 			}
-			InclusiveBounds.Clip(FIntRect(MinX, MinY, MaxX, MaxY));
+			// Clip using exclusive bounds so we can detect empty result.
+			ExclusiveBounds.Clip(FIntRect(MinX, MinY, MaxX + 1, MaxY + 1));
+			if (ExclusiveBounds.Area() <= 0)
+			{
+				// Brush is completely outside the landscape extent
+				return FLandscapeBrushData();
+			}
 
-			// FLandscapeBrushData requires exclusive bounds : 
-			FLandscapeBrushData BrushData(FIntRect(InclusiveBounds.Min, InclusiveBounds.Max + 1));
+			FLandscapeBrushData BrushData(FIntRect(ExclusiveBounds.Min, ExclusiveBounds.Max));
 
-			for (int32 Y = InclusiveBounds.Min.Y; Y <= InclusiveBounds.Max.Y; Y++)
+			for (int32 Y = ExclusiveBounds.Min.Y; Y < ExclusiveBounds.Max.Y; Y++)
 			{
 				float* Scanline = BrushData.GetDataPtr(FIntPoint(0, Y));
 
-				for (int32 X = InclusiveBounds.Min.X; X <= InclusiveBounds.Max.X; X++)
+				for (int32 X = ExclusiveBounds.Min.X; X < ExclusiveBounds.Max.X; X++)
 				{
 					// Find alphamap sample location
 					float ScaleSampleX = ((float)X - LastMousePosition.X) / AlphaBrushScale;
