Change landscape tool code to use new selective layer render code.  Removes use of old "intermediate render" code and avoids poluting the persistent texture data with intermediate results.

diff --git a/Engine/Source/Editor/LandscapeEditor/Private/LandscapeEdModePaintTools.cpp b/Engine/Source/Editor/LandscapeEditor/Private/LandscapeEdModePaintTools.cpp
--- a/Engine/Source/Editor/LandscapeEditor/Private/LandscapeEdModePaintTools.cpp
+++ b/Engine/Source/Editor/LandscapeEditor/Private/LandscapeEdModePaintTools.cpp
@@ -696,6 +696,7 @@ public:
 		CA_SUPPRESS(6326);
 		if (ToolTarget::TargetType == ELandscapeToolTargetType::Heightmap)
 		{
+			// Note:  this is an inclusive bounds that may extend outside the available landscape coords on any side.
 			X1 -= 1;
 			Y1 -= 1;
 			X2 += 1;
diff --git a/Engine/Source/Editor/LandscapeEditor/Private/LandscapeEdModeTools.h b/Engine/Source/Editor/LandscapeEditor/Private/LandscapeEdModeTools.h
--- a/Engine/Source/Editor/LandscapeEditor/Private/LandscapeEdModeTools.h
+++ b/Engine/Source/Editor/LandscapeEditor/Private/LandscapeEdModeTools.h
@@ -299,6 +299,7 @@ public:
 	}
 
 	// X2/Y2 Coordinates are "inclusive" max values
+	// ...but can extend outside valid landscape coords.  GetData/GetDataFast will generate and store values for out-of-range coords and then return coords clamped to the valid range.
 	// Note that this should maybe be called "ExtendDataCache" because the region here will be combined with the existing cached region, not loaded independently, giving a cached region that is the bounding box of previous and new
 	void CacheData(int32 X1, int32 Y1, int32 X2, int32 Y2, bool bCacheOriginalData = false)
 	{
@@ -404,6 +405,65 @@ public:
 		}
 	}
 
+	// Store data in the cache from some external source.
+	// Rect is half-open (exclusive), not inclusive like CacheData.
+	void CacheDataDirect(FIntRect NewDataRect, FIntRect RequestedRect, FIntRect LandscapeExtent, TConstArrayView<AccessorType> Data)
+	{
+		TRACE_CPUPROFILER_EVENT_SCOPE(TLandscapeEditCache_CacheDataDirect);
+		check(!NewDataRect.IsEmpty());
+		check(Data.Num() >= NewDataRect.Area());
+		check(LandscapeExtent.Contains(NewDataRect));
+		check(RequestedRect.Contains(NewDataRect));
+
+		// Allow caching values up to 1 pixel outside the real landscape extent.  
+		FIntRect AllowedCacheExtent = LandscapeExtent;
+		AllowedCacheExtent.InflateRect(1);
+
+		FIntRect OldCachedRect(CachedX1, CachedY1, CachedX2 + 1, CachedY2 + 1);
+
+		FIntRect NewCachedRect = RequestedRect;
+		NewCachedRect.Clip(AllowedCacheExtent);
+		check(NewCachedRect.Contains(NewDataRect));
+
+		// Clamp coord to the closest valid landscape coord (exclusive extent)
+		// This replaces the complicated bUseInterp code (GetData instead of GetDataFast) for coords outside valid landscape extent.  A better option
+		// might be to stop storing outside values at all and adjust coords on data lookups.  That can be explored as part of removing the TMap.
+		auto ClampToLandscape = [&LandscapeExtent](FInt32Point Point) -> FInt32Point
+			{
+				return { FMath::Clamp(Point.X, LandscapeExtent.Min.X, LandscapeExtent.Max.X - 1),
+						 FMath::Clamp(Point.Y, LandscapeExtent.Min.Y, LandscapeExtent.Max.Y - 1) };
+			};
+
+		int32 Stride = NewDataRect.Max.X - NewDataRect.Min.X;  //Stride of the data buffer
+		CachedData.Reserve(NewCachedRect.Area());
+		for (int32 Y = NewCachedRect.Min.Y; Y < NewCachedRect.Max.Y; ++Y)
+		{
+			for (int32 X = NewCachedRect.Min.X; X < NewCachedRect.Max.X; ++X)
+			{
+				FIntPoint Coord(X, Y);
+				if (!bIsValid || !OldCachedRect.Contains(Coord))
+				{
+					FIntPoint Clamped = ClampToLandscape(Coord);
+					FIntPoint BufferRelative = Clamped - NewDataRect.Min;
+					int32 Idx = BufferRelative.Y * Stride + BufferRelative.X;
+					CachedData.Add(Coord, Data[Idx]);
+				}
+			}
+		}
+
+		bIsValid = true;
+
+		CachedX1 = NewCachedRect.Min.X;
+		CachedY1 = NewCachedRect.Min.Y;
+		CachedX2 = NewCachedRect.Max.X - 1;
+		CachedY2 = NewCachedRect.Max.Y - 1;
+
+		ensure(CachedData.Contains({ CachedX1 , CachedY1 }));
+		ensure(CachedData.Contains({ CachedX1 , CachedY2 }));
+		ensure(CachedData.Contains({ CachedX2 , CachedY1 }));
+		ensure(CachedData.Contains({ CachedX2 , CachedY2 }));
+	}
+
 	AccessorType* GetValueRef(int32 LandscapeX, int32 LandscapeY)
 	{
 		return CachedData.Find(FIntPoint(LandscapeX, LandscapeY));
@@ -484,7 +544,7 @@ public:
 	}
 	
 	using FPrepareRegionForCachingFunction = TFunction<FIntRect(const FIntRect& NewCacheBounds)>;
-	bool GetDataAndCache(int32 X1, int32 Y1, int32 X2, int32 Y2, TArray<AccessorType>& OutData, FPrepareRegionForCachingFunction PrepareRegionForCaching)
+	bool GetDataAndCacheWithVisibilityOverride(int32 X1, int32 Y1, int32 X2, int32 Y2, TArray<AccessorType>& OutData, FPrepareRegionForCachingFunction PrepareRegionForCaching, TBitArray<> LayerVisibility, ALandscape* Landscape)
 	{
 		// Drop a visual log to indicate the area requested by this data access :
 		VisualizeLandscapeRegion(X1, Y1, X2, Y2, FColor::Blue, TEXT("CacheDataRequest"));
@@ -495,11 +555,76 @@ public:
 			// The bounds we calculate here need to be what would be the result of calling CacheData with this region, meaning that they should include the previous bounds. This will let us pass the correct region of interest to PrepareRegionForCaching
 			FIntRect NewCacheBounds(bIsValid ? FMath::Min(X1, CachedX1) : X1, bIsValid ? FMath::Min(Y1, CachedY1) : Y1, bIsValid ? FMath::Max(X2, CachedX2) : X2, bIsValid ? FMath::Max(Y2, CachedY2) : Y2);
 
-			// The caller might request a cache region that is actually larger than the data they want to sample for this particular read (e.g. to avoid re-caching when doing expensive samples like layer-collapsing ones) : 
-			NewCacheBounds = PrepareRegionForCaching(NewCacheBounds);
-			check((NewCacheBounds.Min.X <= X1) && (NewCacheBounds.Min.Y <= Y1) && (NewCacheBounds.Max.X >= X2) && (NewCacheBounds.Max.Y >= Y2));
+			// SelectiveRender only supports heightmaps at the moment.
+			if (Landscape->AllowSelectiveRenderEditLayers() && std::is_same<AccessorType, uint16>::value)
+			{
+				// For now, do the same bounds expansion as OnCacheUpdating.  Align with component boundaries.
+				// Note: FloorToInt32 is "strictly less", not "towards zero".
+				// Note, because of the +1/-1 expansion at a higher level in the callstack, this will pad out an entire component outside the landscape.
+				// CacheDataDirect is currently ignoring this excessive padding and limiting to only 1 row/column extra.
+				// TODO:  simplify when removing IntermediateRender.
+				FIntRect DesiredCacheBounds;
+				DesiredCacheBounds.Min.X = (FMath::FloorToInt((float)NewCacheBounds.Min.X / LandscapeInfo->ComponentSizeQuads)) * LandscapeInfo->ComponentSizeQuads;
+				DesiredCacheBounds.Min.Y = (FMath::FloorToInt((float)NewCacheBounds.Min.Y / LandscapeInfo->ComponentSizeQuads)) * LandscapeInfo->ComponentSizeQuads;
+				DesiredCacheBounds.Max.X = (FMath::CeilToInt((float)NewCacheBounds.Max.X / LandscapeInfo->ComponentSizeQuads)) * LandscapeInfo->ComponentSizeQuads;
+				DesiredCacheBounds.Max.Y = (FMath::CeilToInt((float)NewCacheBounds.Max.Y / LandscapeInfo->ComponentSizeQuads)) * LandscapeInfo->ComponentSizeQuads;
+
+				FIntRect LandscapeExtent;
+				bool bRet = LandscapeInfo->GetLandscapeExtent(LandscapeExtent);
+				ensure(bRet);
+				if (!bRet)
+				{
+					// Failure likely means no proxies loaded
+					return false;
+				}
+
+				// FLandscapeToolStrokeSmooth::Apply is extending bounds by 1 in all directions, outside landscape bounds, and asserts if the result bounds don't include that.
+				// landscape with 2017 verts/samples, DesiredCacheBounds is (-1,2017).  For now, don't change the return rect here.
+				FIntRect EvaluateRect = DesiredCacheBounds;
+
+				// Limit to the extent of the loaded landscape.  The calling code will ask for data outside the landscape and the old GetData code would generate it to fill in, but that doesn't
+				// appear to have been useful.
+				EvaluateRect.Clip(LandscapeExtent);  
+				EvaluateRect.Max += FInt32Point(1, 1);  // Convert inclusive range to exclusive (half-open).  [a,b)
+
+#if 0
+				UE_LOG(LogLandscapeTools, Warning, TEXT("PartialRender - New ---- (fr%u)  Base [%d, %d]-[%d, %d]  s[%d, %d],   Padded [%d, %d]-[%d, %d] s[%d, %d],   Evaluate [%d, %d]-[%d, %d] s[%d, %d]"), GFrameNumber,
+					NewCacheBounds.Min.X, NewCacheBounds.Min.Y, NewCacheBounds.Max.X, NewCacheBounds.Max.Y, NewCacheBounds.Size().X, NewCacheBounds.Size().Y,
+					DesiredCacheBounds.Min.X, DesiredCacheBounds.Min.Y, DesiredCacheBounds.Max.X, DesiredCacheBounds.Max.Y, DesiredCacheBounds.Size().X, DesiredCacheBounds.Size().Y,
+					EvaluateRect.Min.X, EvaluateRect.Min.Y, EvaluateRect.Max.X, EvaluateRect.Max.Y, EvaluateRect.Size().X, EvaluateRect.Size().Y);
+#endif
+
+				TArray<AccessorType> Results;
+				Results.AddUninitialized(EvaluateRect.Area());
 
-			CacheData(NewCacheBounds.Min.X, NewCacheBounds.Min.Y, NewCacheBounds.Max.X, NewCacheBounds.Max.Y);
+				bRet = AccessorClass::GetSelectiveLayerData(Landscape, EvaluateRect, LayerVisibility, Results);
+				ensure(bRet);  // what to do if this fails?
+
+				// Convert to exclusive rects
+				FIntRect DesiredCacheBoundsExcl = DesiredCacheBounds;
+				DesiredCacheBoundsExcl.Max += FInt32Point(1, 1);
+				FIntRect LandscapeExtentExcl = LandscapeExtent;
+				LandscapeExtentExcl.Max += FInt32Point(1, 1);
+
+				CacheDataDirect(EvaluateRect, DesiredCacheBoundsExcl, LandscapeExtentExcl, Results);
+			}
+			else
+			{
+				// Legacy intermediate render path.
+				// TODO:  remove with IntermediateRender
+				// 
+				// The caller might request a cache region that is actually larger than the data they want to sample for this particular read (e.g. to avoid re-caching when doing expensive samples like layer-collapsing ones) : 
+				NewCacheBounds = PrepareRegionForCaching(NewCacheBounds);
+				check((NewCacheBounds.Min.X <= X1) && (NewCacheBounds.Min.Y <= Y1) && (NewCacheBounds.Max.X >= X2) && (NewCacheBounds.Max.Y >= Y2));
+
+#if 0
+				//TODO:  remove before check in
+				UE_LOG(LogLandscapeTools, Warning, TEXT("PartialRender - Old ---- (fr%u)  Padded [%d, %d]-[%d, %d] s[%d, %d]"), GFrameNumber,
+					NewCacheBounds.Min.X, NewCacheBounds.Min.Y, NewCacheBounds.Max.X, NewCacheBounds.Max.Y, NewCacheBounds.Size().X, NewCacheBounds.Size().Y);
+#endif
+
+				CacheData(NewCacheBounds.Min.X, NewCacheBounds.Min.Y, NewCacheBounds.Max.X, NewCacheBounds.Max.Y);
+			}
 		}
 		ensureMsgf(!bIsValid || HasCachedData(X1, Y1, X2, Y2), TEXT("Data failed to cache (X1:%i, Y1:%i, X2:%i, Y2:%i), (CachedX1:%i, CachedY1:%i, CachedX2:%i, CachedY2:%i)"), X1, Y1, X2, Y2, CachedX1, CachedY1, CachedX2, CachedY2);
 		return GetCachedData(X1, Y1, X2, Y2, OutData);
@@ -920,7 +1045,7 @@ public:
 		check(bIsInitialized);
 		if (bCombinedLayerOperation)
 		{
-			TArray<bool> NewLayerVisibility;
+			TBitArray<> NewLayerVisibility;
 			int32 I = 0;
 			for (const ULandscapeEditLayerBase* EditLayer : Landscape->GetEditLayersConst())
 			{
@@ -928,6 +1053,7 @@ public:
 				++I;
 			}
 
+			// TODO:  remove with IntermediateRender
 			auto OnCacheUpdating = [&](const FIntRect& NewCacheBounds) -> FIntRect
 			{
 				// This function is triggered when the cache needs expanding. We'll ask for a larger area so that we don't need to re-cache (a GPU-synchronizing operation) every time
@@ -941,6 +1067,10 @@ public:
 
 				TSet<ULandscapeComponent*> AffectedComponents;
 				LandscapeInfo->GetComponentsInRegion(DesiredCacheBounds.Min.X, DesiredCacheBounds.Min.Y, DesiredCacheBounds.Max.X, DesiredCacheBounds.Max.Y, AffectedComponents);
+
+				UE_LOG(LogLandscapeTools, Warning, TEXT("IntRender (fr%u)   [%d, %d]-[%d, %d]  %d comps"), GFrameNumber,
+					DesiredCacheBounds.Min.X, DesiredCacheBounds.Min.Y, DesiredCacheBounds.Max.X, DesiredCacheBounds.Max.Y, AffectedComponents.Num());
+
 				SynchronousUpdateComponentVisibilityForHeight(AffectedComponents, NewLayerVisibility);
 
 				if (bVisibilityChanged)
@@ -955,14 +1085,17 @@ public:
 			FGuid PreviousLayerGUID = EditingLayerGuid.GetValue();
 			SetCacheEditingLayer(FGuid());
 
-			CacheUpToEditingLayer.GetDataAndCache(X1, Y1, X2, Y2, Data, OnCacheUpdating);
+			CacheUpToEditingLayer.GetDataAndCacheWithVisibilityOverride(X1, Y1, X2, Y2, Data, OnCacheUpdating, NewLayerVisibility, Landscape);
 			// Release Texture Mips that will be Locked by the next SynchronousUpdateComponentVisibilityForHeight
+			// remove Flush with IntermediateRender?  Don't need to release any mapped textures but if there were also pending writes then this has an effect.
 			CacheUpToEditingLayer.DataAccess.Flush();
 
 			// Now turn off visibility on the current layer in order to have the data of all bottom layers except the current one
 			NewLayerVisibility[EditingLayerIndex] = false;
-			CacheBottomLayers.GetDataAndCache(X1, Y1, X2, Y2, BottomLayersData, OnCacheUpdating);
+			CacheBottomLayers.GetDataAndCacheWithVisibilityOverride(X1, Y1, X2, Y2, BottomLayersData, OnCacheUpdating, NewLayerVisibility, Landscape);
 			// Do the same here for consistency
+			// remove with IntermediateRender?  CacheBottomLayers never writes back to the landscape, so this one should be safe to remove.  CacheBottomLayers could be
+			// a smaller, simpler class if its only used for GetSelectiveLayerData.
 			CacheBottomLayers.DataAccess.Flush();
 			
 			SetCacheEditingLayer(PreviousLayerGUID);
@@ -1001,6 +1134,7 @@ public:
  				Landscape->GetEditLayerConst(CacheAccessorLayerGuid) ? *(Landscape->GetEditLayerConst(CacheAccessorLayerGuid)->GetName().ToString()) : TEXT("<unknown>"), *CacheAccessorLayerGuid.ToString());
 
 			// Restore layers visibility
+			// remove with IntermediateRender.  nothing will change real layer visibility anymore.
 			SetLayersVisibility(BackupLayerVisibility);
 			// Only store data in cache
 			CacheUpToEditingLayer.SetCachedData(X1, Y1, X2, Y2, Data, PaintingRestriction, false);
@@ -1035,7 +1169,7 @@ private:
 		Landscape->ForceUpdateLayersContent(bIntermediateRender);
 	};
 
-	void SetLayersVisibility(const TArray<bool>& InLayerVisibility)
+	void SetLayersVisibility(const TBitArray<>& InLayerVisibility)
 	{
 		check(InLayerVisibility.Num() == Landscape->GetEditLayersConst().Num());
 		for (int i = 0; i < InLayerVisibility.Num(); ++i)
@@ -1053,7 +1187,8 @@ private:
 		}
 	};
 
-	void SynchronousUpdateComponentVisibilityForHeight(const TSet<ULandscapeComponent*>& InComponents, const TArray<bool>& InLayerVisibility)
+	// remove with IntermediateRender
+	void SynchronousUpdateComponentVisibilityForHeight(const TSet<ULandscapeComponent*>& InComponents, const TBitArray<>& InLayerVisibility)
 	{
 		SetLayersVisibility(InLayerVisibility);
         // No need to update collision here as we are only doing a intermediate render to gather heightdata
@@ -1066,13 +1201,13 @@ private:
 	ALandscape* Landscape;
 	TOptional<FGuid> EditingLayerGuid;
 	uint8 EditingLayerIndex;
-	TArray<bool> BackupLayerVisibility;
+	TBitArray<> BackupLayerVisibility;
 	TArray<typename ToolTarget::CacheClass::DataType> BottomLayersData;
 	bool bIsInitialized;
 	bool bCombinedLayerOperation;
-	bool bVisibilityChanged;
+	bool bVisibilityChanged;    // remove with IntermediateRender
 	bool bTargetIsHeightmap;
-	TSet<ULandscapeComponent*> VisibilityChangedComponents;
+	TSet<ULandscapeComponent*> VisibilityChangedComponents;  // remove with IntermediateRender.
 
 	typename ToolTarget::CacheClass& CacheUpToEditingLayer;
 	typename ToolTarget::CacheClass CacheBottomLayers;
diff --git a/Engine/Source/Runtime/Landscape/Private/LandscapeEditLayer.cpp b/Engine/Source/Runtime/Landscape/Private/LandscapeEditLayer.cpp
--- a/Engine/Source/Runtime/Landscape/Private/LandscapeEditLayer.cpp
+++ b/Engine/Source/Runtime/Landscape/Private/LandscapeEditLayer.cpp
@@ -155,6 +155,7 @@ void ULandscapeEditLayerBase::SetVisible(bool bInVisible, bool bInModify)
 	SetVisible(bInVisible, bInModify, /*bInBroadcastDataChange = */true);
 }
 
+// deprecate/remove with IntermediateRender
 void ULandscapeEditLayerBase::SetVisible(bool bInVisible, bool bInModify, bool bInBroadcastDataChange)
 {
 	const bool bHasValueChanged = (bInVisible != bVisible);
diff --git a/Engine/Source/Runtime/Landscape/Public/LandscapeEdit.h b/Engine/Source/Runtime/Landscape/Public/LandscapeEdit.h
--- a/Engine/Source/Runtime/Landscape/Public/LandscapeEdit.h
+++ b/Engine/Source/Runtime/Landscape/Public/LandscapeEdit.h
@@ -428,6 +428,19 @@ struct FHeightmapAccessor
 		LandscapeEdit->GetHeightDataFast(X1, Y1, X2, Y2, Data, 0);
 	}
 
+	// Render the heightmap with overrides on edit layer visibility.  Does not modify persistent landscape data like FLandscapeEditDataInterface does.
+	// This function exists mainly as a template customization point to handle uint16 for heights and uint8 for weights.
+	// EvalRect is half-open.
+	static bool GetSelectiveLayerData(ALandscape* Landscape, FIntRect EvalRect, TBitArray<> ActiveEditLayers, TArrayView<uint16> OutData)
+	{
+		FLandscapeEditLayerRenderHeightParams RenderParams;
+		RenderParams.Bounds = EvalRect;
+		RenderParams.ActiveEditLayers = ActiveEditLayers;
+		RenderParams.CpuResult = OutData;
+
+		return Landscape->SelectiveRenderEditLayersHeightmaps(RenderParams);
+	}
+
 	void SetData(int32 X1, int32 Y1, int32 X2, int32 Y2, const uint16* Data, ELandscapeLayerPaintingRestriction PaintingRestriction = ELandscapeLayerPaintingRestriction::None)
 	{
 		TSet<ULandscapeComponent*> Components;
@@ -602,6 +615,12 @@ struct FAlphamapAccessor
 		LandscapeEdit.GetWeightDataFast(LayerInfo, X1, Y1, X2, Y2, Data, 0);
 	}
 
+	static bool GetSelectiveLayerData(ALandscape* Landscape, FIntRect EvalRect, TBitArray<> ActiveEditLayers, TArrayView<uint8> OutData)
+	{
+		unimplemented();
+		return false;
+	}
+
 	void SetData(int32 X1, int32 Y1, int32 X2, int32 Y2, const uint8* Data, ELandscapeLayerPaintingRestriction PaintingRestriction)
 	{
 		TSet<ULandscapeComponent*> Components;
