Fix landscape flatten tool behavior at edges of geometry.

#13313

diff --git a/Engine/Source/Editor/LandscapeEditor/Private/LandscapeEdModePaintTools.cpp b/Engine/Source/Editor/LandscapeEditor/Private/LandscapeEdModePaintTools.cpp
--- a/Engine/Source/Editor/LandscapeEditor/Private/LandscapeEdModePaintTools.cpp
+++ b/Engine/Source/Editor/LandscapeEditor/Private/LandscapeEdModePaintTools.cpp
@@ -863,10 +863,15 @@ public:
 				RuntimeDataAccessor.SetEditLayer(FGuid());
 
 				ValueType P00, P10, P01, P11;
-				RuntimeDataAccessor.GetDataFast(FlattenHeightX, FlattenHeightY, FlattenHeightX, FlattenHeightY, &P00);
-				RuntimeDataAccessor.GetDataFast(FlattenHeightX + 1, FlattenHeightY, FlattenHeightX + 1, FlattenHeightY, &P10);
-				RuntimeDataAccessor.GetDataFast(FlattenHeightX, FlattenHeightY + 1, FlattenHeightX, FlattenHeightY + 1, &P01);
-				RuntimeDataAccessor.GetDataFast(FlattenHeightX + 1, FlattenHeightY + 1, FlattenHeightX + 1, FlattenHeightY + 1, &P11);
+				// GetData will handle interpolation/extrapolation as needed if this samples outside valid landscape coords.
+				FIntPoint SampleCoords(FlattenHeightX, FlattenHeightY);
+				RuntimeDataAccessor.GetData(SampleCoords.X, SampleCoords.Y, SampleCoords.X, SampleCoords.Y, &P00);
+				SampleCoords = { FlattenHeightX + 1, FlattenHeightY };
+				RuntimeDataAccessor.GetData(SampleCoords.X, SampleCoords.Y, SampleCoords.X, SampleCoords.Y, &P10);
+				SampleCoords = { FlattenHeightX, FlattenHeightY + 1 };
+				RuntimeDataAccessor.GetData(SampleCoords.X, SampleCoords.Y, SampleCoords.X, SampleCoords.Y, &P01);
+				SampleCoords = { FlattenHeightX + 1, FlattenHeightY + 1 };
+				RuntimeDataAccessor.GetData(SampleCoords.X, SampleCoords.Y, SampleCoords.X, SampleCoords.Y, &P11);
 				// Release Texture Mips that will be Locked by the next SynchronousUpdateComponentVisibilityForHeight (inside the LayerDataCache.Read call)
 				RuntimeDataAccessor.Flush();
 				V00 = P00;
diff --git a/Engine/Source/Runtime/Landscape/Public/LandscapeEdit.h b/Engine/Source/Runtime/Landscape/Public/LandscapeEdit.h
--- a/Engine/Source/Runtime/Landscape/Public/LandscapeEdit.h
+++ b/Engine/Source/Runtime/Landscape/Public/LandscapeEdit.h
@@ -418,6 +418,11 @@ struct FHeightmapAccessor
 		LandscapeEdit->GetHeightData(X1, Y1, X2, Y2, Data);
 	}
 
+	void GetData(int32& X1, int32& Y1, int32& X2, int32& Y2, uint16* Data)
+	{
+		LandscapeEdit->GetHeightData(X1, Y1, X2, Y2, Data, 0);
+	}
+
 	void GetDataFast(int32 X1, int32 Y1, int32 X2, int32 Y2, TMap<FIntPoint, uint16>& Data)
 	{
 		LandscapeEdit->GetHeightDataFast(X1, Y1, X2, Y2, Data);
@@ -605,6 +610,11 @@ struct FAlphamapAccessor
 		LandscapeEdit.GetWeightData(LayerInfo, X1, Y1, X2, Y2, Data);
 	}
 
+	void GetData(int32& X1, int32& Y1, int32& X2, int32& Y2, uint8* Data)
+	{
+		LandscapeEdit.GetWeightData(LayerInfo, X1, Y1, X2, Y2, Data, 0);
+	}
+
 	void GetDataFast(int32 X1, int32 Y1, int32 X2, int32 Y2, TMap<FIntPoint, uint8>& Data)
 	{
 		LandscapeEdit.GetWeightDataFast(LayerInfo, X1, Y1, X2, Y2, Data);
