Fix landscape painting when toggling "invert" during brush stroke.

#13594

diff --git a/Engine/Source/Editor/LandscapeEditor/Private/LandscapeEdModePaintTools.cpp b/Engine/Source/Editor/LandscapeEditor/Private/LandscapeEdModePaintTools.cpp
--- a/Engine/Source/Editor/LandscapeEditor/Private/LandscapeEdModePaintTools.cpp
+++ b/Engine/Source/Editor/LandscapeEditor/Private/LandscapeEdModePaintTools.cpp
@@ -195,11 +195,11 @@ public:
 			this->Cache.GetOriginalData(X1, Y1, X2, Y2, OriginalData);
 			SourceDataArrayPtr = &OriginalData;
 
-			for (int32 Y = Y1; Y < Y2; Y++)
+			for (int32 Y = Y1; Y <= Y2; Y++)
 			{
 				ValueType* DataScanline = Data.GetData() + (Y - Y1) * (X2 - X1 + 1) + (0 - X1);
 				ValueType* OriginalDataScanline = OriginalData.GetData() + (Y - Y1) * (X2 - X1 + 1) + (0 - X1);
-				for (int32 X = X1; X < X2; X++)
+				for (int32 X = X1; X <= X2; X++)
 				{
 					ValueType& SourceValue = OriginalDataScanline[X];
 					ValueType& CurrentValue = DataScanline[X];
