Fixed ctrl+alt+right click drag increasing/decreasing the landscape brush continuously even when the mouse was still.

Details : the mouse capture was released almost immediately when using one of the drag shortcuts, which led the cursor to stop being queried (FSlateUser::QueryCursor()), which prevented the cached mouse position in FEditorViewportClient::GetCursor to be updated, leading to a non-zero mouse delta for as long as the mouse pointer stayed still

diff --git a/Engine/Source/Editor/LandscapeEditor/Private/LandscapeEdMode.cpp b/Engine/Source/Editor/LandscapeEditor/Private/LandscapeEdMode.cpp
--- a/Engine/Source/Editor/LandscapeEditor/Private/LandscapeEdMode.cpp
+++ b/Engine/Source/Editor/LandscapeEditor/Private/LandscapeEdMode.cpp
@@ -939,6 +939,13 @@ bool FEdModeLandscape::MouseMove(FEditorViewportClient* InViewportClient, FViewp
 		int32 MouseXDelta = MouseX - InViewportClient->GetCachedMouseX();
 		int32 MouseYDelta = MouseY - InViewportClient->GetCachedMouseY();
 
+		const FLandscapeEditorCommands& LandscapeActions = FLandscapeEditorCommands::Get();
+		const bool bIsDragChordBrushSizeAndFalloffPressed = InViewportClient->IsCommandChordPressed(LandscapeActions.DragBrushSizeAndFalloff);
+		const bool bIsDragChordBrushSizeOnlyPressed = InViewportClient->IsCommandChordPressed(LandscapeActions.DragBrushSize);
+		const bool bIsDragChordBrushFalloffOnlyPressed = InViewportClient->IsCommandChordPressed(LandscapeActions.DragBrushFalloff);
+		const bool bIsDragChordStrengthPressed = InViewportClient->IsCommandChordPressed(LandscapeActions.DragBrushStrength);
+		const bool bIsDragChordPressed = bIsDragChordBrushSizeAndFalloffPressed || bIsDragChordBrushSizeOnlyPressed || bIsDragChordBrushFalloffOnlyPressed || bIsDragChordStrengthPressed;
+
 		if (FMath::Abs(MouseXDelta) > 0 || FMath::Abs(MouseYDelta) > 0)
 		{
 			const bool bHorizontallyDominant = FMath::Abs(MouseXDelta) >= FMath::Abs(MouseYDelta);
@@ -947,8 +954,7 @@ bool FEdModeLandscape::MouseMove(FEditorViewportClient* InViewportClient, FViewp
 				MouseYDelta < 0; // The way y position is stored here is inverted relative to expected mouse movement to change brush size
 
 			// Are we altering something about the brush?
-			const FLandscapeEditorCommands& LandscapeActions = FLandscapeEditorCommands::Get();
-			if (InViewportClient->IsCommandChordPressed(LandscapeActions.DragBrushSizeAndFalloff))
+			if (bIsDragChordBrushSizeAndFalloffPressed)
 			{
 				// If a chord controlling both Size and Falloff is specified, control Size when moving mostly left-right and Falloff when moving mostly up/down
 				if (bHorizontallyDominant)
@@ -962,19 +968,19 @@ bool FEdModeLandscape::MouseMove(FEditorViewportClient* InViewportClient, FViewp
 				return true;
 			}
 
-			if (InViewportClient->IsCommandChordPressed(LandscapeActions.DragBrushSize))
+			if (bIsDragChordBrushSizeOnlyPressed)
 			{
 				ChangeBrushSize(bIncrease);
 				return true;
 			}
 
-			if (InViewportClient->IsCommandChordPressed(LandscapeActions.DragBrushFalloff))
+			if (bIsDragChordBrushFalloffOnlyPressed)
 			{
 				ChangeBrushFalloff(bIncrease);
 				return true;
 			}
 
-			if (InViewportClient->IsCommandChordPressed(LandscapeActions.DragBrushStrength))
+			if (bIsDragChordStrengthPressed)
 			{
 				ChangeBrushStrength(bIncrease);
 				return true;
@@ -984,15 +990,19 @@ bool FEdModeLandscape::MouseMove(FEditorViewportClient* InViewportClient, FViewp
 		// Require Ctrl or not as per user preference
 		const ELandscapeFoliageEditorControlType LandscapeEditorControlType = GetDefault<ULevelEditorViewportSettings>()->LandscapeEditorControlType;
 
-		if (ToolActiveViewport != InViewport ||
-			!InViewport->KeyState(EKeys::LeftMouseButton) ||
-			(LandscapeEditorControlType == ELandscapeFoliageEditorControlType::RequireCtrl && !IsCtrlDown(InViewport)))
+		if ((ToolActiveViewport != InViewport)
+			|| !InViewport->KeyState(EKeys::LeftMouseButton)
+			|| (LandscapeEditorControlType == ELandscapeFoliageEditorControlType::RequireCtrl && !IsCtrlDown(InViewport)))
 		{
 			if (CurrentTool->IsToolActive())
 			{
 				CurrentTool->EndTool(InViewportClient);
 			}
-			InViewport->CaptureMouse(false);
+			// Don't let go of the mouse capture while using one the drag shortcuts in order to maintain the cursor updated every frame : 
+			if (!bIsDragChordPressed)
+			{
+				InViewport->CaptureMouse(false);
+			}
 			ToolActiveViewport = nullptr;
 		}
 	}
