Sort landscapes alphabetically when multiple landscape actors are editable in landscape mode

diff --git a/Engine/Source/Editor/LandscapeEditor/Private/LandscapeEditorDetails.cpp b/Engine/Source/Editor/LandscapeEditor/Private/LandscapeEditorDetails.cpp
--- a/Engine/Source/Editor/LandscapeEditor/Private/LandscapeEditorDetails.cpp
+++ b/Engine/Source/Editor/LandscapeEditor/Private/LandscapeEditorDetails.cpp
@@ -326,16 +326,26 @@ TSharedRef<SWidget> FLandscapeEditorDetails::GetTargetLandscapeMenu()
 	{
 		FMenuBuilder MenuBuilder(true, NULL);
 
+		TArray<ULandscapeInfo*> LandscapeInfoList;
+
 		const TArray<FLandscapeListInfo>& LandscapeList = LandscapeEdMode->GetLandscapeList();
 		for (auto It = LandscapeList.CreateConstIterator(); It; It++)
 		{
 			if (ALandscape* LandscapeActor = It->Info->LandscapeActor.Get())
 			{
-				FUIAction Action = FUIAction(FExecuteAction::CreateStatic(&FLandscapeEditorDetails::OnChangeTargetLandscape, MakeWeakObjectPtr(It->Info)));
-				MenuBuilder.AddMenuEntry(GetLandscapeDisplayName(LandscapeActor), FText(), FSlateIcon(), Action);
+				LandscapeInfoList.Add(It->Info);
 			}
 		}
 
+		// Sort landscapes alphabetically : 
+		LandscapeInfoList.Sort([](const ULandscapeInfo& InLHS, const ULandscapeInfo& InRHS) { return GetLandscapeDisplayName(InLHS.LandscapeActor.Get()).ToString().Compare(GetLandscapeDisplayName(InRHS.LandscapeActor.Get()).ToString()) < 0;});
+
+		for (ULandscapeInfo* LandscapeInfo : LandscapeInfoList)
+		{
+			FUIAction Action = FUIAction(FExecuteAction::CreateStatic(&FLandscapeEditorDetails::OnChangeTargetLandscape, MakeWeakObjectPtr(LandscapeInfo)));
+			MenuBuilder.AddMenuEntry(GetLandscapeDisplayName(LandscapeInfo->LandscapeActor.Get()), FText(), FSlateIcon(), Action);
+		}
+
 		return MenuBuilder.MakeWidget();
 	}
 
