Trivial landscape spline readbility improvement with EnumOnlyContainsFlags

diff --git a/Engine/Source/Editor/LandscapeEditor/Private/LandscapeSplineSelection.cpp b/Engine/Source/Editor/LandscapeEditor/Private/LandscapeSplineSelection.cpp
--- a/Engine/Source/Editor/LandscapeEditor/Private/LandscapeSplineSelection.cpp
+++ b/Engine/Source/Editor/LandscapeEditor/Private/LandscapeSplineSelection.cpp
@@ -25,8 +25,7 @@ ULandscapeSplineSelection::ULandscapeSplineSelection(const FObjectInitializer& O
 
 void ULandscapeSplineSelection::SelectSegment(ULandscapeSplineSegment* Segment, ESplineNavigationFlags Flags)
 {
-	// UpdateProperties, AddToSelection, SegmentModeFlags are only valid flags
-	check(Flags == ESplineNavigationFlags::None || !EnumHasAnyFlags(Flags, ~(ESplineNavigationFlags::UpdatePropertiesWindows | ESplineNavigationFlags::AddToSelection | ESplineNavigationFlags::SegmentSelectModeEnabled)));
+	check(EnumOnlyContainsFlags(Flags, ESplineNavigationFlags::UpdatePropertiesWindows | ESplineNavigationFlags::AddToSelection | ESplineNavigationFlags::SegmentSelectModeEnabled));
 
 	// Do nothing if the given point is already selected
 	if (IsSegmentSelected(Segment))
@@ -64,8 +63,7 @@ void ULandscapeSplineSelection::SelectSegment(ULandscapeSplineSegment* Segment,
 
 void ULandscapeSplineSelection::SelectControlPoint(ULandscapeSplineControlPoint* ControlPoint, ESplineNavigationFlags Flags)
 {
-	// UpdateProperties, AddToSelection, ControlPointSelect are only valid flags
-	check(Flags == ESplineNavigationFlags::None || !EnumHasAnyFlags(Flags, ~(ESplineNavigationFlags::UpdatePropertiesWindows | ESplineNavigationFlags::AddToSelection | ESplineNavigationFlags::ControlPointSelectModeEnabled)));
+	check(EnumOnlyContainsFlags(Flags, ESplineNavigationFlags::UpdatePropertiesWindows | ESplineNavigationFlags::AddToSelection | ESplineNavigationFlags::ControlPointSelectModeEnabled));
 
 	// Do nothing if the given point is already selected
 	if (IsControlPointSelected(ControlPoint))
@@ -136,8 +134,7 @@ void ULandscapeSplineSelection::ClearSelection()
 
 void ULandscapeSplineSelection::DeselectControlPoint(ULandscapeSplineControlPoint* ControlPoint, ESplineNavigationFlags Flags)
 {
-	// UpdateProperties is only valid flag
-	check(Flags == ESplineNavigationFlags::None || !EnumHasAnyFlags(Flags, ~(ESplineNavigationFlags::UpdatePropertiesWindows)));
+	check(EnumOnlyContainsFlags(Flags, ESplineNavigationFlags::UpdatePropertiesWindows));
 	check(ControlPoint->IsSplineSelected());
 
 	ControlPoint->Modify(false);
@@ -154,8 +151,7 @@ void ULandscapeSplineSelection::DeselectControlPoint(ULandscapeSplineControlPoin
 
 void ULandscapeSplineSelection::DeSelectSegment(ULandscapeSplineSegment* Segment, ESplineNavigationFlags Flags)
 {
-	// UpdateProperties is only valid flag
-	check(Flags == ESplineNavigationFlags::None || !EnumHasAnyFlags(Flags, ~(ESplineNavigationFlags::UpdatePropertiesWindows)));
+	check(EnumOnlyContainsFlags(Flags, ESplineNavigationFlags::UpdatePropertiesWindows));
 	check(Segment->IsSplineSelected());
 
 	Segment->Modify(false);
