Unshelved from pending changelist '45185355':

Landscape Editor - Remove non-edit layer deform landscape to spline features

diff --git a/Engine/Source/Editor/LandscapeEditor/Private/LandscapeEdMode.cpp b/Engine/Source/Editor/LandscapeEditor/Private/LandscapeEdMode.cpp
--- a/Engine/Source/Editor/LandscapeEditor/Private/LandscapeEdMode.cpp
+++ b/Engine/Source/Editor/LandscapeEditor/Private/LandscapeEdMode.cpp
@@ -3940,16 +3940,6 @@ ELandscapeEditingState FEdModeLandscape::GetEditingState() const
 	return ELandscapeEditingState::Enabled;
 }
 
-bool FEdModeLandscape::HasSplinesEditLayer() const
-{
-	if (ALandscape* Landscape = GetLandscape(); Landscape != nullptr)
-	{
-		const ULandscapeEditLayerBase* SplinesLayer = Landscape->FindEditLayerOfTypeConst(ULandscapeEditLayerSplines::StaticClass());
-		return SplinesLayer != nullptr;
-	}
-	return false;
-}
-
 int32 FEdModeLandscape::GetLayerCount() const
 {
 	ALandscape* Landscape = GetLandscape();
@@ -4206,14 +4196,6 @@ bool FEdModeLandscape::CanEditTargetLayer(FText* Reason /*=nullptr*/, const ULan
 	return true;
 }
 
-void FEdModeLandscape::UpdateLandscapeSplines(bool bUpdateOnlySelected /* = false*/)
-{
-	if (ALandscape* Landscape = GetLandscape())
-	{
-		Landscape->UpdateLandscapeSplines(GetCurrentLayerGuid(), bUpdateOnlySelected);
-	}
-}
-
 FGuid FEdModeLandscape::GetCurrentLayerGuid() const
 {
 	const ULandscapeEditLayerBase* EditLayer = GetCurrentEditLayerConst();
diff --git a/Engine/Source/Editor/LandscapeEditor/Private/LandscapeEdMode.h b/Engine/Source/Editor/LandscapeEditor/Private/LandscapeEdMode.h
--- a/Engine/Source/Editor/LandscapeEditor/Private/LandscapeEdMode.h
+++ b/Engine/Source/Editor/LandscapeEditor/Private/LandscapeEdMode.h
@@ -547,7 +547,6 @@ public:
 	bool IsGridBased() const;
 
 	// Edit Layers
-	bool HasSplinesEditLayer() const;
 	int32 GetLayerCount() const;
 	void SetSelectedEditLayer(int32 InLayerIndex);
 	int32 GetSelectedEditLayerIndex() const;
@@ -558,7 +557,6 @@ public:
 	bool IsLayerAlphaVisible(int32 InLayerIndex) const;
 	const ULandscapeEditLayerBase* GetCurrentEditLayerConst() const;
 	FGuid GetCurrentLayerGuid() const;
-	void UpdateLandscapeSplines(bool bUpdateOnlySelected = false);
 	void AutoUpdateDirtyLandscapeSplines();
 	bool CanEditLayer(FText* Reason = nullptr, const ULandscapeEditLayerBase* InLayer = nullptr);
 	bool CanEditSelectedTargetLayer(FText* Reason = nullptr, const ULandscapeEditLayerBase* InLayer = nullptr);
diff --git a/Engine/Source/Editor/LandscapeEditor/Private/LandscapeEditorDetailCustomization_MiscTools.cpp b/Engine/Source/Editor/LandscapeEditor/Private/LandscapeEditorDetailCustomization_MiscTools.cpp
--- a/Engine/Source/Editor/LandscapeEditor/Private/LandscapeEditorDetailCustomization_MiscTools.cpp
+++ b/Engine/Source/Editor/LandscapeEditor/Private/LandscapeEditorDetailCustomization_MiscTools.cpp
@@ -132,57 +132,6 @@ void FLandscapeEditorDetailCustomization_MiscTools::CustomizeDetails(IDetailLayo
 
 	if (IsToolActive("Splines"))
 	{
-		ToolsCategory.AddCustomRow(LOCTEXT("ApplySplinesLabel", "Apply Splines"))
-		[
-			SNew(SHorizontalBox)
-			+ SHorizontalBox::Slot()
-			.Padding(0, 6, 0, 0)
-			[
-				SNew(STextBlock)
-				.Font(DetailBuilder.GetDetailFont())
-				.ShadowOffset(FVector2D::UnitVector)
-				.Text(LOCTEXT("Spline.ApplySplines", "Deform Landscape to Splines:"))
-			]
-		];
-
-		// Once a Splines Edit Layer exists, disable the buttons since spline updates happen automatically
-		ToolsCategory.AddCustomRow(LOCTEXT("ApplySplinesLabel", "Apply Splines"))
-		[
-			SNew(SHorizontalBox)
-			+ SHorizontalBox::Slot()
-			[
-				SNew(SButton)
-				.ToolTipText_Lambda([LandscapeEdMode]()
-				{
-					return LandscapeEdMode && LandscapeEdMode->HasSplinesEditLayer() ? LOCTEXT("Spline.SplineEditLayerEnabled", "The landscape deformation is controlled by the Splines Edit Layer. Changes made to the spline are instantly reflected to the landscape.")
-						: LOCTEXT("Spline.ApplySplines.All.Tooltip", "Deforms and paints the landscape to fit all the landscape spline segments and points.");
-				})
-				.Text(LOCTEXT("Spline.ApplySplines.All", "All Splines"))
-				.HAlign(HAlign_Center)
-				.OnClicked_Static(&FLandscapeEditorDetailCustomization_MiscTools::OnApplyAllSplinesButtonClicked)
-				.IsEnabled_Lambda([LandscapeEdMode]()
-				{
-					return LandscapeEdMode && !LandscapeEdMode->HasSplinesEditLayer();
-				})
-			]
-			+ SHorizontalBox::Slot()
-			[
-				SNew(SButton)
-				.ToolTipText_Lambda([LandscapeEdMode]()
-				{
-					return LandscapeEdMode && LandscapeEdMode->HasSplinesEditLayer() ? LOCTEXT("Spline.SplineEditLayerEnabled", "The landscape deformation is controlled by the Splines Edit Layer. Changes made to the spline are instantly reflected to the landscape.")
-						: LOCTEXT("Spline.ApplySplines.Segments.Tooltip", "Deforms and paints the landscape to fit only the selected landscape spline segments.");
-				})
-				.Text(LOCTEXT("Spline.ApplySplines.Selected", "Selected Segments"))
-				.HAlign(HAlign_Center)
-				.OnClicked_Static(&FLandscapeEditorDetailCustomization_MiscTools::OnApplySelectedSplinesButtonClicked)
-				.IsEnabled_Lambda([LandscapeEdMode]()
-				{
-					return LandscapeEdMode && !LandscapeEdMode->HasSplinesEditLayer();
-				})
-			]
-		];
-
 		ToolsCategory.AddCustomRow(LOCTEXT("SelectAllLabel", "Select all"))
 		[
 			SNew(SHorizontalBox)
@@ -525,29 +474,6 @@ FReply FLandscapeEditorDetailCustomization_MiscTools::OnClearRegionSelectionButt
 	return FReply::Handled();
 }
 
-FReply FLandscapeEditorDetailCustomization_MiscTools::OnApplyAllSplinesButtonClicked()
-{
-	FEdModeLandscape* LandscapeEdMode = GetEditorMode();
-	if (LandscapeEdMode)
-	{
-		FScopedTransaction Transaction(LOCTEXT("LandscapeSpline_ApplyAllSplines", "Apply All Splines to Landscape"));
-		LandscapeEdMode->UpdateLandscapeSplines(false);
-	}
-	return FReply::Handled();
-}
-
-FReply FLandscapeEditorDetailCustomization_MiscTools::OnApplySelectedSplinesButtonClicked()
-{
-	FEdModeLandscape* LandscapeEdMode = GetEditorMode();
-	if (LandscapeEdMode)
-	{
-		FScopedTransaction Transaction(LOCTEXT("LandscapeSpline_ApplySelectedSplines", "Apply Selected Splines to Landscape"));
-		LandscapeEdMode->UpdateLandscapeSplines(true);
-	}
-
-	return FReply::Handled();
-}
-
 FReply FLandscapeEditorDetailCustomization_MiscTools::OnSelectAllControlPointsButtonClicked()
 {
 	FEdModeLandscape* LandscapeEdMode = GetEditorMode();
diff --git a/Engine/Source/Editor/LandscapeEditor/Private/LandscapeEditorDetailCustomization_MiscTools.h b/Engine/Source/Editor/LandscapeEditor/Private/LandscapeEditorDetailCustomization_MiscTools.h
--- a/Engine/Source/Editor/LandscapeEditor/Private/LandscapeEditorDetailCustomization_MiscTools.h
+++ b/Engine/Source/Editor/LandscapeEditor/Private/LandscapeEditorDetailCustomization_MiscTools.h
@@ -41,8 +41,6 @@ protected:
 	static FReply OnClearRegionSelectionButtonClicked();
 
 	// Splines Tool
-	static FReply OnApplyAllSplinesButtonClicked();
-	static FReply OnApplySelectedSplinesButtonClicked();
 	static FReply OnSelectAllControlPointsButtonClicked();
 	static FReply OnSelectAllSegmentsButtonClicked();
 	void OnbUseAutoRotateControlPointChanged(ECheckBoxState NewState);
diff --git a/Engine/Source/Runtime/Landscape/Classes/Landscape.h b/Engine/Source/Runtime/Landscape/Classes/Landscape.h
--- a/Engine/Source/Runtime/Landscape/Classes/Landscape.h
+++ b/Engine/Source/Runtime/Landscape/Classes/Landscape.h
@@ -507,7 +507,9 @@ PRAGMA_ENABLE_DEPRECATION_WARNINGS
 	LANDSCAPE_API bool IsMaxLayersReached() const;
 	LANDSCAPE_API void ShowOnlySelectedLayer(int32 InLayerIndex);
 	LANDSCAPE_API void ShowAllLayers();
+	UE_DEPRECATED(5.7, "Use UpdateLandscapeSplines without bInUpdateOnlySelected. Edit layer landscapes update all splines regardless of selection state")
 	LANDSCAPE_API void UpdateLandscapeSplines(const FGuid& InLayerGuid = FGuid(), bool bInUpdateOnlySelected = false, bool bInForceUpdateAllCompoments = false);
+	LANDSCAPE_API void UpdateAllLandscapeSplines(const FGuid& InLayerGuid = FGuid(), bool bInForceUpdateAllCompoments = false);
 	LANDSCAPE_API void SetSelectedEditLayerIndex(const int32 InEditLayerIndex);
 	LANDSCAPE_API const int32 GetSelectedEditLayerIndex() const;
 
diff --git a/Engine/Source/Runtime/Landscape/Classes/LandscapeInfo.h b/Engine/Source/Runtime/Landscape/Classes/LandscapeInfo.h
--- a/Engine/Source/Runtime/Landscape/Classes/LandscapeInfo.h
+++ b/Engine/Source/Runtime/Landscape/Classes/LandscapeInfo.h
@@ -206,8 +206,9 @@ public:
 	LANDSCAPE_API void ExportHeightmap(const FString& Filename, const FIntRect& ExportRegion);
 	LANDSCAPE_API void ExportLayer(ULandscapeLayerInfoObject* LayerInfo, const FString& Filename);
 	LANDSCAPE_API void ExportLayer(ULandscapeLayerInfoObject* LayerInfo, const FString& Filename, const FIntRect& ExportRegion);
-	// TODO [jared.ritchie] - deprecate for 5.7 after non-edit layer landscape deprecation
+	UE_DEPRECATED(5.7, "Use ApplySplines without bOnlySelected. Edit layer landscapes update all splines regardless of selection state")
 	LANDSCAPE_API bool ApplySplines(bool bOnlySelected, TSet<TObjectPtr<ULandscapeComponent>>* OutModifiedComponents = nullptr, bool bMarkPackageDirty = true);
+	LANDSCAPE_API bool ApplySplines(TSet<TObjectPtr<ULandscapeComponent>>* OutModifiedComponents = nullptr, bool bMarkPackageDirty = true);
 
 	LANDSCAPE_API bool GetSelectedExtent(int32& MinX, int32& MinY, int32& MaxX, int32& MaxY) const;
 	FVector GetLandscapeCenterPos(float& LengthZ, int32 MinX = MAX_int32, int32 MinY = MAX_int32, int32 MaxX = MIN_int32, int32 MaxY = MIN_int32);
@@ -451,7 +452,7 @@ public:
 private:
 	inline static bool bForceNonSpatiallyLoadedByDefault = false;
 
-	bool ApplySplinesInternal(bool bOnlySelected, TScriptInterface<ILandscapeSplineInterface> SplineOwner, TSet<TObjectPtr<ULandscapeComponent>>* OutModifiedComponents, bool bMarkPackageDirty, int32 LandscapeMinX, int32 LandscapeMinY, int32 LandscapeMaxX, int32 LandscapeMaxY, TFunctionRef<TSharedPtr<FModulateAlpha>(ULandscapeLayerInfoObject*)> GetOrCreateModulate);
+	bool ApplySplinesInternal(TScriptInterface<ILandscapeSplineInterface> SplineOwner, TSet<TObjectPtr<ULandscapeComponent>>* OutModifiedComponents, bool bMarkPackageDirty, int32 LandscapeMinX, int32 LandscapeMinY, int32 LandscapeMaxX, int32 LandscapeMaxY, TFunctionRef<TSharedPtr<FModulateAlpha>(ULandscapeLayerInfoObject*)> GetOrCreateModulate);
 	void MoveSegment(ULandscapeSplineSegment* InSegment, TScriptInterface<ILandscapeSplineInterface> From, TScriptInterface<ILandscapeSplineInterface> To);
 	void MoveControlPoint(ULandscapeSplineControlPoint* InControlPoint, TScriptInterface<ILandscapeSplineInterface> From, TScriptInterface<ILandscapeSplineInterface> To);
 	bool UpdateLayerInfoMapInternal(ALandscapeProxy* Proxy);
diff --git a/Engine/Source/Runtime/Landscape/Private/LandscapeEditLayer.cpp b/Engine/Source/Runtime/Landscape/Private/LandscapeEditLayer.cpp
--- a/Engine/Source/Runtime/Landscape/Private/LandscapeEditLayer.cpp
+++ b/Engine/Source/Runtime/Landscape/Private/LandscapeEditLayer.cpp
@@ -444,7 +444,7 @@ TArray<ULandscapeEditLayerSplines::FEditLayerAction> ULandscapeEditLayerSplines:
 		LOCTEXT("LandscapeEditLayerSplines_UpdateSplines", "Update Splines"),
 		FEditLayerAction::FExecuteDelegate::CreateWeakLambda(this, [](const FEditLayerAction::FExecuteParams& InParams)
 		{
-			InParams.GetLandscape()->UpdateLandscapeSplines(FGuid(), /*bUpdateOnlySelection = */false, /*bForceUpdate =*/true);
+			InParams.GetLandscape()->UpdateAllLandscapeSplines(FGuid(), /*bForceUpdate =*/true);
 			return FEditLayerAction::FExecuteResult(/*bInSuccess = */true);
 		}),
 		FEditLayerAction::FCanExecuteDelegate::CreateWeakLambda(this, [](const FEditLayerAction::FExecuteParams& InParams, FText& OutReason)
diff --git a/Engine/Source/Runtime/Landscape/Private/LandscapeEditLayers.cpp b/Engine/Source/Runtime/Landscape/Private/LandscapeEditLayers.cpp
--- a/Engine/Source/Runtime/Landscape/Private/LandscapeEditLayers.cpp
+++ b/Engine/Source/Runtime/Landscape/Private/LandscapeEditLayers.cpp
@@ -12404,7 +12404,7 @@ void ALandscape::UpdateLayersContent(bool bInWaitForStreaming, bool bInSkipMonit
 				return;
 			}
 
-			UpdateLandscapeSplines();
+			UpdateAllLandscapeSplines();
 			bSplineLayerUpdateRequested = false;
 		}
 	}
@@ -13996,9 +13996,15 @@ uint32 ULandscapeComponent::ComputeLayerHash(bool InReturnEditingHash) const
 	return Hash;
 }
 
+// Deprecated
 void ALandscape::UpdateLandscapeSplines(const FGuid& InTargetLayer, bool bInUpdateOnlySelected, bool bInForceUpdateAllCompoments)
 {
-	TRACE_CPUPROFILER_EVENT_SCOPE(LandscapeLayers_UpdateLandscapeSplines);
+	UpdateAllLandscapeSplines(InTargetLayer, bInForceUpdateAllCompoments);
+}
+
+void ALandscape::UpdateAllLandscapeSplines(const FGuid& InTargetLayer, bool bInForceUpdateAllCompoments)
+{
+	TRACE_CPUPROFILER_EVENT_SCOPE(LandscapeLayers_UpdateAllLandscapeSplines);
 	check(!IsTemplate());
 	ULandscapeInfo* LandscapeInfo = GetLandscapeInfo();
 	const ULandscapeEditLayerBase* SplinesEditLayer = FindEditLayerOfTypeConst(ULandscapeEditLayerSplines::StaticClass());
@@ -14009,7 +14015,7 @@ void ALandscape::UpdateLandscapeSplines(const FGuid& InTargetLayer, bool bInUpda
 		FScopedSetLandscapeEditingLayer Scope(this, TargetLayerGuid, [this] { this->RequestLayersContentUpdate(ELandscapeLayerUpdateMode::Update_All); });
 		// Temporarily disable material instance updates since it will be done once at the end (requested by RequestLayersContentUpdateForceAll)
 		GDisableUpdateLandscapeMaterialInstances = true;
-		TSet<TObjectPtr<ULandscapeComponent>>* ModifiedComponent = nullptr;
+		TSet<TObjectPtr<ULandscapeComponent>>* ModifiedComponents = nullptr;
 		if (SplinesEditLayer != nullptr)
 		{
 			// Check that we can modify data
@@ -14041,12 +14047,10 @@ void ALandscape::UpdateLandscapeSplines(const FGuid& InTargetLayer, bool bInUpda
 			const bool bMarkPackageDirty = false;
 			ClearEditLayer(SplinesEditLayer->GetGuid(), (!bInForceUpdateAllCompoments && LandscapeSplinesAffectedComponents.Num()) ? &LandscapeSplinesAffectedComponents : nullptr, ELandscapeToolTargetTypeFlags::All, bMarkPackageDirty);
 			LandscapeSplinesAffectedComponents.Empty();
-			ModifiedComponent = &LandscapeSplinesAffectedComponents;
-			// For now, in Landscape Layer System Mode with a reserved layer for splines, we always update all the splines since we clear the whole layer first
-			bInUpdateOnlySelected = false;
+			ModifiedComponents = &LandscapeSplinesAffectedComponents;
 
 			// Apply splines without clearing up weightmap allocations
-			LandscapeInfo->ApplySplines(bInUpdateOnlySelected, ModifiedComponent, bMarkPackageDirty);
+			LandscapeInfo->ApplySplines(ModifiedComponents, bMarkPackageDirty);
 
 			for (const TPair<ULandscapeComponent*, uint32>& Pair : PreviousHashes)
 			{
@@ -14067,7 +14071,7 @@ void ALandscape::UpdateLandscapeSplines(const FGuid& InTargetLayer, bool bInUpda
 		}
 		else
 		{
-			LandscapeInfo->ApplySplines(bInUpdateOnlySelected, ModifiedComponent);
+			LandscapeInfo->ApplySplines(ModifiedComponents);
 		}
 		GDisableUpdateLandscapeMaterialInstances = false;
 	}
diff --git a/Engine/Source/Runtime/Landscape/Private/LandscapeSplineRaster.cpp b/Engine/Source/Runtime/Landscape/Private/LandscapeSplineRaster.cpp
--- a/Engine/Source/Runtime/Landscape/Private/LandscapeSplineRaster.cpp
+++ b/Engine/Source/Runtime/Landscape/Private/LandscapeSplineRaster.cpp
@@ -576,7 +576,13 @@ void RasterizeSegmentAlpha(int32& MinX, int32& MinY, int32& MaxX, int32& MaxY, F
 	LandscapeEdit.GetComponentsInRegion(MinX, MinY, MaxX, MaxY, &ModifiedComponents);
 }
 
+// Deprecated
 bool ULandscapeInfo::ApplySplines(bool bOnlySelected, TSet<TObjectPtr<ULandscapeComponent>>* OutModifiedComponents, bool bMarkPackageDirty)
+{
+	return ApplySplines(OutModifiedComponents, bMarkPackageDirty);
+}
+
+bool ULandscapeInfo::ApplySplines(TSet<TObjectPtr<ULandscapeComponent>>* OutModifiedComponents, bool bMarkPackageDirty)
 {
 	TRACE_CPUPROFILER_EVENT_SCOPE(LandscapeInfo_ApplySplines);
 
@@ -610,13 +616,13 @@ bool ULandscapeInfo::ApplySplines(bool bOnlySelected, TSet<TObjectPtr<ULandscape
 
 	ForAllSplineActors([&](TScriptInterface<ILandscapeSplineInterface> SplineOwner)
 	{
-		bResult |= ApplySplinesInternal(bOnlySelected, SplineOwner, OutModifiedComponents, bMarkPackageDirty, LandscapeMinX, LandscapeMinY, LandscapeMaxX, LandscapeMaxY, GetOrCreateModulate);
+		bResult |= ApplySplinesInternal(SplineOwner, OutModifiedComponents, bMarkPackageDirty, LandscapeMinX, LandscapeMinY, LandscapeMaxX, LandscapeMaxY, GetOrCreateModulate);
 	});
 
 	return bResult;
 }
 
-bool ULandscapeInfo::ApplySplinesInternal(bool bOnlySelected, TScriptInterface<ILandscapeSplineInterface> SplineOwner, TSet<TObjectPtr<ULandscapeComponent>>* OutModifiedComponents, bool bMarkPackageDirty, int32 LandscapeMinX, int32 LandscapeMinY, int32 LandscapeMaxX, int32 LandscapeMaxY, TFunctionRef<TSharedPtr<FModulateAlpha>(ULandscapeLayerInfoObject*)> GetOrCreateModulate)
+bool ULandscapeInfo::ApplySplinesInternal(TScriptInterface<ILandscapeSplineInterface> SplineOwner, TSet<TObjectPtr<ULandscapeComponent>>* OutModifiedComponents, bool bMarkPackageDirty, int32 LandscapeMinX, int32 LandscapeMinY, int32 LandscapeMaxX, int32 LandscapeMaxY, TFunctionRef<TSharedPtr<FModulateAlpha>(ULandscapeLayerInfoObject*)> GetOrCreateModulate)
 {
 	TRACE_CPUPROFILER_EVENT_SCOPE(LandscapeInfo_ApplySplinesInternal);
 
@@ -642,11 +648,6 @@ bool ULandscapeInfo::ApplySplinesInternal(bool bOnlySelected, TScriptInterface<I
 
 	for (const ULandscapeSplineControlPoint* ControlPoint : SplineComponent->ControlPoints)
 	{
-		if (bOnlySelected && !ControlPoint->IsSplineSelected())
-		{
-			continue;
-		}
-
 		if (ControlPoint->GetPoints().Num() < 2)
 		{
 			continue;
@@ -718,11 +719,6 @@ bool ULandscapeInfo::ApplySplinesInternal(bool bOnlySelected, TScriptInterface<I
 
 	for (const ULandscapeSplineSegment* Segment : SplineComponent->Segments)
 	{
-		if (bOnlySelected && !Segment->IsSplineSelected())
-		{
-			continue;
-		}
-
 		FBox SegmentBounds = Segment->GetBounds();
 		SegmentBounds = SegmentBounds.TransformBy(SplineToLandscape.ToMatrixWithScale());
 
