Landscape : fixed crash when auto-filling target layers from material

diff --git a/Engine/Source/Editor/LandscapeEditor/Private/LandscapeEditorDetailCustomization_TargetLayers.cpp b/Engine/Source/Editor/LandscapeEditor/Private/LandscapeEditorDetailCustomization_TargetLayers.cpp
--- a/Engine/Source/Editor/LandscapeEditor/Private/LandscapeEditorDetailCustomization_TargetLayers.cpp
+++ b/Engine/Source/Editor/LandscapeEditor/Private/LandscapeEditorDetailCustomization_TargetLayers.cpp
@@ -743,9 +743,9 @@ FReply FLandscapeEditorCustomNodeBuilder_TargetLayers::HandleAutoFillTargetLayer
 		if (TargetInfo->TargetType == ELandscapeToolTargetType::Weightmap && (!TargetInfo->LayerInfoObj.IsValid() || bUpdateAllLayers))
 		{
 			const FString TargetLayerAssetFilePath = LandscapeEdMode->GetTargetLayerAssetPackagePath();
-			const FAssetData* AssetData = LandscapeEditorUtils::FindLandscapeTargetLayerInfoAsset(TargetInfo->LayerName, TargetLayerAssetFilePath);
-			
-			if (AssetData != nullptr)
+
+			if (const TOptional<FAssetData> AssetData = LandscapeEditorUtils::FindLandscapeTargetLayerInfoAsset(TargetInfo->LayerName, TargetLayerAssetFilePath);
+				AssetData.IsSet())
 			{
 				OnTargetLayerSetObject(*AssetData, TargetInfo);
 			}
diff --git a/Engine/Source/Editor/LandscapeEditor/Private/LandscapeEditorObject.cpp b/Engine/Source/Editor/LandscapeEditor/Private/LandscapeEditorObject.cpp
--- a/Engine/Source/Editor/LandscapeEditor/Private/LandscapeEditorObject.cpp
+++ b/Engine/Source/Editor/LandscapeEditor/Private/LandscapeEditorObject.cpp
@@ -964,7 +964,8 @@ void ULandscapeEditorObject::AutoFillTargetLayers(const FString& DefaultTargetLa
 {
 	for (FLandscapeImportLayer& NewLayer : NewLandscape_Layers)
 	{
-		if (const FAssetData* AssetData = LandscapeEditorUtils::FindLandscapeTargetLayerInfoAsset(NewLayer.LayerName, DefaultTargetLayerAssetPath))
+		if (const TOptional<FAssetData> AssetData = LandscapeEditorUtils::FindLandscapeTargetLayerInfoAsset(NewLayer.LayerName, DefaultTargetLayerAssetPath); 
+			AssetData.IsSet())
 		{
 			NewLayer.LayerInfo = CastChecked<ULandscapeLayerInfoObject>(AssetData->GetAsset());
 		}
diff --git a/Engine/Source/Editor/LandscapeEditor/Private/LandscapeEditorUtils.cpp b/Engine/Source/Editor/LandscapeEditor/Private/LandscapeEditorUtils.cpp
--- a/Engine/Source/Editor/LandscapeEditor/Private/LandscapeEditorUtils.cpp
+++ b/Engine/Source/Editor/LandscapeEditor/Private/LandscapeEditorUtils.cpp
@@ -98,7 +98,7 @@ namespace LandscapeEditorUtils
 		LandscapeEditorUtils::SaveObjects(Proxies);
 	}
 
-	const TArray<FAssetData> GetLandscapeTargetLayerInfoAssets()
+	TArray<FAssetData> GetLandscapeTargetLayerInfoAssets()
 	{
 		TArray<FAssetData> LayerInfoAssets;
 
@@ -114,12 +114,12 @@ namespace LandscapeEditorUtils
 		return LayerInfoAssets;
 	}
 
-	const FAssetData* FindLandscapeTargetLayerInfoAsset(const FName& LayerName, const FString& TargetLayerAssetPackagePath)
+	TOptional<FAssetData> FindLandscapeTargetLayerInfoAsset(const FName& LayerName, const FString& TargetLayerAssetPackagePath)
 	{
-		const TArray<FAssetData>& LayerInfoAssets = GetLandscapeTargetLayerInfoAssets();
+		TArray<FAssetData> LayerInfoAssets = GetLandscapeTargetLayerInfoAssets();
 
 		int32 BestScore = -1.0f;
-		const FAssetData* BestLayerInfoObj = nullptr;
+		TOptional<FAssetData> BestLayerInfoObj;
 
 		for (const FAssetData& LayerInfoAsset : LayerInfoAssets)
 		{
@@ -136,7 +136,7 @@ namespace LandscapeEditorUtils
 				if (Score > BestScore)
 				{
 					BestScore = Score;
-					BestLayerInfoObj = &LayerInfoAsset;
+					BestLayerInfoObj.Emplace(LayerInfoAsset);
 				}
 			}
 		}
diff --git a/Engine/Source/Editor/LandscapeEditor/Public/LandscapeEditorUtils.h b/Engine/Source/Editor/LandscapeEditor/Public/LandscapeEditorUtils.h
--- a/Engine/Source/Editor/LandscapeEditor/Public/LandscapeEditorUtils.h
+++ b/Engine/Source/Editor/LandscapeEditor/Public/LandscapeEditorUtils.h
@@ -27,6 +27,6 @@ namespace LandscapeEditorUtils
 
 	void SaveLandscapeProxies(UWorld* InWorld, TArrayView<ALandscapeProxy*> Proxies);
 
-	const TArray<FAssetData> GetLandscapeTargetLayerInfoAssets();
-	const FAssetData* FindLandscapeTargetLayerInfoAsset(const FName& LayerName, const FString& TargetLayerAssetPackagePath);
+	TArray<FAssetData> GetLandscapeTargetLayerInfoAssets();
+	TOptional<FAssetData> FindLandscapeTargetLayerInfoAsset(const FName& LayerName, const FString& TargetLayerAssetPackagePath);
 }
