Landscape Editor - Adding option for landscape spline select all segments/control points to clear previous selection
* When using Select All Segments then Select All Points both segments and points are selected which hides the details view properties. A better UX is to reset the selection before selecting all

diff --git a/Engine/Source/Editor/LandscapeEditor/Private/LandscapeEdMode.h b/Engine/Source/Editor/LandscapeEditor/Private/LandscapeEdMode.h
--- a/Engine/Source/Editor/LandscapeEditor/Private/LandscapeEdMode.h
+++ b/Engine/Source/Editor/LandscapeEditor/Private/LandscapeEdMode.h
@@ -299,8 +299,8 @@ public:
 	bool HasEndLinearSplineConnection(ESplineNavigationFlags Flags) const;
 	void FlipSelectedSplineSegments();
 	void GetSelectedSplineOwners(TSet<AActor*>& SelectedSplineOwners) const;
-	virtual void SelectAllSplineControlPoints();
-	virtual void SelectAllSplineSegments();
+	virtual void SelectAllSplineControlPoints(ESplineNavigationFlags Flags);
+	virtual void SelectAllSplineSegments(ESplineNavigationFlags Flags);
 	virtual void SelectAllConnectedSplineControlPoints();
 	virtual void SelectAllConnectedSplineSegments();
 	virtual void SelectAdjacentLinearSplineElement(ESplineNavigationFlags Flags) const;
diff --git a/Engine/Source/Editor/LandscapeEditor/Private/LandscapeEdModeSplineTools.cpp b/Engine/Source/Editor/LandscapeEditor/Private/LandscapeEdModeSplineTools.cpp
--- a/Engine/Source/Editor/LandscapeEditor/Private/LandscapeEdModeSplineTools.cpp
+++ b/Engine/Source/Editor/LandscapeEditor/Private/LandscapeEdModeSplineTools.cpp
@@ -2255,25 +2255,37 @@ void FEdModeLandscape::GetSelectedSplineOwners(TSet<AActor*>& SelectedSplineOwne
 	}
 }
 
-void FEdModeLandscape::SelectAllSplineControlPoints()
+void FEdModeLandscape::SelectAllSplineControlPoints(ESplineNavigationFlags Flags)
 {
+	check(EnumOnlyContainsFlags(Flags, ESplineNavigationFlags::AddToSelection));
+
 	FScopedTransaction Transaction(LOCTEXT("LandscapeSpline_SelectAllPoints", "Select All Landscape Spline Points"));
 
 	ULandscapeInfo* CurrentLandscapeInfo = CurrentToolTarget.LandscapeInfo.Get();
 	if (SplinesTool && CurrentLandscapeInfo)
 	{
+		if (!EnumHasAllFlags(Flags, ESplineNavigationFlags::AddToSelection))
+		{
+			SplinesTool->SplineSelection->ClearSelection();
+		}
 		SplinesTool->SplineSelection->SelectAllControlPoints(*CurrentLandscapeInfo);
 		SplinesTool->SplineSelection->UpdatePropertiesWindows();
 	}
 }
 
-void FEdModeLandscape::SelectAllSplineSegments()
+void FEdModeLandscape::SelectAllSplineSegments(ESplineNavigationFlags Flags)
 {
+	check(EnumOnlyContainsFlags(Flags, ESplineNavigationFlags::AddToSelection));
+
 	FScopedTransaction Transaction(LOCTEXT("LandscapeSpline_SelectAllSegments", "Select All Landscape Spline Segments"));
 
 	ULandscapeInfo* CurrentLandscapeInfo = CurrentToolTarget.LandscapeInfo.Get();
 	if (SplinesTool && CurrentLandscapeInfo)
 	{
+		if (!EnumHasAllFlags(Flags, ESplineNavigationFlags::AddToSelection))
+		{
+			SplinesTool->SplineSelection->ClearSelection();
+		}
 		SplinesTool->SplineSelection->SelectAllSplineSegments(*CurrentLandscapeInfo);
 		SplinesTool->SplineSelection->UpdatePropertiesWindows();
 	}
diff --git a/Engine/Source/Editor/LandscapeEditor/Private/LandscapeEditorDetailCustomization_MiscTools.cpp b/Engine/Source/Editor/LandscapeEditor/Private/LandscapeEditorDetailCustomization_MiscTools.cpp
--- a/Engine/Source/Editor/LandscapeEditor/Private/LandscapeEditorDetailCustomization_MiscTools.cpp
+++ b/Engine/Source/Editor/LandscapeEditor/Private/LandscapeEditorDetailCustomization_MiscTools.cpp
@@ -20,6 +20,7 @@
 #include "Widgets/Input/SNumericEntryBox.h"
 #include "SFlattenHeightEyeDropperButton.h"
 #include "LandscapeEdModeTools.h"
+#include "LandscapeSplineSelection.h"
 
 #define LOCTEXT_NAMESPACE "LandscapeEditor.Tools"
 
@@ -479,7 +480,8 @@ FReply FLandscapeEditorDetailCustomization_MiscTools::OnSelectAllControlPointsBu
 	FEdModeLandscape* LandscapeEdMode = GetEditorMode();
 	if (LandscapeEdMode && LandscapeEdMode->CurrentToolTarget.LandscapeInfo.IsValid())
 	{
-		LandscapeEdMode->SelectAllSplineControlPoints();
+		// Clears current selection before selecting all control points
+		LandscapeEdMode->SelectAllSplineControlPoints(ESplineNavigationFlags::None);
 	}
 
 	return FReply::Handled();
@@ -490,7 +492,8 @@ FReply FLandscapeEditorDetailCustomization_MiscTools::OnSelectAllSegmentsButtonC
 	FEdModeLandscape* LandscapeEdMode = GetEditorMode();
 	if (LandscapeEdMode && LandscapeEdMode->CurrentToolTarget.LandscapeInfo.IsValid())
 	{
-		LandscapeEdMode->SelectAllSplineSegments();
+		// Clears current selection before selecting all segments
+		LandscapeEdMode->SelectAllSplineSegments(ESplineNavigationFlags::None);
 	}
 
 	return FReply::Handled();
