Fixed potential issue sorting layer info assets

diff --git a/Engine/Source/Editor/LandscapeEditor/Private/LandscapeEditorUtils.cpp b/Engine/Source/Editor/LandscapeEditor/Private/LandscapeEditorUtils.cpp
--- a/Engine/Source/Editor/LandscapeEditor/Private/LandscapeEditorUtils.cpp
+++ b/Engine/Source/Editor/LandscapeEditor/Private/LandscapeEditorUtils.cpp
@@ -118,7 +118,7 @@ namespace LandscapeEditorUtils
 	{
 		TArray<FAssetData> LayerInfoAssets = GetLandscapeTargetLayerInfoAssets();
 
-		int32 BestScore = -1.0f;
+		float BestScore = -1.0f;
 		TOptional<FAssetData> BestLayerInfoObj;
 
 		for (const FAssetData& LayerInfoAsset : LayerInfoAssets)
