Landscape Editor - Edit layer context menu customization system. Adds spline edit layer actions on the context menu and double click
* Adds edit layer customization registration system to create custom actions for edit layer types
* Deprecates ULandscapeEditLayerBase::GetActions and FEditLayerAction to remove custom actions from runtime logic
* Moves generic edit layer editor functions into a core class
* Toggle spline tool on edit layer double click and context menu button
* Fixes missing tooltips on edit layer icons

diff --git a/Engine/Plugins/Experimental/LandscapePatch/Source/LandscapePatch/Public/LandscapePatchEditLayer.h b/Engine/Plugins/Experimental/LandscapePatch/Source/LandscapePatch/Public/LandscapePatchEditLayer.h
--- a/Engine/Plugins/Experimental/LandscapePatch/Source/LandscapePatch/Public/LandscapePatchEditLayer.h
+++ b/Engine/Plugins/Experimental/LandscapePatch/Source/LandscapePatch/Public/LandscapePatchEditLayer.h
@@ -14,7 +14,7 @@ class ULandscapePatchComponent;
  * Special edit layer used only for landscape patches. Patches point to the layer via a guid
  *  and determine their ordering by their Priority values.
  */
-UCLASS(MinimalAPI)
+UCLASS(MinimalAPI, meta = (ShortToolTip = "Special edit layer for landscape patches"))
 class ULandscapePatchEditLayer : public ULandscapeEditLayerProcedural
 	, public ILandscapeBrushRenderCallAdapter_GlobalMergeLegacySupport
 {
diff --git a/Engine/Source/Editor/LandscapeEditor/Private/LandscapeEdMode.cpp b/Engine/Source/Editor/LandscapeEditor/Private/LandscapeEdMode.cpp
--- a/Engine/Source/Editor/LandscapeEditor/Private/LandscapeEdMode.cpp
+++ b/Engine/Source/Editor/LandscapeEditor/Private/LandscapeEdMode.cpp
@@ -4010,19 +4010,35 @@ void FEdModeLandscape::SetSelectedEditLayer(int32 InLayerIndex)
 	{
 		Landscape->Modify(false);
 		Landscape->SetSelectedEditLayerIndex(InLayerIndex);
-		const ULandscapeEditLayerBase* EditLayer = Landscape->GetEditLayerConst(InLayerIndex);
-
-		if (EditLayer && EditLayer->IsA<ULandscapeEditLayerSplines>())
-		{
-			SetCurrentToolMode("ToolMode_Manage", false);
-			SetCurrentTool(FName("Splines"));
-		}
 	}
 
 	RefreshDetailPanel();
 	RequestLayersContentUpdateForceAll(ELandscapeLayerUpdateMode::Update_Client_Editing);
 }
 
+void FEdModeLandscape::ToggleSplinesTool(const ULandscapeEditLayerBase* InEditLayer)
+{
+	const ALandscape* Landscape = GetLandscape();
+	 
+	if (Landscape == nullptr || InEditLayer == nullptr)
+	{
+		return;
+	}
+
+	if (InEditLayer->IsA<ULandscapeEditLayerSplines>() && CurrentTool->GetToolName() != FName("Splines"))
+	{
+		SetCurrentToolMode("ToolMode_Manage", false);
+		SetCurrentTool(FName("Splines"));
+	}
+	else
+	{
+		SetCurrentToolMode("ToolMode_Sculpt", false);
+		SetCurrentTool(FName("Sculpt"));
+	}
+
+	SetSelectedEditLayer(Landscape->GetLayerIndex(InEditLayer->GetGuid()));
+}
+
 int32 FEdModeLandscape::GetSelectedEditLayerIndex() const
 {
 	if (ALandscape* Landscape = GetLandscape())
diff --git a/Engine/Source/Editor/LandscapeEditor/Private/LandscapeEdMode.h b/Engine/Source/Editor/LandscapeEditor/Private/LandscapeEdMode.h
--- a/Engine/Source/Editor/LandscapeEditor/Private/LandscapeEdMode.h
+++ b/Engine/Source/Editor/LandscapeEditor/Private/LandscapeEdMode.h
@@ -547,6 +547,7 @@ public:
 	bool IsGridBased() const;
 
 	// Edit Layers
+	void ToggleSplinesTool(const ULandscapeEditLayerBase* InEditLayer);
 	int32 GetLayerCount() const;
 	void SetSelectedEditLayer(int32 InLayerIndex);
 	int32 GetSelectedEditLayerIndex() const;
diff --git a/Engine/Source/Editor/LandscapeEditor/Private/LandscapeEditorDetailCustomization_EditLayers.cpp b/Engine/Source/Editor/LandscapeEditor/Private/LandscapeEditorDetailCustomization_EditLayers.cpp
new file mode 100644
--- /dev/null
+++ b/Engine/Source/Editor/LandscapeEditor/Private/LandscapeEditorDetailCustomization_EditLayers.cpp
@@ -0,0 +1,605 @@
+// Copyright Epic Games, Inc. All Rights Reserved.
+
+#include "LandscapeEditorDetailCustomization_EditLayers.h"
+
+#include "EditorModeManager.h"
+#include "EditorModes.h"
+#include "Framework/MultiBox/MultiBoxBuilder.h"
+#include "LandscapeEdMode.h"
+#include "LandscapeEditLayerCustomization.h"
+#include "LandscapeEditLayer.h"
+#include "LandscapeEditorUtils.h"
+#include "Misc/MessageDialog.h"
+#include "ScopedTransaction.h"
+
+#define LOCTEXT_NAMESPACE "LandscapeEditor.EditLayersCustomization"
+
+FEdModeLandscape* FLandscapeEditLayerCustomizationCommon::GetEditorMode()
+{
+	return (FEdModeLandscape*)GLevelEditorModeTools().GetActiveMode(FBuiltinEditorModes::EM_Landscape);
+}
+
+bool FLandscapeEditLayerCustomizationCommon::CanDeleteLayer(int32 InLayerIndex, FText& OutReason) const
+{
+	FEdModeLandscape* LandscapeEdMode = GetEditorMode();
+	ALandscape* Landscape = LandscapeEdMode ? LandscapeEdMode->GetLandscape() : nullptr;
+	if (!Landscape || Landscape->GetEditLayersConst().Num() <= 1)
+	{
+		OutReason = LOCTEXT("DeleteLayer_CantDeleteLastLayer", "The last layer cannot be deleted");
+		return false;
+	}
+
+	const ULandscapeEditLayerBase* EditLayer = LandscapeEdMode->GetEditLayerConst(InLayerIndex);
+	if (EditLayer == nullptr)
+	{
+		return false;
+	}
+
+	if (EditLayer->IsLocked())
+	{
+		OutReason = LOCTEXT("DeleteLayer_CantDeleteLocked", "Cannot delete a locked edit layer");
+		return false;
+	}
+
+	OutReason = LOCTEXT("DeleteLayer_CanDelete", "Delete the edit layer");
+	return true;
+}
+
+void FLandscapeEditLayerCustomizationCommon::DeleteLayer(int32 InLayerIndex)
+{
+	FEdModeLandscape* LandscapeEdMode = GetEditorMode();
+	ALandscape* Landscape = LandscapeEdMode ? LandscapeEdMode->GetLandscape() : nullptr;
+	if (Landscape && Landscape->GetEditLayersConst().Num() > 1)
+	{
+		const ULandscapeEditLayerBase* EditLayer = LandscapeEdMode->GetEditLayerConst(InLayerIndex);
+		if (EditLayer)
+		{
+			EAppReturnType::Type Result = FMessageDialog::Open(EAppMsgType::YesNo, FText::Format(LOCTEXT("Landscape_DeleteLayer_Message", "The layer {0} will be deleted.  Continue?"), FText::FromName(EditLayer->GetName())));
+			if (Result == EAppReturnType::Yes)
+			{
+				const FScopedTransaction Transaction(LOCTEXT("Landscape_Layers_Delete", "Delete Layer"));
+				// Delete layer will update the selected edit layer index
+				Landscape->DeleteLayer(InLayerIndex);
+				LandscapeEdMode->UpdateTargetList();
+				LandscapeEdMode->RefreshDetailPanel();
+			}
+		}
+	}
+}
+
+bool FLandscapeEditLayerCustomizationCommon::CanCollapseLayer(int32 InLayerIndex, FText& OutReason) const
+{
+	FEdModeLandscape* LandscapeEdMode = GetEditorMode();
+	ALandscape* Landscape = LandscapeEdMode ? LandscapeEdMode->GetLandscape() : nullptr;
+	if (!Landscape || Landscape->GetEditLayersConst().Num() <= 1)
+	{
+		OutReason = LOCTEXT("Landscape_CollapseLayer_Reason_NotEnoughLayersToCollapse", "Not enough layers to do collapse");
+		return false;
+	}
+
+	if (InLayerIndex < 1)
+	{
+		OutReason = LOCTEXT("Landscape_CollapseLayer_Reason_CantCollapseBaseLayer", "Cannot collapse the first layer");
+		return false;
+	}
+
+	const FLandscapeLayer* TopLayer = Landscape->GetLayerConst(InLayerIndex);
+	const FLandscapeLayer* BottomLayer = Landscape->GetLayerConst(InLayerIndex - 1);
+	if ((TopLayer == nullptr) || (BottomLayer == nullptr))
+	{
+		return false;
+	}
+
+	const ULandscapeEditLayerBase* TopEditLayer = LandscapeEdMode->GetEditLayerConst(InLayerIndex);
+	const ULandscapeEditLayerBase* BottomEditLayer = LandscapeEdMode->GetEditLayerConst(InLayerIndex - 1);
+
+	check((TopEditLayer != nullptr) && (BottomEditLayer != nullptr));
+
+	if (!TopEditLayer->SupportsCollapsingTo())
+	{
+		OutReason = FText::Format(LOCTEXT("Landscape_CollapseLayer_Reason_TopLayerDoesntSupportCollapsing", "Cannot collapse layer '{0}' onto layer '{1}'. The type of layer '{0}' ({2}) doesn't support collapsing to another one"),
+			FText::FromName(TopEditLayer->GetName()), FText::FromName(BottomEditLayer->GetName()), TopEditLayer->GetClass()->GetDisplayNameText());
+		return false;
+	}
+
+	if (!BottomEditLayer->SupportsBeingCollapsedAway())
+	{
+		OutReason = FText::Format(LOCTEXT("Landscape_CollapseLayer_Reason_BottomLayerDoesntSupportCollapsing", "Cannot collapse layer '{0}' onto layer '{1}'. The type of layer '{1}' ({2}) doesn't support being collapsed away"),
+			FText::FromName(TopEditLayer->GetName()), FText::FromName(BottomEditLayer->GetName()), BottomEditLayer->GetClass()->GetDisplayNameText());
+		return false;
+	}
+
+	if (TopEditLayer->IsLocked())
+	{
+		OutReason = FText::Format(LOCTEXT("Landscape_CollapseLayer_Reason_TopLayerIsLocked", "Cannot collapse layer '{0}' onto layer '{1}'. Layer '{0}' will be deleted in the operation but it is currently locked"),
+			FText::FromName(TopEditLayer->GetName()), FText::FromName(BottomEditLayer->GetName()));
+		return false;
+	}
+
+	if (BottomEditLayer->IsLocked())
+	{
+		OutReason = FText::Format(LOCTEXT("Landscape_CollapseLayer_Reason_BottomLayerIsLocked", "Cannot collapse layer '{0}' onto layer '{1}'. Destination layer '{1}' is currently locked"),
+			FText::FromName(TopEditLayer->GetName()), FText::FromName(BottomEditLayer->GetName()));
+		return false;
+	}
+
+	// Can't collapse on layer that has a Brush because result will change...
+	if (TopLayer->Brushes.Num() > 0)
+	{
+		OutReason = FText::Format(LOCTEXT("Landscape_CollapseLayer_Reason_TopLayerHasBrush", "Cannot collapse layer '{0}' onto layer '{1}'. Layer '{0}' contains brush(es)"),
+			FText::FromName(TopEditLayer->GetName()), FText::FromName(BottomEditLayer->GetName()));
+		return false;
+	}
+
+	if (BottomLayer->Brushes.Num() > 0)
+	{
+		OutReason = FText::Format(LOCTEXT("Landscape_CollapseLayer_Reason_BottomLayerHasBrush", "Cannot collapse layer '{0}' onto layer '{1}'. Layer '{1}' contains brush(es)"),
+			FText::FromName(TopEditLayer->GetName()), FText::FromName(BottomEditLayer->GetName()));
+		return false;
+	}
+
+	OutReason = FText::Format(LOCTEXT("Landscape_CollapseLayer_Reason_Collapse", "Collapse layer '{0}' onto layer '{1}'"),
+		FText::FromName(TopEditLayer->GetName()), FText::FromName(BottomEditLayer->GetName()));
+	return true;
+}
+
+void FLandscapeEditLayerCustomizationCommon::CollapseLayer(int32 InLayerIndex)
+{
+	FText Reason;
+	check(CanCollapseLayer(InLayerIndex, Reason));
+
+	if (FEdModeLandscape* LandscapeEdMode = GetEditorMode())
+	{
+		ALandscape* Landscape = LandscapeEdMode->GetLandscape();
+		const ULandscapeEditLayerBase* Layer = LandscapeEdMode->GetEditLayerConst(InLayerIndex);
+		const ULandscapeEditLayerBase* BaseLayer = LandscapeEdMode->GetEditLayerConst(InLayerIndex - 1);
+
+		if (Landscape && Layer && BaseLayer)
+		{
+			EAppReturnType::Type Result = FMessageDialog::Open(EAppMsgType::YesNo, FText::Format(LOCTEXT("Landscape_CollapseLayer_Message", "The layer {0} will be collapsed into layer {1}.  Continue?"), FText::FromName(Layer->GetName()), FText::FromName(BaseLayer->GetName())));
+			if (Result == EAppReturnType::Yes)
+			{
+				const FScopedTransaction Transaction(LOCTEXT("Landscape_Layers_Collapse", "Collapse Layer"));
+				Landscape->CollapseLayer(InLayerIndex);
+				OnLayerSelectionChanged(InLayerIndex - 1);
+				LandscapeEdMode->RefreshDetailPanel();
+			}
+		}
+	}
+}
+
+bool FLandscapeEditLayerCustomizationCommon::CanToggleVisibility(int32 InLayerIndex, FText& OutReason) const
+{
+	FEdModeLandscape* LandscapeEdMode = GetEditorMode();
+
+	const ULandscapeEditLayerBase* EditLayer = LandscapeEdMode->GetEditLayerConst(InLayerIndex);
+	if (EditLayer == nullptr)
+	{
+		return false;
+	}
+
+	if (EditLayer->IsLocked())
+	{
+		OutReason = LOCTEXT("ToggleVisibility_CantToggleLocked", "Cannot change the visibility of a locked edit layer");
+		return false;
+	}
+
+	OutReason = LOCTEXT("ToggleVisibility_CanToggle", "Toggle the visibility of the edit layer");
+	return true;
+}
+
+FReply FLandscapeEditLayerCustomizationCommon::OnToggleVisibility(int32 InLayerIndex)
+{
+	if (FEdModeLandscape* LandscapeEdMode = GetEditorMode())
+	{
+		const FScopedTransaction Transaction(LOCTEXT("Landscape_Layers_SetVisibility", "Set Layer Visibility"));
+
+		ULandscapeEditLayerBase* EditLayer = LandscapeEdMode->GetEditLayer(InLayerIndex);
+		check(EditLayer != nullptr);
+		EditLayer->SetVisible(!EditLayer->IsVisible(), /*bInModify= */true);
+
+		if (EditLayer->IsVisible())
+		{
+			OnLayerSelectionChanged(InLayerIndex);
+		}
+	}
+	return FReply::Handled();
+}
+
+const FSlateBrush* FLandscapeEditLayerCustomizationCommon::GetVisibilityBrushForLayer(int32 InLayerIndex) const
+{
+	bool bIsVisible = false;
+
+	if (FEdModeLandscape* LandscapeEdMode = GetEditorMode())
+	{
+		const ULandscapeEditLayerBase* EditLayer = LandscapeEdMode->GetEditLayerConst(InLayerIndex);
+		check(EditLayer != nullptr);
+
+		bIsVisible = EditLayer->IsVisible();
+	}
+
+	return bIsVisible ? FAppStyle::GetBrush("Level.VisibleIcon16x") : FAppStyle::GetBrush("Level.NotVisibleIcon16x");
+}
+
+void FLandscapeEditLayerCustomizationCommon::OnLayerSelectionChanged(int32 InLayerIndex)
+{
+	if (FEdModeLandscape* LandscapeEdMode = GetEditorMode(); LandscapeEdMode && LandscapeEdMode->GetSelectedEditLayerIndex() != InLayerIndex)
+	{
+		FScopedTransaction Transaction(LOCTEXT("Landscape_Layers_SetCurrentLayer", "Set Current Layer"));
+		LandscapeEdMode->SetSelectedEditLayer(InLayerIndex);
+		LandscapeEdMode->UpdateTargetList();
+	}
+}
+
+void FLandscapeEditLayerCustomizationCommon::ShowOnlySelectedLayer(int32 InLayerIndex)
+{
+	FEdModeLandscape* LandscapeEdMode = GetEditorMode();
+	ALandscape* Landscape = LandscapeEdMode ? LandscapeEdMode->GetLandscape() : nullptr;
+	if (Landscape)
+	{
+		const FScopedTransaction Transaction(LOCTEXT("ShowOnlySelectedLayer", "Show Only Selected Layer"));
+		Landscape->ShowOnlySelectedLayer(InLayerIndex);
+		OnLayerSelectionChanged(InLayerIndex);
+	}
+}
+
+void FLandscapeEditLayerCustomizationCommon::ShowAllLayers()
+{
+	FEdModeLandscape* LandscapeEdMode = GetEditorMode();
+	ALandscape* Landscape = LandscapeEdMode ? LandscapeEdMode->GetLandscape() : nullptr;
+	if (Landscape)
+	{
+		const FScopedTransaction Transaction(LOCTEXT("ShowAllLayers", "Show All Layers"));
+		Landscape->ShowAllLayers();
+	}
+}
+
+
+TOptional<float> FLandscapeEditLayerCustomizationCommon::GetLayerAlpha(int32 InLayerIndex) const
+{
+	if (FEdModeLandscape* LandscapeEdMode = GetEditorMode())
+	{
+		const ULandscapeEditLayerBase* EditLayer = LandscapeEdMode->GetEditLayerConst(InLayerIndex);
+		check(EditLayer != nullptr);
+
+		return EditLayer->GetAlphaForTargetType(LandscapeEdMode->GetLandscapeToolTargetType());
+	}
+
+	return 1.0f;
+}
+
+float FLandscapeEditLayerCustomizationCommon::GetLayerAlphaMinValue() const
+{
+	if (const FEdModeLandscape* LandscapeEdMode = GetEditorMode(); LandscapeEdMode && LandscapeEdMode->GetLandscapeToolTargetType() == ELandscapeToolTargetType::Heightmap)
+	{
+		return -1.0f;
+	}
+	else
+	{
+		return 0.0f;
+	}
+}
+
+bool FLandscapeEditLayerCustomizationCommon::CanSetLayerAlpha(int32 InLayerIndex, FText& OutReason) const
+{
+	const FEdModeLandscape* LandscapeEdMode = GetEditorMode();
+
+	if (LandscapeEdMode == nullptr)
+	{
+		return false;
+	}
+
+	const ULandscapeEditLayerBase* EditLayer = LandscapeEdMode->GetEditLayerConst(InLayerIndex);
+	if (EditLayer == nullptr)
+	{
+		return false;
+	}
+
+	if (!EditLayer->SupportsAlphaForTargetType(LandscapeEdMode->GetLandscapeToolTargetType()))
+	{
+		OutReason = FText::Format(LOCTEXT("SetLayerAlpha_LayerDoesntSupportAlpha", "Cannot change alpha : the type of layer {0} ({1}) doesn't support alpha"),
+			FText::FromName(EditLayer->GetName()), EditLayer->GetClass()->GetDisplayNameText());
+		return false;
+	}
+
+
+	if (EditLayer->IsLocked())
+	{
+		OutReason = LOCTEXT("SetLayerAlpha_LayerIsLocked", "Cannot change the alpha of a locked edit layer");
+		return false;
+	}
+
+	OutReason = LOCTEXT("SetLayerAlpha_CanSet", "Set the edit layer's alpha");
+	return true;
+}
+
+void FLandscapeEditLayerCustomizationCommon::SetLayerAlpha(float InAlpha, int32 InLayerIndex, bool bCommit, int32 InSliderIndex)
+{
+	FEdModeLandscape* LandscapeEdMode = GetEditorMode();
+	if (LandscapeEdMode)
+	{
+		ULandscapeEditLayerBase* EditLayer = LandscapeEdMode->GetEditLayer(InLayerIndex);
+		check(EditLayer != nullptr);
+
+		FScopedTransaction Transaction(LOCTEXT("Landscape_Layers_SetAlpha", "Set Layer Alpha"), InSliderIndex == INDEX_NONE && bCommit);
+		// Set Value when using slider or when committing text
+		EditLayer->SetAlphaForTargetType(LandscapeEdMode->GetLandscapeToolTargetType(), InAlpha, /*bInModify = */true, bCommit ? EPropertyChangeType::ValueSet : EPropertyChangeType::Interactive);
+	}
+}
+
+TSharedRef<IEditLayerCustomization> FLandscapeEditLayerContextMenuCustomization_Base::MakeInstance()
+{
+	return MakeShareable(new FLandscapeEditLayerContextMenuCustomization_Base);
+}
+
+void FLandscapeEditLayerContextMenuCustomization_Base::CustomizeContextMenu(ULandscapeEditLayerBase* InEditLayer, FEditLayerCategoryToEntryMap& OutContextMenuEntryMap)
+{
+	FEdModeLandscape* LandscapeEdMode = GetEditorMode();
+	ALandscape* Landscape = LandscapeEdMode ? LandscapeEdMode->GetLandscape() : nullptr;
+
+	if (Landscape && LandscapeEdMode && InEditLayer)
+	{
+		const int32 EditLayerIndex = Landscape->GetLayerIndex(InEditLayer->GetGuid());
+
+		const FName EditLayerCategory = FName("Edit Layers");
+		OutContextMenuEntryMap.FindOrAdd(EditLayerCategory).SectionLabel = LOCTEXT("LandscapeEditLayerContextMenuCommon.Heading", "Edit Layers");
+
+		// Delete Layer
+		FMenuEntryParams DeleteLayerEntry;
+		DeleteLayerEntry.DirectActions = FUIAction(
+			FExecuteAction::CreateSP(this, &FLandscapeEditLayerContextMenuCustomization_Base::DeleteLayer, EditLayerIndex),
+			FCanExecuteAction::CreateSPLambda(this, [this, EditLayerIndex] { FText Reason; return CanDeleteLayer(EditLayerIndex, Reason); }));
+		DeleteLayerEntry.ToolTipOverride = MakeAttributeLambda([this, EditLayerIndex] {
+			FText Reason;
+			CanDeleteLayer(EditLayerIndex, Reason);
+			return Reason; });
+		DeleteLayerEntry.LabelOverride = LOCTEXT("DeleteLayer", "Delete");
+
+		OutContextMenuEntryMap.FindOrAdd(EditLayerCategory).Entries.Add(DeleteLayerEntry);
+
+		// Collapse Layer
+		FMenuEntryParams CollapseLayerEntry;
+		CollapseLayerEntry.DirectActions = FUIAction(
+			FExecuteAction::CreateSP(this, &FLandscapeEditLayerContextMenuCustomization_Base::CollapseLayer, EditLayerIndex),
+			FCanExecuteAction::CreateSPLambda(this, [this, EditLayerIndex] { FText Reason; return CanCollapseLayer(EditLayerIndex, Reason); }));
+		CollapseLayerEntry.ToolTipOverride = MakeAttributeSPLambda(this, [this, EditLayerIndex] { FText Reason; CanCollapseLayer(EditLayerIndex, Reason); return Reason; });
+		CollapseLayerEntry.LabelOverride = LOCTEXT("CollapseLayer", "Collapse");
+
+		OutContextMenuEntryMap.FindOrAdd(EditLayerCategory).Entries.Add(CollapseLayerEntry);
+		
+		// Visibility Section (new category)
+		const FName VisibilityCategory = "Visibility";
+		OutContextMenuEntryMap.FindOrAdd(VisibilityCategory).SectionLabel = LOCTEXT("LandscapeEditLayerContextMenuVisibility.Heading", "Visibility");
+
+		// Show/Hide Selected Layer
+		FText VisibilityText = LOCTEXT("ShowSelected", "Show Selected");
+		if (InEditLayer->IsVisible())
+		{
+			// Hide Selected Layer
+			VisibilityText = LOCTEXT("HideSelected", "Hide Selected");
+		}
+		FMenuEntryParams VisibilityEntry;
+		VisibilityEntry.LabelOverride = VisibilityText;
+		VisibilityEntry.ToolTipOverride =  MakeAttributeSPLambda(this, [this, EditLayerIndex] { FText Reason; CanToggleVisibility(EditLayerIndex, Reason); return Reason; });
+		VisibilityEntry.DirectActions = FUIAction(FExecuteAction::CreateSPLambda(this, [this, EditLayerIndex] { OnToggleVisibility(EditLayerIndex); }),
+			FCanExecuteAction::CreateSPLambda(this, [this, EditLayerIndex] { FText Reason; return CanToggleVisibility(EditLayerIndex, Reason); }));
+
+		OutContextMenuEntryMap.FindOrAdd(VisibilityCategory).Entries.Add(VisibilityEntry);
+
+		// Show Only Selected Layer
+		FUIAction ShowOnlySelectedLayerAction = FUIAction(FExecuteAction::CreateSPLambda(this, [this, EditLayerIndex] { ShowOnlySelectedLayer(EditLayerIndex); }));
+		FMenuEntryParams ShowOnlyEntry;
+		ShowOnlyEntry.LabelOverride = LOCTEXT("ShowOnlySelected", "Show Only Selected");
+		ShowOnlyEntry.ToolTipOverride = LOCTEXT("ShowOnlySelectedLayerTooltip", "Show Only Selected Layer");
+		ShowOnlyEntry.DirectActions = ShowOnlySelectedLayerAction;
+
+		OutContextMenuEntryMap.FindOrAdd(VisibilityCategory).Entries.Add(ShowOnlyEntry);
+
+		// Show All Layers
+		FUIAction ShowAllLayersAction = FUIAction(FExecuteAction::CreateSP(this, &FLandscapeEditLayerContextMenuCustomization_Base::ShowAllLayers));
+		FMenuEntryParams ShowAllEntry;
+		ShowAllEntry.LabelOverride = LOCTEXT("ShowAllLayers", "Show All Layers");
+		ShowAllEntry.ToolTipOverride = LOCTEXT("ShowAllLayersTooltip", "Show All Layers");
+		ShowAllEntry.DirectActions = ShowAllLayersAction;
+
+		OutContextMenuEntryMap.FindOrAdd(VisibilityCategory).Entries.Add(ShowAllEntry);
+	}
+}
+
+TSharedRef<IEditLayerCustomization> FLandscapeEditLayerContextMenuCustomization_Layer::MakeInstance()
+{
+	return MakeShareable(new FLandscapeEditLayerContextMenuCustomization_Layer);
+}
+
+void FLandscapeEditLayerContextMenuCustomization_Layer::CustomizeContextMenu(ULandscapeEditLayerBase* InEditLayer, FEditLayerCategoryToEntryMap& OutContextMenuEntryMap)
+{
+	FEdModeLandscape* LandscapeEdMode = GetEditorMode();
+	ALandscape* Landscape = LandscapeEdMode ? LandscapeEdMode->GetLandscape() : nullptr;
+	if (Landscape && LandscapeEdMode && InEditLayer)
+	{
+		const int32 EditLayerIndex = Landscape->GetLayerIndex(InEditLayer->GetGuid());
+		const FName BrushCategory = FName("Brushes");
+		OutContextMenuEntryMap.FindOrAdd(BrushCategory).SectionLabel = LOCTEXT("LandscapeEditLayerContextMenuBrushes.Heading", "Brushes");
+
+		// Add unassigned brush to layer
+		const TArray<ALandscapeBlueprintBrushBase*>& Brushes = LandscapeEdMode->GetBrushList();
+		TArray<ALandscapeBlueprintBrushBase*> FilteredBrushes = Brushes.FilterByPredicate([](ALandscapeBlueprintBrushBase* Brush) { return Brush->GetOwningLandscape() == nullptr; });
+
+		// If there are no unassigned brushes or the edit layer does not support brushes, show a disabled state instead of hiding the entire section
+		if (FilteredBrushes.Num() > 0 && InEditLayer->SupportsBlueprintBrushes())
+		{
+			FMenuEntryParams AssignBrushEntry;
+			AssignBrushEntry.LabelOverride = LOCTEXT("LandscapeEditorBrushAddSubMenu", "Assign Existing Brush");
+			AssignBrushEntry.ToolTipOverride = LOCTEXT("LandscapeEditorBrushAddSubMenuToolTip", "To modify the terrain, brushes need to be assigned to a landscape actor. Add the brush to this edit layer");
+			AssignBrushEntry.bIsSubMenu = true;
+			AssignBrushEntry.EntryBuilder = FNewMenuDelegate::CreateSP(this, &FLandscapeEditLayerContextMenuCustomization_Layer::FillUnassignedBrushMenu, FilteredBrushes, EditLayerIndex);
+
+			OutContextMenuEntryMap.FindOrAdd(BrushCategory).Entries.Add(AssignBrushEntry);
+		}
+		else
+		{
+			FMenuEntryParams NoBrushEntry;
+			NoBrushEntry.LabelOverride = LOCTEXT("LandscapeEditorBrushNone", "None"),
+				NoBrushEntry.ToolTipOverride = InEditLayer->SupportsBlueprintBrushes() ? LOCTEXT("LandscapeEditorBrushAllBrushActorsAssigned", "All Blueprint Brush actors are assigned to a landscape edit layer")
+				: FText::Format(LOCTEXT("LandscapeEditorBrushUnsupported", "This layer's type ({0}) doesn't support blueprint brushes."), InEditLayer->GetClass()->GetDisplayNameText());
+			NoBrushEntry.DirectActions = FUIAction(FExecuteAction::CreateLambda([]() {}),
+				FCanExecuteAction::CreateLambda([]() { return false; }));
+
+			OutContextMenuEntryMap.FindOrAdd(BrushCategory).Entries.Add(NoBrushEntry);
+		}
+	}
+}
+
+void FLandscapeEditLayerContextMenuCustomization_Layer::FillUnassignedBrushMenu(FMenuBuilder& MenuBuilder, TArray<ALandscapeBlueprintBrushBase*> Brushes, int32 InLayerIndex)
+{
+	for (ALandscapeBlueprintBrushBase* Brush : Brushes)
+	{
+		FUIAction AddAction = FUIAction(FExecuteAction::CreateSP(this, &FLandscapeEditLayerContextMenuCustomization_Layer::AssignBrushToEditLayer, Brush, InLayerIndex));
+		MenuBuilder.AddMenuEntry(FText::FromString(Brush->GetActorLabel()), FText(), FSlateIcon(), AddAction);
+	}
+}
+
+void FLandscapeEditLayerContextMenuCustomization_Layer::AssignBrushToEditLayer(ALandscapeBlueprintBrushBase* Brush, const int32 InLayerIndex)
+{
+	const FScopedTransaction Transaction(LOCTEXT("LandscapeBrushAddToCurrentLayerTransaction", "Add brush to edit layer"));
+	if (FEdModeLandscape* LandscapeEdMode = GetEditorMode())
+	{
+		if (ALandscape* Landscape = LandscapeEdMode->GetLandscape())
+		{
+			Landscape->AddBrushToLayer(InLayerIndex, Brush);
+		}
+	}
+}
+
+TSharedRef<IEditLayerCustomization> FLandscapeEditLayerContextMenuCustomization_Persistent::MakeInstance()
+{
+	return MakeShareable(new FLandscapeEditLayerContextMenuCustomization_Persistent);
+}
+
+void FLandscapeEditLayerContextMenuCustomization_Persistent::CustomizeContextMenu(ULandscapeEditLayerBase* InEditLayer, FEditLayerCategoryToEntryMap& OutContextMenuEntryMap)
+{
+	FEdModeLandscape* LandscapeEdMode = GetEditorMode();
+	ALandscape* Landscape = LandscapeEdMode ? LandscapeEdMode->GetLandscape() : nullptr;
+	if (Landscape && LandscapeEdMode && InEditLayer)
+	{
+		const int32 EditLayerIndex = Landscape->GetLayerIndex(InEditLayer->GetGuid());
+
+		// Put Persistent edit layer actions (Clearing weightmap/heightmap) in a separate category
+		const FName DataCategory = FName("Data");
+		OutContextMenuEntryMap.FindOrAdd(DataCategory).SectionLabel = LOCTEXT("LandscapeEditLayerContextMenuData.Heading", "Data");
+
+		if (LandscapeEdMode->DoesCurrentToolAffectEditLayers())
+		{
+			const ELandscapeToolTargetType CurrentTargetType = LandscapeEdMode->CurrentToolTarget.TargetType;
+			// Use ELandscapeToolTargetType::Weightmap as default mode: used for Target Layers and ELandscapeToolTargetType::Invalid (Paint tool with no target layers)
+			FText ClearEditLayerDataText = LOCTEXT("ClearAllTargetLayers", "Clear All Target Layers");
+
+			if (CurrentTargetType == ELandscapeToolTargetType::Heightmap)
+			{
+				ClearEditLayerDataText = LOCTEXT("ClearSculptLayer", "Clear Sculpt Layer");
+			}
+			else if (LandscapeEdMode->CurrentToolTarget.TargetType == ELandscapeToolTargetType::Visibility)
+			{
+				ClearEditLayerDataText = LOCTEXT("ClearVisibilityLayer", "Clear Visibility Layer");
+			}
+
+			FMenuEntryParams ClearEditLayerEntry;
+			ClearEditLayerEntry.LabelOverride = ClearEditLayerDataText;
+			ClearEditLayerEntry.ToolTipOverride =  MakeAttributeSPLambda(this, [this, EditLayerIndex, CurrentTargetType] { FText Reason; CanClearEditLayerData(EditLayerIndex, CurrentTargetType, Reason); return Reason; });
+			ClearEditLayerEntry.DirectActions = FUIAction(
+				FExecuteAction::CreateSP(this, &FLandscapeEditLayerContextMenuCustomization_Persistent::ClearEditLayerData, EditLayerIndex, CurrentTargetType),
+				FCanExecuteAction::CreateSPLambda(this, [this, EditLayerIndex, CurrentTargetType] { FText Reason; return CanClearEditLayerData(EditLayerIndex, CurrentTargetType, Reason); }));
+
+			OutContextMenuEntryMap.FindOrAdd(DataCategory).Entries.Add(ClearEditLayerEntry);
+		}
+	}
+}
+
+void FLandscapeEditLayerContextMenuCustomization_Persistent::ClearEditLayerData(int32 InEditLayerIndex, ELandscapeToolTargetType InClearMode)
+{
+	if (FEdModeLandscape* LandscapeEdMode = GetEditorMode())
+	{
+		ALandscape* Landscape = LandscapeEdMode->GetLandscape();
+		const ULandscapeEditLayerBase* EditLayer = LandscapeEdMode->GetEditLayerConst(InEditLayerIndex);
+		if (Landscape != nullptr && EditLayer != nullptr)
+		{
+			const FText ClearType = UEnum::GetDisplayValueAsText(InClearMode);
+
+			EAppReturnType::Type Result = FMessageDialog::Open(EAppMsgType::YesNo, FText::Format(LOCTEXT("Landscape_ClearEditLayer_Dialog", "Edit layer {0}'s {1} data will be completely cleared.  Continue?"), FText::FromName(EditLayer->GetName()), ClearType));
+			if (Result == EAppReturnType::Yes)
+			{
+				const FScopedTransaction Transaction(FText::Format(LOCTEXT("Landscape_Edit_Layers_Clear", "Clear Edit Layer {0} Data"), ClearType));
+				Landscape->ClearEditLayer(InEditLayerIndex, /*InComponents = */ nullptr, UE::Landscape::GetLandscapeToolTargetTypeAsFlags(InClearMode));
+				OnLayerSelectionChanged(InEditLayerIndex);
+
+				if (InClearMode == ELandscapeToolTargetType::Weightmap || InClearMode == ELandscapeToolTargetType::Visibility)
+				{
+					LandscapeEdMode->RequestUpdateLayerUsageInformation();
+				}
+			}
+		}
+	}
+}
+
+bool FLandscapeEditLayerContextMenuCustomization_Persistent::CanClearEditLayerData(int32 InEditLayerIndex, ELandscapeToolTargetType InClearMode, FText& OutToolTip) const
+{
+	if (FEdModeLandscape* LandscapeEdMode = GetEditorMode())
+	{
+		switch (InClearMode)
+		{
+		case ELandscapeToolTargetType::Invalid:
+			// In paint mode with no target layers 
+			OutToolTip = LOCTEXT("EditLayerContextMenu.Clear_WeightmapNoneToolTip", "No target layer data to clear for this edit layer.");
+			return false;
+		case ELandscapeToolTargetType::Heightmap:
+			OutToolTip = LOCTEXT("EditLayerContextMenu.Clear_HeightmapTooltip", "Clears all sculpting heightmap data for this edit layer.");
+			break;
+		case ELandscapeToolTargetType::Visibility:
+			OutToolTip = LOCTEXT("EditLayerContextMenu.Clear_VisibilityLayerTooltip", "Removes all holes painted in visibility for this edit layer.");
+			break;
+		case ELandscapeToolTargetType::Weightmap:
+			OutToolTip = LOCTEXT("EditLayerContextMenu.Clear_AllTargetLayerTooltip", "Clears all target layer weightmap data for this edit layer.");
+			break;
+		default:
+			check(false);
+			return false;
+		}
+
+		// InClearMode should always set the tooltip when valid
+		check(OutToolTip.ToString() != FString());
+		const ULandscapeEditLayerBase* EditLayer = LandscapeEdMode->GetEditLayerConst(InEditLayerIndex);
+		return EditLayer && LandscapeEdMode->GetLandscape() && LandscapeEdMode->CanEditLayer(&OutToolTip, EditLayer);
+	}
+	return false;
+}
+
+TSharedRef<IEditLayerCustomization> FLandscapeEditLayerContextMenuCustomization_Splines::MakeInstance()
+{
+	return MakeShareable(new FLandscapeEditLayerContextMenuCustomization_Splines);
+}
+
+void FLandscapeEditLayerContextMenuCustomization_Splines::CustomizeContextMenu(ULandscapeEditLayerBase* InEditLayer, FEditLayerCategoryToEntryMap& OutContextMenuEntryMap)
+{
+	FEdModeLandscape* LandscapeEdMode = GetEditorMode();
+	if (LandscapeEdMode && InEditLayer)
+	{
+		const FName SplinesCategory = FName("Splines");
+		OutContextMenuEntryMap.FindOrAdd(SplinesCategory).SectionLabel = LOCTEXT("LandscapeEditLayerContextMenuSplines.Heading", "Splines");
+
+		FMenuEntryParams ToggleSplineEntry;
+		ToggleSplineEntry.LabelOverride = LOCTEXT("LandscapeEditorToggleSplinesTool", "Toggle Splines Tool"),
+		ToggleSplineEntry.ToolTipOverride = LOCTEXT("LandscapeEditorToggleSplinesToolTip", "Toggles landscape splines mode");
+		ToggleSplineEntry.DirectActions = FUIAction(FExecuteAction::CreateLambda([LandscapeEdMode, InEditLayer]()
+		{
+			if (LandscapeEdMode && InEditLayer)
+			{
+				FScopedTransaction Transaction(LOCTEXT("Landscape_Layers_ToggleSplinesTool", "ToggleSplinesTool"));
+				LandscapeEdMode->ToggleSplinesTool(InEditLayer);
+			}
+		}));
+
+		OutContextMenuEntryMap.FindOrAdd(SplinesCategory).Entries.Add(ToggleSplineEntry);
+	}
+}
+
+#undef LOCTEXT_NAMESPACE
\ No newline at end of file
diff --git a/Engine/Source/Editor/LandscapeEditor/Private/LandscapeEditorDetailCustomization_EditLayers.h b/Engine/Source/Editor/LandscapeEditor/Private/LandscapeEditorDetailCustomization_EditLayers.h
new file mode 100644
--- /dev/null
+++ b/Engine/Source/Editor/LandscapeEditor/Private/LandscapeEditorDetailCustomization_EditLayers.h
@@ -0,0 +1,89 @@
+// Copyright Epic Games, Inc. All Rights Reserved.
+#pragma once
+
+#include "LandscapeEditLayerCustomization.h"
+
+class FMenuBuilder;
+class ALandscapeBlueprintBrushBase; 
+class ULandscapeEditLayerBase;
+class FEdModeLandscape;
+class FReply;
+struct FSlateBrush;
+struct FMenuEntryParams;
+enum class ELandscapeToolTargetType : uint8;
+
+/** Edit layer common editor functionality */
+class FLandscapeEditLayerCustomizationCommon
+{
+protected:
+	static FEdModeLandscape* GetEditorMode();
+
+	bool CanDeleteLayer(int32 InLayerIndex, FText& OutReason) const;
+	void DeleteLayer(int32 InLayerIndex);
+
+	bool CanCollapseLayer(int32 InLayerIndex, FText& OutReason) const;
+	void CollapseLayer(int32 InLayerIndex);
+
+	bool CanToggleVisibility(int32 InLayerIndex, FText& OutReason) const;
+	FReply OnToggleVisibility(int32 InLayerIndex);
+	const FSlateBrush* GetVisibilityBrushForLayer(int32 InLayerIndex) const;
+
+	void OnLayerSelectionChanged(int32 LayerIndex);
+	void ShowOnlySelectedLayer(int32 InLayerIndex);
+	void ShowAllLayers();
+
+	TOptional<float> GetLayerAlpha(int32 InLayerIndex) const;
+	float GetLayerAlphaMinValue() const;
+	bool CanSetLayerAlpha(int32 InLayerIndex, FText& OutReason) const;
+	void SetLayerAlpha(float InAlpha, int32 InLayerIndex, bool bCommit, int32 InSliderIndex);
+};
+
+/* Edit layer customizations for ULandscapeEditLayerBase */
+class FLandscapeEditLayerContextMenuCustomization_Base : public IEditLayerCustomization, public FLandscapeEditLayerCustomizationCommon
+{
+public:
+	static TSharedRef<IEditLayerCustomization> MakeInstance();
+
+	/** IEditLayerCustomization interface */
+	virtual void CustomizeContextMenu(ULandscapeEditLayerBase* InEditLayer, FEditLayerCategoryToEntryMap& OutContextMenuEntryMap) override;
+};
+
+/* Edit layer customizations for ULandscapeEditLayer */
+class FLandscapeEditLayerContextMenuCustomization_Layer : public IEditLayerCustomization, public FLandscapeEditLayerCustomizationCommon
+{
+public:
+	static TSharedRef<IEditLayerCustomization> MakeInstance();
+
+	/** IEditLayerCustomization interface */
+	virtual void CustomizeContextMenu(ULandscapeEditLayerBase* InEditLayer, FEditLayerCategoryToEntryMap& OutContextMenuEntryMap) override;
+
+private:
+	// Add Existing BP Brush (Brushes with no landscape actor assigned)
+	void FillUnassignedBrushMenu(FMenuBuilder& MenuBuilder, TArray<ALandscapeBlueprintBrushBase*> Brushes, int32 InLayerIndex);
+	void AssignBrushToEditLayer(ALandscapeBlueprintBrushBase* Brush, int32 InLayerIndex);
+};
+
+/* Edit layer customizations for ULandscapeEditLayerPersistent */
+class FLandscapeEditLayerContextMenuCustomization_Persistent : public IEditLayerCustomization, public FLandscapeEditLayerCustomizationCommon
+{
+public:
+	static TSharedRef<IEditLayerCustomization> MakeInstance();
+
+	/** IEditLayerCustomization interface */
+	virtual void CustomizeContextMenu(ULandscapeEditLayerBase* InEditLayer, FEditLayerCategoryToEntryMap& OutContextMenuEntryMap) override;
+
+private:
+	// Weightmap/Heightmap related functions
+	void ClearEditLayerData(int32 InEditLayerIndex, ELandscapeToolTargetType InClearMode);
+	bool CanClearEditLayerData(int32 InEditLayerIndex, ELandscapeToolTargetType InClearMode, FText& OutToolTip) const;
+};
+
+/* Edit layer customizations for ULandscapeEditLayerSplines */
+class FLandscapeEditLayerContextMenuCustomization_Splines : public IEditLayerCustomization, public FLandscapeEditLayerCustomizationCommon
+{
+public:
+	static TSharedRef<IEditLayerCustomization> MakeInstance();
+
+	/** IEditLayerCustomization interface */
+	virtual void CustomizeContextMenu(ULandscapeEditLayerBase* InEditLayer, FEditLayerCategoryToEntryMap& OutContextMenuEntryMap) override;
+};
\ No newline at end of file
diff --git a/Engine/Source/Editor/LandscapeEditor/Private/LandscapeEditorDetailCustomization_Layers.cpp b/Engine/Source/Editor/LandscapeEditor/Private/LandscapeEditorDetailCustomization_Layers.cpp
--- a/Engine/Source/Editor/LandscapeEditor/Private/LandscapeEditorDetailCustomization_Layers.cpp
+++ b/Engine/Source/Editor/LandscapeEditor/Private/LandscapeEditorDetailCustomization_Layers.cpp
@@ -25,6 +25,7 @@
 #include "LandscapeEditorDetailCustomization_LayersBrushStack.h" // FLandscapeBrushDragDropOp
 #include "LandscapeEditorModule.h"
 #include "LandscapeEditorObject.h"
+#include "LandscapeEditorUtils.h"
 #include "Landscape.h"
 #include "Styling/AppStyle.h"
 #include "DetailLayoutBuilder.h"
@@ -72,7 +73,7 @@ void FLandscapeEditorDetailCustomization_Layers::CustomizeDetails(IDetailLayoutB
 		&& (LandscapeEdMode->CurrentToolMode != nullptr)
 		&& (FName(LandscapeEdMode->CurrentTool->GetToolName()) != TEXT("Mask")))
 	{
-		LayerCategory.AddCustomBuilder(MakeShareable(new FLandscapeEditorCustomNodeBuilder_Layers(DetailBuilder.GetThumbnailPool().ToSharedRef())));
+		LayerCategory.AddCustomBuilder(MakeShareable(new FLandscapeEditorCustomNodeBuilder_Layers()));
 
 		LayerCategory.AddCustomRow(FText())
 			.Visibility(TAttribute<EVisibility>::Create(TAttribute<EVisibility>::FGetter::CreateLambda([]() { return ShouldShowLayersErrorMessageTip() ? EVisibility::Visible : EVisibility::Collapsed; })))
@@ -116,9 +117,8 @@ FEdModeLandscape* FLandscapeEditorCustomNodeBuilder_Layers::GetEditorMode()
 	return (FEdModeLandscape*)GLevelEditorModeTools().GetActiveMode(FBuiltinEditorModes::EM_Landscape);
 }
 
-FLandscapeEditorCustomNodeBuilder_Layers::FLandscapeEditorCustomNodeBuilder_Layers(TSharedRef<FAssetThumbnailPool> InThumbnailPool)
-	: ThumbnailPool(InThumbnailPool)
-	, CurrentSlider(INDEX_NONE)
+FLandscapeEditorCustomNodeBuilder_Layers::FLandscapeEditorCustomNodeBuilder_Layers()
+	: CurrentSlider(INDEX_NONE)
 {
 }
 
@@ -226,6 +226,7 @@ TSharedPtr<SWidget> FLandscapeEditorCustomNodeBuilder_Layers::GenerateRow(int32
 		.OnContextMenuOpening(this, &FLandscapeEditorCustomNodeBuilder_Layers::OnLayerContextMenuOpening, InLayerIndex)
 		.OnSelected(this, &FLandscapeEditorCustomNodeBuilder_Layers::OnLayerSelectionChanged, InLayerIndex)
 		.IsSelected(TAttribute<bool>::Create(TAttribute<bool>::FGetter::CreateSP(this, &FLandscapeEditorCustomNodeBuilder_Layers::IsLayerSelected, InLayerIndex)))
+		.OnDoubleClick(this, &FLandscapeEditorCustomNodeBuilder_Layers::OnLayerDoubleClick, InLayerIndex)
 		.Visibility(EVisibility::Visible)
 		[
 			SNew(SHorizontalBox)
@@ -251,8 +252,7 @@ TSharedPtr<SWidget> FLandscapeEditorCustomNodeBuilder_Layers::GenerateRow(int32
 						[
 							SNew(SImage)
 								.Image(this, &FLandscapeEditorCustomNodeBuilder_Layers::GetEditLayerIconBrush, InLayerIndex)
-								// TODO [jonathan.bard] : investigate why this doesn't work : the tooltip just doesn't show up
-								.ToolTip(this, &FLandscapeEditorCustomNodeBuilder_Layers::GetEditLayerTypeTooltip, InLayerIndex)
+								.ToolTipText(this, &FLandscapeEditorCustomNodeBuilder_Layers::GetEditLayerShortTooltipText, InLayerIndex)
 						]
 				]
 
@@ -358,8 +358,8 @@ TSharedPtr<SWidget> FLandscapeEditorCustomNodeBuilder_Layers::GenerateRow(int32
 								.Visibility(this, &FLandscapeEditorCustomNodeBuilder_Layers::GetLayerAlphaVisibility, InLayerIndex)
 								.ToolTipText(MakeAttributeSPLambda(this, [this, InLayerIndex] { FText Reason; CanSetLayerAlpha(InLayerIndex, Reason); return Reason; }))
 								.Value(this, &FLandscapeEditorCustomNodeBuilder_Layers::GetLayerAlpha, InLayerIndex)
-								.OnValueChanged_Lambda([this, InLayerIndex](float InValue) { SetLayerAlpha(InValue, InLayerIndex, false); })
-								.OnValueCommitted_Lambda([this, InLayerIndex](float InValue, ETextCommit::Type InCommitType) { SetLayerAlpha(InValue, InLayerIndex, true); })
+								.OnValueChanged_Lambda([this, InLayerIndex](float InValue) { SetLayerAlpha(InValue, InLayerIndex, false, CurrentSlider); })
+								.OnValueCommitted_Lambda([this, InLayerIndex](float InValue, ETextCommit::Type InCommitType) { SetLayerAlpha(InValue, InLayerIndex, true, CurrentSlider); })
 								.OnBeginSliderMovement_Lambda([this, InLayerIndex]()
 								{
 									CurrentSlider = InLayerIndex;
@@ -455,150 +455,39 @@ FSlateColor FLandscapeEditorCustomNodeBuilder_Layers::GetLayerTextColor(int32 In
 
 TSharedPtr<SWidget> FLandscapeEditorCustomNodeBuilder_Layers::OnLayerContextMenuOpening(int32 InLayerIndex)
 {
-	FEdModeLandscape* LandscapeEdMode = GetEditorMode();
-	ALandscape* Landscape = LandscapeEdMode ? LandscapeEdMode->GetLandscape() : nullptr;
-	if (Landscape && LandscapeEdMode)
+	FMenuBuilder MenuBuilder(true, nullptr);
+
+	if (const FEdModeLandscape* LandscapeEdMode = GetEditorMode())
 	{
-		const ULandscapeEditLayerBase* EditLayer = LandscapeEdMode->GetEditLayerConst(InLayerIndex);
-		FMenuBuilder MenuBuilder(true, nullptr);
-		MenuBuilder.BeginSection("LandscapeEditorLayerActions", LOCTEXT("LandscapeEditorLayerActions.Heading", "Edit Layers"));
+		if (ULandscapeEditLayerBase* EditLayer = LandscapeEdMode->GetEditLayer(InLayerIndex))
 		{
-			if (EditLayer)
+			// Instead of building the MenuBuilder directly, gather a map of entries so multiple customization classes 
+			// can place entries within the same category
+			FEditLayerCategoryToEntryMap ContextMenuEntries;
+
+			// Some context menu actions require slate components generated by this edit layer stack
 			{
+				FName EditLayerCategory = FName("Edit Layers");
+
 				// Rename Layer
-				FUIAction RenameLayerAction = FUIAction(
+				FMenuEntryParams RenameLayerEntry;
+				RenameLayerEntry.DirectActions = FUIAction(
 					FExecuteAction::CreateSP(this, &FLandscapeEditorCustomNodeBuilder_Layers::RenameLayer, InLayerIndex),
 					FCanExecuteAction::CreateSPLambda(this, [this, InLayerIndex] { FText Reason; return CanRenameLayer(InLayerIndex, Reason); }));
-				TAttribute<FText> CanRenameLayerTooltipText = MakeAttributeSPLambda(this, [this, InLayerIndex] { FText Reason; CanRenameLayer(InLayerIndex, Reason); return Reason; });
-				MenuBuilder.AddMenuEntry(LOCTEXT("RenameLayer", "Rename"), CanRenameLayerTooltipText, FSlateIcon(), RenameLayerAction);
-
-				// Delete Layer
-				FUIAction DeleteLayerAction = FUIAction(
-					FExecuteAction::CreateSP(this, &FLandscapeEditorCustomNodeBuilder_Layers::DeleteLayer, InLayerIndex),
-					FCanExecuteAction::CreateSPLambda(this, [this, InLayerIndex] { FText Reason; return CanDeleteLayer(InLayerIndex, Reason); }));
-				TAttribute<FText> DeleteLayerTooltipText = MakeAttributeSPLambda(this, [this, InLayerIndex] { FText Reason; CanDeleteLayer(InLayerIndex, Reason); return Reason; });
-				MenuBuilder.AddMenuEntry(LOCTEXT("DeleteLayer", "Delete"), DeleteLayerTooltipText, FSlateIcon(), DeleteLayerAction);
-
-				// Collapse Layer
-				FUIAction CollapseLayerAction = FUIAction(
-					FExecuteAction::CreateSP(this, &FLandscapeEditorCustomNodeBuilder_Layers::CollapseLayer, InLayerIndex),
-					FCanExecuteAction::CreateSPLambda(this, [this, InLayerIndex] { FText Reason; return CanCollapseLayer(InLayerIndex, Reason); }));
-				TAttribute<FText> CollapseLayerTooltipText = MakeAttributeSPLambda(this, [this, InLayerIndex] { FText Reason; CanCollapseLayer(InLayerIndex, Reason); return Reason; });
-				MenuBuilder.AddMenuEntry(LOCTEXT("CollapseLayer", "Collapse"), CollapseLayerTooltipText, FSlateIcon(), CollapseLayerAction);
-
-				if (LandscapeEdMode->DoesCurrentToolAffectEditLayers())
-				{
-					const ELandscapeToolTargetType CurrentTargetType = LandscapeEdMode->CurrentToolTarget.TargetType;
-					// Use ELandscapeToolTargetType::Weightmap as default mode: used for Target Layers and ELandscapeToolTargetType::Invalid (Paint tool with no target layers)
-					FText ClearEditLayerDataText = LOCTEXT("ClearAllTargetLayers", "Clear All Target Layers");
-
-					if (CurrentTargetType == ELandscapeToolTargetType::Heightmap)
-					{
-						ClearEditLayerDataText = LOCTEXT("ClearSculptLayer", "Clear Sculpt Layer");
-					}
-					else if (LandscapeEdMode->CurrentToolTarget.TargetType == ELandscapeToolTargetType::Visibility)
-					{
-						ClearEditLayerDataText = LOCTEXT("ClearVisibilityLayer", "Clear Visibility Layer");
-					}
-
-					FUIAction ClearEditLayerAction = FUIAction(
-						FExecuteAction::CreateSP(this, &FLandscapeEditorCustomNodeBuilder_Layers::ClearEditLayerData, InLayerIndex, CurrentTargetType),
-						FCanExecuteAction::CreateSPLambda(this, [this, InLayerIndex, CurrentTargetType] { FText Reason; return CanClearEditLayerData(InLayerIndex, CurrentTargetType, Reason); }));
-					TAttribute<FText> ClearEditLayerDataTooltip = MakeAttributeSPLambda(this, [this, InLayerIndex, CurrentTargetType] { FText Reason; CanClearEditLayerData(InLayerIndex, CurrentTargetType, Reason); return Reason; });
-
-					MenuBuilder.AddMenuEntry(ClearEditLayerDataText, ClearEditLayerDataTooltip, FSlateIcon(), ClearEditLayerAction);
-				}
+				RenameLayerEntry.ToolTipOverride = MakeAttributeSPLambda(this, [this, InLayerIndex] { FText Reason; CanRenameLayer(InLayerIndex, Reason); return Reason; });
+				RenameLayerEntry.LabelOverride = LOCTEXT("RenameLayer", "Rename");
 
-				// Add custom actions : 
-				if (!EditLayer->GetActions().IsEmpty())
-				{
-					MenuBuilder.AddMenuSeparator();
-
-					for (const ULandscapeEditLayerBase::FEditLayerAction& LayerAction : EditLayer->GetActions())
-					{
-						FUIAction LayerUIAction = FUIAction(
-							FExecuteAction::CreateSPLambda(this, [this, InLayerIndex, LayerAction] { ExecuteCustomLayerAction(InLayerIndex, LayerAction); }),
-							FCanExecuteAction::CreateSPLambda(this, [this, InLayerIndex, LayerAction] { FText Reason; return CanExecuteCustomLayerAction(InLayerIndex, LayerAction, Reason); }));
-						TAttribute<FText> LayerActionTooltipText = MakeAttributeSPLambda(this, [this, InLayerIndex, LayerAction] { FText Reason; CanExecuteCustomLayerAction(InLayerIndex, LayerAction, Reason); return Reason; });
-						MenuBuilder.AddMenuEntry(LayerAction.GetLabel(), LayerActionTooltipText, FSlateIcon(), LayerUIAction);
-					}
-				}
+				ContextMenuEntries.FindOrAdd(EditLayerCategory).Entries.Add(RenameLayerEntry);
 			}
-		}
-		MenuBuilder.EndSection();
-
-		MenuBuilder.BeginSection("LandscapeEditorLayerVisibility", LOCTEXT("LandscapeEditorLayerVisibility.Heading", "Visibility"));
-		{
-			if (EditLayer)
-			{
-				if (EditLayer->IsVisible())
-				{
-					// Hide Selected Layer
-					FUIAction HideSelectedLayerAction = FUIAction(FExecuteAction::CreateSPLambda(this, [this, InLayerIndex] { OnToggleVisibility(InLayerIndex); }),
-						FCanExecuteAction::CreateSPLambda(this, [this, InLayerIndex] { FText Reason; return CanToggleVisibility(InLayerIndex, Reason); }));
-					TAttribute<FText> HideSelectedLayerActionTooltipText = MakeAttributeSPLambda(this, [this, InLayerIndex] { FText Reason; CanToggleVisibility(InLayerIndex, Reason); return Reason; });
-					MenuBuilder.AddMenuEntry(LOCTEXT("HideSelected", "Hide Selected"), HideSelectedLayerActionTooltipText, FSlateIcon(), HideSelectedLayerAction);
-				}
-				else
-				{
-					// Show Selected Layer
-					FUIAction ShowSelectedLayerAction = FUIAction(FExecuteAction::CreateSPLambda(this, [this, InLayerIndex] { OnToggleVisibility(InLayerIndex); }),
-						FCanExecuteAction::CreateSPLambda(this, [this, InLayerIndex] { FText Reason; return CanToggleVisibility(InLayerIndex, Reason); }));
-					TAttribute<FText> ShowSelectedLayerActionTooltipText = MakeAttributeSPLambda(this, [this, InLayerIndex] { FText Reason; CanToggleVisibility(InLayerIndex, Reason); return Reason; });
-					MenuBuilder.AddMenuEntry(LOCTEXT("ShowSelected", "Show Selected"), ShowSelectedLayerActionTooltipText, FSlateIcon(), ShowSelectedLayerAction);
-				}
 
-				// Show Only Selected Layer
-				FUIAction ShowOnlySelectedLayerAction = FUIAction(FExecuteAction::CreateSPLambda(this, [this, InLayerIndex] { ShowOnlySelectedLayer(InLayerIndex); }));
-				MenuBuilder.AddMenuEntry(LOCTEXT("ShowOnlySelected", "Show Only Selected"), LOCTEXT("ShowOnlySelectedLayerTooltip", "Show Only Selected Layer"), FSlateIcon(), ShowOnlySelectedLayerAction);
-			}
-
-			// Show All Layers
-			FUIAction ShowAllLayersAction = FUIAction(FExecuteAction::CreateSP(this, &FLandscapeEditorCustomNodeBuilder_Layers::ShowAllLayers));
-			MenuBuilder.AddMenuEntry(LOCTEXT("ShowAllLayers", "Show All Layers"), LOCTEXT("ShowAllLayersTooltip", "Show All Layers"), FSlateIcon(), ShowAllLayersAction);
+			// Context menu actions that do not require slate state are defined in edit layer customization classes
+			ILandscapeEditorModule& LandscapeEditorModule = FModuleManager::GetModuleChecked<ILandscapeEditorModule>("LandscapeEditor");
+			EditLayerCustomizationClassInstances = LandscapeEditorModule.QueryCustomEditLayerLayoutRecursive(EditLayer->GetClass());
+			LandscapeEditorModule.ApplyEditLayerContextMenuCustomizations(EditLayer, EditLayerCustomizationClassInstances, ContextMenuEntries);
+			LandscapeEditorUtils::BuildContextMenuFromCategoryEntryMap(ContextMenuEntries, MenuBuilder);
 		}
-		MenuBuilder.EndSection();
-
-		// Add unassigned brush to layer
-		const TArray<ALandscapeBlueprintBrushBase*>& Brushes = LandscapeEdMode->GetBrushList();
-		TArray<ALandscapeBlueprintBrushBase*> FilteredBrushes = Brushes.FilterByPredicate([](ALandscapeBlueprintBrushBase* Brush) { return Brush->GetOwningLandscape() == nullptr; });
-		
-		MenuBuilder.BeginSection("LandscapeEditorBrushActions", LOCTEXT("LandscapeEditorBrushActions.Heading", "Brushes"));
-		{
-			if (EditLayer)
-			{
-				// If there are no unassigned brushes or the edit layer does not support brushes, show a disabled state instead of hiding the entire section
-				if (FilteredBrushes.Num() > 0 && EditLayer->SupportsBlueprintBrushes())
-				{
-					MenuBuilder.AddSubMenu(
-						LOCTEXT("LandscapeEditorBrushAddSubMenu", "Assign Existing Brush"),
-						LOCTEXT("LandscapeEditorBrushAddSubMenuToolTip", "To modify the terrain, brushes need to be assigned to a landscape actor. Add the brush to this edit layer"),
-						FNewMenuDelegate::CreateSP(this, &FLandscapeEditorCustomNodeBuilder_Layers::FillUnassignedBrushMenu, FilteredBrushes, InLayerIndex),
-						false,
-						FSlateIcon()
-					);
-				}
-				else
-				{
-					FUIAction NoOpAction = FUIAction(FExecuteAction::CreateLambda([]() {}),
-						FCanExecuteAction::CreateLambda([]() { return false; }));
-					const FText NoOpToolTip = EditLayer->SupportsBlueprintBrushes() ? LOCTEXT("LandscapeEditorBrushAllBrushActorsAssigned", "All Blueprint Brush actors are assigned to a landscape edit layer")
-						: FText::Format(LOCTEXT("LandscapeEditorBrushUnsupported", "This layer's type ({0}) doesn't support blueprint brushes."), EditLayer->GetClass()->GetDisplayNameText());
-
-					MenuBuilder.AddMenuEntry(
-						LOCTEXT("LandscapeEditorBrushNone", "None"),
-						NoOpToolTip,
-						FSlateIcon(),
-						NoOpAction
-					);
-				}
-			}
-		}
-		MenuBuilder.EndSection();
-		
-		return MenuBuilder.MakeWidget();
 	}
-	return nullptr;
+	return MenuBuilder.MakeWidget();
 }
 
 bool FLandscapeEditorCustomNodeBuilder_Layers::CanRenameLayer(int32 InLayerIndex, FText& OutReason) const
@@ -629,276 +518,6 @@ void FLandscapeEditorCustomNodeBuilder_Layers::RenameLayer(int32 InLayerIndex)
 	}
 }
 
-void FLandscapeEditorCustomNodeBuilder_Layers::ClearEditLayerData(int32 InEditLayerIndex, ELandscapeToolTargetType InClearMode)
-{
-	if (FEdModeLandscape* LandscapeEdMode = GetEditorMode())
-	{
-		ALandscape* Landscape = LandscapeEdMode->GetLandscape();
-		const ULandscapeEditLayerBase* EditLayer = LandscapeEdMode->GetEditLayerConst(InEditLayerIndex);
-		if (Landscape != nullptr && EditLayer != nullptr)
-		{
-			const FText ClearType = UEnum::GetDisplayValueAsText(InClearMode);
-
-			EAppReturnType::Type Result = FMessageDialog::Open(EAppMsgType::YesNo, FText::Format(LOCTEXT("Landscape_ClearEditLayer_Dialog", "Edit layer {0}'s {1} data will be completely cleared.  Continue?"), FText::FromName(EditLayer->GetName()), ClearType));
-			if (Result == EAppReturnType::Yes)
-			{
-				const FScopedTransaction Transaction(FText::Format(LOCTEXT("Landscape_Edit_Layers_Clear", "Clear Edit Layer {0} Data"), ClearType));
-				Landscape->ClearEditLayer(InEditLayerIndex, /*InComponents = */ nullptr, UE::Landscape::GetLandscapeToolTargetTypeAsFlags(InClearMode));
-				OnLayerSelectionChanged(InEditLayerIndex);
-
-				if (InClearMode == ELandscapeToolTargetType::Weightmap || InClearMode == ELandscapeToolTargetType::Visibility)
-				{
-					LandscapeEdMode->RequestUpdateLayerUsageInformation();
-				}
-			}
-		}
-	}
-}
-
-bool FLandscapeEditorCustomNodeBuilder_Layers::CanClearEditLayerData(int32 InEditLayerIndex, ELandscapeToolTargetType InClearMode, FText& OutToolTip) const
-{
-	if (FEdModeLandscape* LandscapeEdMode = GetEditorMode())
-	{
-		switch (InClearMode)
-		{
-			case ELandscapeToolTargetType::Invalid:
-				// In paint mode with no target layers 
-				OutToolTip = LOCTEXT("EditLayerContextMenu.Clear_WeightmapNoneToolTip", "No target layer data to clear for this edit layer.");
-				return false;
-			case ELandscapeToolTargetType::Heightmap:
-				OutToolTip = LOCTEXT("EditLayerContextMenu.Clear_HeightmapTooltip", "Clears all sculpting heightmap data for this edit layer.");
-				break;
-			case ELandscapeToolTargetType::Visibility:
-				OutToolTip = LOCTEXT("EditLayerContextMenu.Clear_VisibilityLayerTooltip", "Removes all holes painted in visibility for this edit layer.");
-				break;
-			case ELandscapeToolTargetType::Weightmap:
-				OutToolTip = LOCTEXT("EditLayerContextMenu.Clear_AllTargetLayerTooltip", "Clears all target layer weightmap data for this edit layer.");
-				break;
-			default:
-				check(false);
-				return false;
-		}
-
-		// InClearMode should always set the tooltip when valid
-		check(OutToolTip.ToString() != FString());
-		const ULandscapeEditLayerBase* EditLayer = LandscapeEdMode->GetEditLayerConst(InEditLayerIndex);
-		return EditLayer && LandscapeEdMode->GetLandscape() && LandscapeEdMode->CanEditLayer(&OutToolTip, EditLayer);
-	}
-	return false;
-}
-
-bool FLandscapeEditorCustomNodeBuilder_Layers::CanDeleteLayer(int32 InLayerIndex, FText& OutReason) const
-{
-	FEdModeLandscape* LandscapeEdMode = GetEditorMode();
-	ALandscape* Landscape = LandscapeEdMode ? LandscapeEdMode->GetLandscape() : nullptr;
-	if (!Landscape || Landscape->GetEditLayersConst().Num() <= 1)
-	{
-		OutReason = LOCTEXT("DeleteLayer_CantDeleteLastLayer", "The last layer cannot be deleted");
-		return false;
-	}
-
-	const ULandscapeEditLayerBase* EditLayer = LandscapeEdMode->GetEditLayerConst(InLayerIndex);
-	if (EditLayer == nullptr)
-	{
-		return false;
-	}
-
-	if (EditLayer->IsLocked())
-	{
-		OutReason = LOCTEXT("DeleteLayer_CantDeleteLocked", "Cannot delete a locked edit layer");
-		return false;
-	}
-
-	OutReason = LOCTEXT("DeleteLayer_CanDelete", "Delete the edit layer");
-	return true;
-}
-
-void FLandscapeEditorCustomNodeBuilder_Layers::DeleteLayer(int32 InLayerIndex)
-{
-	FEdModeLandscape* LandscapeEdMode = GetEditorMode();
-	ALandscape* Landscape = LandscapeEdMode ? LandscapeEdMode->GetLandscape() : nullptr;
-	if (Landscape && Landscape->GetEditLayersConst().Num() > 1)
-	{
-		const ULandscapeEditLayerBase* EditLayer = LandscapeEdMode->GetEditLayerConst(InLayerIndex);
-		if (EditLayer)
-		{
-			EAppReturnType::Type Result = FMessageDialog::Open(EAppMsgType::YesNo, FText::Format(LOCTEXT("Landscape_DeleteLayer_Message", "The layer {0} will be deleted.  Continue?"), FText::FromName(EditLayer->GetName())));
-			if (Result == EAppReturnType::Yes)
-			{
-				const FScopedTransaction Transaction(LOCTEXT("Landscape_Layers_Delete", "Delete Layer"));
-				// Delete layer will update the selected edit layer index
-				Landscape->DeleteLayer(InLayerIndex);
-				LandscapeEdMode->UpdateTargetList();
-				LandscapeEdMode->RefreshDetailPanel();
-			}
-		}
-	}
-}
-
-bool FLandscapeEditorCustomNodeBuilder_Layers::CanCollapseLayer(int32 InLayerIndex, FText& OutReason) const
-{
-	FEdModeLandscape* LandscapeEdMode = GetEditorMode();
-	ALandscape* Landscape = LandscapeEdMode ? LandscapeEdMode->GetLandscape() : nullptr;
-	if (!Landscape || Landscape->GetEditLayersConst().Num() <= 1)
-	{
-		OutReason = LOCTEXT("Landscape_CollapseLayer_Reason_NotEnoughLayersToCollapse", "Not enough layers to do collapse");
-		return false;
-	}
-
-	if (InLayerIndex < 1)
-	{
-		OutReason = LOCTEXT("Landscape_CollapseLayer_Reason_CantCollapseBaseLayer", "Cannot collapse the first layer");
-		return false;
-	}
-	
-	const FLandscapeLayer* TopLayer = Landscape->GetLayerConst(InLayerIndex);
-	const FLandscapeLayer* BottomLayer = Landscape->GetLayerConst(InLayerIndex - 1);
-	if ((TopLayer == nullptr) || (BottomLayer == nullptr))
-	{
-		return false;
-	}
-
-	const ULandscapeEditLayerBase* TopEditLayer = LandscapeEdMode->GetEditLayerConst(InLayerIndex);
-	const ULandscapeEditLayerBase* BottomEditLayer = LandscapeEdMode->GetEditLayerConst(InLayerIndex - 1);
-
-	check((TopEditLayer != nullptr) && (BottomEditLayer != nullptr));
-
-	if (!TopEditLayer->SupportsCollapsingTo())
-	{
-		OutReason = FText::Format(LOCTEXT("Landscape_CollapseLayer_Reason_TopLayerDoesntSupportCollapsing", "Cannot collapse layer '{0}' onto layer '{1}'. The type of layer '{0}' ({2}) doesn't support collapsing to another one"), 
-			FText::FromName(TopEditLayer->GetName()), FText::FromName(BottomEditLayer->GetName()), TopEditLayer->GetClass()->GetDisplayNameText());
-		return false;
-	}
-
-	if (!BottomEditLayer->SupportsBeingCollapsedAway())
-	{
-		OutReason = FText::Format(LOCTEXT("Landscape_CollapseLayer_Reason_BottomLayerDoesntSupportCollapsing", "Cannot collapse layer '{0}' onto layer '{1}'. The type of layer '{1}' ({2}) doesn't support being collapsed away"),
-			FText::FromName(TopEditLayer->GetName()), FText::FromName(BottomEditLayer->GetName()), BottomEditLayer->GetClass()->GetDisplayNameText());
-		return false;
-	}
-
-	if (TopEditLayer->IsLocked())
-	{
-		OutReason = FText::Format(LOCTEXT("Landscape_CollapseLayer_Reason_TopLayerIsLocked", "Cannot collapse layer '{0}' onto layer '{1}'. Layer '{0}' will be deleted in the operation but it is currently locked"),
-			FText::FromName(TopEditLayer->GetName()), FText::FromName(BottomEditLayer->GetName()));
-		return false;
-	}
-
-	if (BottomEditLayer->IsLocked())
-	{
-		OutReason = FText::Format(LOCTEXT("Landscape_CollapseLayer_Reason_BottomLayerIsLocked", "Cannot collapse layer '{0}' onto layer '{1}'. Destination layer '{1}' is currently locked"),
-			FText::FromName(TopEditLayer->GetName()), FText::FromName(BottomEditLayer->GetName()));
-		return false;
-	}
-
-	// Can't collapse on layer that has a Brush because result will change...
-	if (TopLayer->Brushes.Num() > 0)
-	{
-		OutReason = FText::Format(LOCTEXT("Landscape_CollapseLayer_Reason_TopLayerHasBrush", "Cannot collapse layer '{0}' onto layer '{1}'. Layer '{0}' contains brush(es)"),
-			FText::FromName(TopEditLayer->GetName()), FText::FromName(BottomEditLayer->GetName()));
-		return false;
-	}
-
-	if (BottomLayer->Brushes.Num() > 0)
-	{
-		OutReason = FText::Format(LOCTEXT("Landscape_CollapseLayer_Reason_BottomLayerHasBrush", "Cannot collapse layer '{0}' onto layer '{1}'. Layer '{1}' contains brush(es)"),
-			FText::FromName(TopEditLayer->GetName()), FText::FromName(BottomEditLayer->GetName()));
-		return false;
-	}
-	
-	OutReason = FText::Format(LOCTEXT("Landscape_CollapseLayer_Reason_Collapse", "Collapse layer '{0}' onto layer '{1}'"),
-		FText::FromName(TopEditLayer->GetName()), FText::FromName(BottomEditLayer->GetName()));
-	return true;
-}
-
-void FLandscapeEditorCustomNodeBuilder_Layers::CollapseLayer(int32 InLayerIndex)
-{
-	FText Reason;
-	check(CanCollapseLayer(InLayerIndex, Reason));
-
-	if (FEdModeLandscape* LandscapeEdMode = GetEditorMode())
-	{
-		ALandscape* Landscape = LandscapeEdMode->GetLandscape();
-		const ULandscapeEditLayerBase* Layer = LandscapeEdMode->GetEditLayerConst(InLayerIndex);
-		const ULandscapeEditLayerBase* BaseLayer = LandscapeEdMode->GetEditLayerConst(InLayerIndex - 1);
-
-		if (Landscape && Layer && BaseLayer)
-		{
-			EAppReturnType::Type Result = FMessageDialog::Open(EAppMsgType::YesNo, FText::Format(LOCTEXT("Landscape_CollapseLayer_Message", "The layer {0} will be collapsed into layer {1}.  Continue?"), FText::FromName(Layer->GetName()), FText::FromName(BaseLayer->GetName())));
-			if (Result == EAppReturnType::Yes)
-			{
-				const FScopedTransaction Transaction(LOCTEXT("Landscape_Layers_Collapse", "Collapse Layer"));
-				Landscape->CollapseLayer(InLayerIndex);
-				OnLayerSelectionChanged(InLayerIndex - 1);
-				LandscapeEdMode->RefreshDetailPanel();
-			}
-		}
-	}
-}
-
-bool FLandscapeEditorCustomNodeBuilder_Layers::CanExecuteCustomLayerAction(int32 InLayerIndex, const ULandscapeEditLayerBase::FEditLayerAction& InCustomLayerAction, FText& OutReason) const
-{
-	FEdModeLandscape* LandscapeEdMode = GetEditorMode();
-	ALandscape* Landscape = LandscapeEdMode ? LandscapeEdMode->GetLandscape() : nullptr;
-	if (Landscape == nullptr)
-	{
-		OutReason = LOCTEXT("Landscape_CannotExecuteCustomLayerAction_Invalid", "Invalid landscape");
-		return false;
-	}
-
-	const ULandscapeEditLayerBase* EditLayer = LandscapeEdMode->GetEditLayerConst(InLayerIndex);
-	if (EditLayer == nullptr)
-	{
-		return false;
-	}
-
-	ULandscapeEditLayerBase::FEditLayerAction::FExecuteParams ExecuteParams(EditLayer, Landscape);
-	return InCustomLayerAction.GetCanExecuteDelegate().Execute(ExecuteParams, OutReason);
-}
-
-void FLandscapeEditorCustomNodeBuilder_Layers::ExecuteCustomLayerAction(int32 InLayerIndex, const ULandscapeEditLayerBase::FEditLayerAction& InCustomLayerAction)
-{
-	FEdModeLandscape* LandscapeEdMode = GetEditorMode();
-	if (ALandscape* Landscape = LandscapeEdMode ? LandscapeEdMode->GetLandscape() : nullptr)
-	{
-		const ULandscapeEditLayerBase* EditLayer = LandscapeEdMode->GetEditLayerConst(InLayerIndex);
-		if (EditLayer != nullptr)
-		{
-			const FScopedTransaction Transaction(InCustomLayerAction.GetLabel());
-			ULandscapeEditLayerBase::FEditLayerAction::FExecuteParams ExecuteParams(EditLayer, Landscape);
-			ULandscapeEditLayerBase::FEditLayerAction::FExecuteResult Result = InCustomLayerAction.GetExecuteDelegate().Execute(ExecuteParams);
-			if (!Result.bSuccess)
-			{
-				// Indicate to the user that the action failed : 
-				FMessageDialog::Open(EAppMsgType::Ok, Result.Reason);
-			}
-		}
-	}
-}
-
-void FLandscapeEditorCustomNodeBuilder_Layers::ShowOnlySelectedLayer(int32 InLayerIndex)
-{
-	FEdModeLandscape* LandscapeEdMode = GetEditorMode();
-	ALandscape* Landscape = LandscapeEdMode ? LandscapeEdMode->GetLandscape() : nullptr;
-	if (Landscape)
-	{
-		const FScopedTransaction Transaction(LOCTEXT("ShowOnlySelectedLayer", "Show Only Selected Layer"));
-		Landscape->ShowOnlySelectedLayer(InLayerIndex);
-		OnLayerSelectionChanged(InLayerIndex);
-	}
-}
-
-void FLandscapeEditorCustomNodeBuilder_Layers::ShowAllLayers()
-{
-	FEdModeLandscape* LandscapeEdMode = GetEditorMode();
-	ALandscape* Landscape = LandscapeEdMode ? LandscapeEdMode->GetLandscape() : nullptr;
-	if (Landscape)
-	{
-		const FScopedTransaction Transaction(LOCTEXT("ShowAllLayers", "Show All Layers"));
-		Landscape->ShowAllLayers();
-	}
-}
-
 TSubclassOf<ULandscapeEditLayerBase> FLandscapeEditorCustomNodeBuilder_Layers::PickEditLayerClass() const
 {
 	class FLandscapeEditLayerClassFilter : public IClassViewerFilter
@@ -1005,139 +624,14 @@ bool FLandscapeEditorCustomNodeBuilder_Layers::CanCreateLayer(FText& OutReason)
 	return true;
 }
 
-void FLandscapeEditorCustomNodeBuilder_Layers::OnLayerSelectionChanged(int32 InLayerIndex)
+void FLandscapeEditorCustomNodeBuilder_Layers::OnLayerDoubleClick(int32 InLayerIndex) const
 {
-	if (FEdModeLandscape* LandscapeEdMode = GetEditorMode(); LandscapeEdMode && LandscapeEdMode->GetSelectedEditLayerIndex() != InLayerIndex)
+	if (FEdModeLandscape* LandscapeEdMode = GetEditorMode(); LandscapeEdMode && LandscapeEdMode->CurrentToolMode)
 	{
 		FScopedTransaction Transaction(LOCTEXT("Landscape_Layers_SetCurrentLayer", "Set Current Layer"));
-		LandscapeEdMode->SetSelectedEditLayer(InLayerIndex);
-		LandscapeEdMode->UpdateTargetList();
-	}
-}
-
-TOptional<float> FLandscapeEditorCustomNodeBuilder_Layers::GetLayerAlpha(int32 InLayerIndex) const
-{
-	if (FEdModeLandscape* LandscapeEdMode = GetEditorMode())
-	{
 		const ULandscapeEditLayerBase* EditLayer = LandscapeEdMode->GetEditLayerConst(InLayerIndex);
-		check(EditLayer != nullptr);
-
-		return EditLayer->GetAlphaForTargetType(LandscapeEdMode->GetLandscapeToolTargetType());
+		LandscapeEdMode->ToggleSplinesTool(EditLayer);
 	}
-
-	return 1.0f;
-}
-
-float FLandscapeEditorCustomNodeBuilder_Layers::GetLayerAlphaMinValue() const
-{
-	if (const FEdModeLandscape* LandscapeEdMode = GetEditorMode(); LandscapeEdMode && LandscapeEdMode->GetLandscapeToolTargetType() == ELandscapeToolTargetType::Heightmap)
-	{
-		return -1.0f;
-	}
-	else
-	{
-		return 0.0f;
-	}
-}
-
-bool FLandscapeEditorCustomNodeBuilder_Layers::CanSetLayerAlpha(int32 InLayerIndex, FText& OutReason) const
-{
-	const FEdModeLandscape* LandscapeEdMode = GetEditorMode();
-
-	if (LandscapeEdMode == nullptr)
-	{
-		return false;
-	}
-
-	const ULandscapeEditLayerBase* EditLayer = LandscapeEdMode->GetEditLayerConst(InLayerIndex);
-	if (EditLayer == nullptr)
-	{
-		return false;
-	}
-
-	if (!EditLayer->SupportsAlphaForTargetType(LandscapeEdMode->GetLandscapeToolTargetType()))
-	{
-		OutReason = FText::Format(LOCTEXT("SetLayerAlpha_LayerDoesntSupportAlpha", "Cannot change alpha : the type of layer {0} ({1}) doesn't support alpha"),
-			FText::FromName(EditLayer->GetName()), EditLayer->GetClass()->GetDisplayNameText());
-		return false;
-	}
-
-
-	if (EditLayer->IsLocked())
-	{
-		OutReason = LOCTEXT("SetLayerAlpha_LayerIsLocked", "Cannot change the alpha of a locked edit layer");
-		return false;
-	}
-
-	OutReason = LOCTEXT("SetLayerAlpha_CanSet", "Set the edit layer's alpha");
-	return true;
-}
-
-void FLandscapeEditorCustomNodeBuilder_Layers::SetLayerAlpha(float InAlpha, int32 InLayerIndex, bool bCommit)
-{
-	FEdModeLandscape* LandscapeEdMode = GetEditorMode();
-	if (LandscapeEdMode)
-	{
-		ULandscapeEditLayerBase* EditLayer = LandscapeEdMode->GetEditLayer(InLayerIndex);
-		check(EditLayer != nullptr);
-
-		FScopedTransaction Transaction(LOCTEXT("Landscape_Layers_SetAlpha", "Set Layer Alpha"), CurrentSlider == INDEX_NONE && bCommit);
-		// Set Value when using slider or when committing text
-		EditLayer->SetAlphaForTargetType(LandscapeEdMode->GetLandscapeToolTargetType(), InAlpha, /*bInModify = */true, bCommit ? EPropertyChangeType::ValueSet : EPropertyChangeType::Interactive);
-	}
-}
-
-bool FLandscapeEditorCustomNodeBuilder_Layers::CanToggleVisibility(int32 InLayerIndex, FText& OutReason) const
-{
-	FEdModeLandscape* LandscapeEdMode = GetEditorMode();
-
-	const ULandscapeEditLayerBase* EditLayer = LandscapeEdMode->GetEditLayerConst(InLayerIndex);
-	if (EditLayer == nullptr)
-	{
-		return false;
-	}
-
-	if (EditLayer->IsLocked())
-	{
-		OutReason = LOCTEXT("ToggleVisibility_CantToggleLocked", "Cannot change the visibility of a locked edit layer");
-		return false;
-	}
-
-	OutReason = LOCTEXT("ToggleVisibility_CanToggle", "Toggle the visibility of the edit layer");
-	return true;
-}
-
-FReply FLandscapeEditorCustomNodeBuilder_Layers::OnToggleVisibility(int32 InLayerIndex)
-{
-	if (FEdModeLandscape* LandscapeEdMode = GetEditorMode())
-	{
-		const FScopedTransaction Transaction(LOCTEXT("Landscape_Layers_SetVisibility", "Set Layer Visibility"));
-		
-		ULandscapeEditLayerBase* EditLayer = LandscapeEdMode->GetEditLayer(InLayerIndex);
-		check(EditLayer != nullptr);
-		EditLayer->SetVisible(!EditLayer->IsVisible(), /*bInModify= */true);
-
-		if (EditLayer->IsVisible())
-		{
-			OnLayerSelectionChanged(InLayerIndex);
-		}
-	}
-	return FReply::Handled();
-}
-
-const FSlateBrush* FLandscapeEditorCustomNodeBuilder_Layers::GetVisibilityBrushForLayer(int32 InLayerIndex) const
-{
-	bool bIsVisible = false;
-
-	if (FEdModeLandscape* LandscapeEdMode = GetEditorMode())
-	{
-		const ULandscapeEditLayerBase* EditLayer = LandscapeEdMode->GetEditLayerConst(InLayerIndex);
-		check(EditLayer != nullptr);
-
-		bIsVisible = EditLayer->IsVisible();
-	}
-	
-	return bIsVisible ? FAppStyle::GetBrush("Level.VisibleIcon16x") : FAppStyle::GetBrush("Level.NotVisibleIcon16x");
 }
 
 void FLandscapeEditorCustomNodeBuilder_Layers::OnSetInspectedDetailsToEditLayer(int32 InLayerIndex) const
@@ -1172,11 +666,17 @@ EVisibility FLandscapeEditorCustomNodeBuilder_Layers::GetLayerAlphaVisibility(in
 	return bIsVisible ? EVisibility::Visible : EVisibility::Hidden;
 }
 
-TSharedPtr<IToolTip> FLandscapeEditorCustomNodeBuilder_Layers::GetEditLayerTypeTooltip(int32 InLayerIndex) const
+FText FLandscapeEditorCustomNodeBuilder_Layers::GetEditLayerShortTooltipText(int32 InEditLayerIndex) const
 {
-	FEdModeLandscape* LandscapeEdMode = GetEditorMode();
-	const ULandscapeEditLayerBase* EditLayer = (LandscapeEdMode != nullptr) ? LandscapeEdMode->GetEditLayerConst(InLayerIndex) : nullptr;
-	return (EditLayer != nullptr) ? TSharedPtr<IToolTip>(FEditorClassUtils::GetTooltip(EditLayer->GetClass())) : nullptr;
+	if (FEdModeLandscape* LandscapeEdMode = GetEditorMode())
+	{
+		const ULandscapeEditLayerBase* EditLayer = LandscapeEdMode->GetEditLayerConst(InEditLayerIndex);
+
+		// Edit layer tooltips defined in ShortToolTip field
+		return EditLayer ? EditLayer->GetClass()->GetToolTipText(/*bShortToolTip =*/ true) : FText();
+	}
+
+	return FText();
 }
 
 const FSlateBrush* FLandscapeEditorCustomNodeBuilder_Layers::GetLockBrushForLayer(int32 InLayerIndex) const
@@ -1191,27 +691,6 @@ const FSlateBrush* FLandscapeEditorCustomNodeBuilder_Layers::GetLockBrushForLaye
 	return bIsLocked ? FAppStyle::GetBrush(TEXT("PropertyWindow.Locked")) : FAppStyle::GetBrush(TEXT("PropertyWindow.Unlocked"));
 }
 
-void FLandscapeEditorCustomNodeBuilder_Layers::FillUnassignedBrushMenu(FMenuBuilder& MenuBuilder, TArray<ALandscapeBlueprintBrushBase*> Brushes, int32 InLayerIndex)
-{
-	for (ALandscapeBlueprintBrushBase* Brush : Brushes)
-	{
-		FUIAction AddAction = FUIAction(FExecuteAction::CreateSP(this, &FLandscapeEditorCustomNodeBuilder_Layers::AssignBrushToEditLayer, Brush, InLayerIndex));
-		MenuBuilder.AddMenuEntry(FText::FromString(Brush->GetActorLabel()), FText(), FSlateIcon(), AddAction);
-	}
-}
-
-void FLandscapeEditorCustomNodeBuilder_Layers::AssignBrushToEditLayer(ALandscapeBlueprintBrushBase* Brush, const int32 InLayerIndex)
-{
-	const FScopedTransaction Transaction(LOCTEXT("LandscapeBrushAddToCurrentLayerTransaction", "Add brush to edit layer"));
-	if (FEdModeLandscape* LandscapeEdMode = GetEditorMode())
-	{
-		if (ALandscape* Landscape = LandscapeEdMode->GetLandscape())
-		{
-			Landscape->AddBrushToLayer(InLayerIndex, Brush);
-		}
-	}
-}
-
 const FSlateBrush* FLandscapeEditorCustomNodeBuilder_Layers::GetEditLayerIconBrush(int32 InLayerIndex) const
 {
 	FEdModeLandscape* LandscapeEdMode = GetEditorMode();
diff --git a/Engine/Source/Editor/LandscapeEditor/Private/LandscapeEditorDetailCustomization_Layers.h b/Engine/Source/Editor/LandscapeEditor/Private/LandscapeEditorDetailCustomization_Layers.h
--- a/Engine/Source/Editor/LandscapeEditor/Private/LandscapeEditorDetailCustomization_Layers.h
+++ b/Engine/Source/Editor/LandscapeEditor/Private/LandscapeEditorDetailCustomization_Layers.h
@@ -5,6 +5,7 @@
 #include "CoreMinimal.h"
 #include "Misc/Attribute.h"
 #include "LandscapeEditorDetailCustomization_Base.h"
+#include "LandscapeEditorDetailCustomization_EditLayers.h" // FLandscapeEditLayerCustomizationCommon
 #include "LandscapeEdMode.h"
 #include "LandscapeEditLayer.h"
 #include "Layout/Visibility.h"
@@ -30,7 +31,7 @@ class ULandscapeLayerInfoObject;
 class FMenuBuilder;
 
 /**
- * Slate widgets customizer for the layers list in the Landscape Editor
+ * Slate widgets customizer for the edit layer stack in the Landscape Editor
  */
 
 class FLandscapeEditorDetailCustomization_Layers : public FLandscapeEditorDetailCustomization_Base
@@ -47,10 +48,10 @@ protected:
 	static FText GetLayersErrorMessageText();
 };
 
-class FLandscapeEditorCustomNodeBuilder_Layers : public IDetailCustomNodeBuilder, public TSharedFromThis<FLandscapeEditorCustomNodeBuilder_Layers>
+class FLandscapeEditorCustomNodeBuilder_Layers : public IDetailCustomNodeBuilder, public TSharedFromThis<FLandscapeEditorCustomNodeBuilder_Layers>, public FLandscapeEditLayerCustomizationCommon
 {
 public:
-	FLandscapeEditorCustomNodeBuilder_Layers(TSharedRef<FAssetThumbnailPool> ThumbnailPool);
+	FLandscapeEditorCustomNodeBuilder_Layers();
 	~FLandscapeEditorCustomNodeBuilder_Layers();
 
 	virtual void SetOnRebuildChildren( FSimpleDelegate InOnRegenerateChildren ) override;
@@ -62,8 +63,6 @@ public:
 	virtual FName GetName() const override { return "Layers"; }
 
 protected:
-	TSharedRef<FAssetThumbnailPool> ThumbnailPool;
-
 	static class FEdModeLandscape* GetEditorMode();
 
 	TSharedPtr<SWidget> GenerateRow(int32 InLayerIndex);
@@ -75,7 +74,7 @@ protected:
 	FReply HandleAcceptDrop(FDragDropEvent const& DragDropEvent, SDragAndDropVerticalBox::EItemDropZone DropZone, int32 SlotIndex, SVerticalBox::FSlot* Slot);
 
 	bool IsLayerSelected(int32 LayerIndex) const;
-	void OnLayerSelectionChanged(int32 LayerIndex);
+	void OnLayerDoubleClick(int32 LayerIndex) const;
 	TSharedPtr<SWidget> OnLayerContextMenuOpening(int32 InLayerIndex);
 
 	FText GetNumLayersText() const;
@@ -85,55 +84,33 @@ protected:
 
 	const FSlateBrush* GetEditLayerIconBrush(int32 InLayerIndex) const;
 
-	void ClearEditLayerData(int32 InEditLayerIndex, ELandscapeToolTargetType InClearMode);
-	bool CanClearEditLayerData(int32 InEditLayerIndex, ELandscapeToolTargetType InClearMode, FText& OutToolTip) const;
-
 	bool CanRenameLayerTo(const FText& NewText, FText& OutErrorMessage, int32 InLayerIndex);
 	bool CanRenameLayer(int32 InLayerIndex, FText& OutReason) const;
 	void RenameLayer(int32 InLayerIndex);
 
-	bool CanDeleteLayer(int32 InLayerIndex, FText& OutReason) const;
-	void DeleteLayer(int32 InLayerIndex);
-
-	bool CanCollapseLayer(int32 InLayerIndex, FText& OutReason) const;
-	void CollapseLayer(int32 InLayerIndex);
-
-	bool CanExecuteCustomLayerAction(int32 InLayerIndex, const ULandscapeEditLayerBase::FEditLayerAction& InCustomLayerAction, FText& OutReason) const;
-	void ExecuteCustomLayerAction(int32 InLayerIndex, const ULandscapeEditLayerBase::FEditLayerAction& InCustomLayerAction);
-
-	void ShowOnlySelectedLayer(int32 InLayerIndex);
-	void ShowAllLayers();
 	void SetLayerName(const FText& InText, ETextCommit::Type InCommitType, int32 InLayerIndex);
 	FText GetLayerText(int32 InLayerIndex) const;
 	FSlateColor GetLayerTextColor(int32 InLayerIndex) const;
 	FText GetLayerDisplayName(int32 InLayerIndex) const;
 	EVisibility GetLayerAlphaVisibility(int32 InLayerIndex) const;
-	TSharedPtr<IToolTip> GetEditLayerTypeTooltip(int32 InLayerIndex) const;
-	TSubclassOf<ULandscapeEditLayerBase> PickEditLayerClass() const;
+	FText GetEditLayerShortTooltipText(int32 InEditLayerIndex) const;
+	FReply OnEditLayerIconClicked(int32 InEditLayerIndex) const;
 
-	TOptional<float> GetLayerAlpha(int32 InLayerIndex) const;
-	float GetLayerAlphaMinValue() const;
-	bool CanSetLayerAlpha(int32 InLayerIndex, FText& OutReason) const;
-	void SetLayerAlpha(float InAlpha, int32 InLayerIndex, bool bCommit);
-	
-	bool CanToggleVisibility(int32 InLayerIndex, FText& OutReason) const;
-	FReply OnToggleVisibility(int32 InLayerIndex);
-	const FSlateBrush* GetVisibilityBrushForLayer(int32 InLayerIndex) const;
+	TSubclassOf<ULandscapeEditLayerBase> PickEditLayerClass() const;
 	
 	void OnSetInspectedDetailsToEditLayer(int32 InLayerIndex) const;
 
 	FReply OnToggleLock(int32 InLayerIndex);
 	const FSlateBrush* GetLockBrushForLayer(int32 InLayerIndex) const;
 
-	// Add Existing BP Brush (Brushes with no landscape actor assigned)
-	void FillUnassignedBrushMenu(FMenuBuilder& MenuBuilder, TArray<ALandscapeBlueprintBrushBase*> Brushes, int32 InLayerIndex);
-	void AssignBrushToEditLayer(ALandscapeBlueprintBrushBase* Brush, int32 InLayerIndex);
-
 private:
 
 	/** Widgets for displaying and editing the layer name */
 	TArray< TSharedPtr< SInlineEditableTextBlock > > InlineTextBlocks;
 
+	/** The edit layer customization class instances active in this view */
+	TArray<TSharedRef<IEditLayerCustomization>> EditLayerCustomizationClassInstances;
+
 	int32 CurrentSlider;
 };
 
diff --git a/Engine/Source/Editor/LandscapeEditor/Private/LandscapeEditorDetailCustomization_TargetLayers.cpp b/Engine/Source/Editor/LandscapeEditor/Private/LandscapeEditorDetailCustomization_TargetLayers.cpp
--- a/Engine/Source/Editor/LandscapeEditor/Private/LandscapeEditorDetailCustomization_TargetLayers.cpp
+++ b/Engine/Source/Editor/LandscapeEditor/Private/LandscapeEditorDetailCustomization_TargetLayers.cpp
@@ -1157,7 +1157,6 @@ TSharedPtr<SWidget> FLandscapeEditorCustomNodeBuilder_TargetLayers::GenerateRow(
 			.OnContextMenuOpening(this, &FLandscapeEditorCustomNodeBuilder_TargetLayers::OnTargetLayerContextMenuOpening, Target, static_cast<int32>(INDEX_NONE))
 			.OnSelected_Static(&FLandscapeEditorCustomNodeBuilder_TargetLayers::OnTargetSelectionChanged, Target)
 			.IsSelected_Static(&FLandscapeEditorCustomNodeBuilder_TargetLayers::GetTargetLayerIsSelected, Target)
-			.OnDoubleClick_Static(&FLandscapeEditorCustomNodeBuilder_TargetLayers::OnTargetLayerDoubleClicked, Target)
 			.Visibility(this, &FLandscapeEditorCustomNodeBuilder_TargetLayers::ShouldShowLayer, Target)
 			[
 				SNew(SHorizontalBox)
@@ -1567,18 +1566,6 @@ void FLandscapeEditorCustomNodeBuilder_TargetLayers::OnTargetSelectionChanged(co
 	}
 }
 
-void FLandscapeEditorCustomNodeBuilder_TargetLayers::OnTargetLayerDoubleClicked(const TSharedRef<FLandscapeTargetListInfo> Target)
-{
-	if (FEdModeLandscape* LandscapeEdMode = GetEditorMode())
-	{
-		LandscapeEdMode->CurrentToolTarget.TargetType = Target->TargetType;
-		if ((Target->TargetType == ELandscapeToolTargetType::Weightmap) && Target->LayerInfoObj.IsValid())
-		{
-			GEditor->GetEditorSubsystem<UAssetEditorSubsystem>()->OpenEditorForAsset(Target->LayerInfoObj.Get());
-		}
-	}
-}
-
 TSharedPtr<SWidget> FLandscapeEditorCustomNodeBuilder_TargetLayers::OnTargetLayerContextMenuOpening(const TSharedRef<FLandscapeTargetListInfo> Target, const int32 InLayerIndex)
 {
 	FMenuBuilder MenuBuilder(true, nullptr);
@@ -2246,6 +2233,7 @@ void SLandscapeEditorSelectableBorder::Construct(const FArguments& InArgs)
 	OnContextMenuOpening = InArgs._OnContextMenuOpening;
 	OnSelected = InArgs._OnSelected;
 	IsSelected = InArgs._IsSelected;
+	OnDoubleClick = InArgs._OnDoubleClick;
 }
 
 FReply SLandscapeEditorSelectableBorder::OnMouseButtonUp(const FGeometry& MyGeometry, const FPointerEvent& MouseEvent)
@@ -2277,7 +2265,7 @@ FReply SLandscapeEditorSelectableBorder::OnMouseButtonUp(const FGeometry& MyGeom
 	return FReply::Unhandled();
 }
 
-FReply SLandscapeEditorSelectableBorder::OnMouseDoubleClick(const FGeometry& MyGeometry, const FPointerEvent& MouseEvent)
+FReply SLandscapeEditorSelectableBorder::OnMouseButtonDoubleClick(const FGeometry& MyGeometry, const FPointerEvent& MouseEvent)
 {
 	if (MyGeometry.IsUnderLocation(MouseEvent.GetScreenSpacePosition()))
 	{
diff --git a/Engine/Source/Editor/LandscapeEditor/Private/LandscapeEditorDetailCustomization_TargetLayers.h b/Engine/Source/Editor/LandscapeEditor/Private/LandscapeEditorDetailCustomization_TargetLayers.h
--- a/Engine/Source/Editor/LandscapeEditor/Private/LandscapeEditorDetailCustomization_TargetLayers.h
+++ b/Engine/Source/Editor/LandscapeEditor/Private/LandscapeEditorDetailCustomization_TargetLayers.h
@@ -93,7 +93,6 @@ protected:
 
 	static bool GetTargetLayerIsSelected(const TSharedRef<FLandscapeTargetListInfo> Target);
 	static void OnTargetSelectionChanged(const TSharedRef<FLandscapeTargetListInfo> Target);
-	static void OnTargetLayerDoubleClicked(const TSharedRef<FLandscapeTargetListInfo> Target);
 	TSharedPtr<SWidget> OnTargetLayerContextMenuOpening(const TSharedRef<FLandscapeTargetListInfo> Target, const int32 InLayerIndex);
 	void OnRenameLayer(const int32 InLayerIndex);
 	static void OnExportLayer(const TSharedRef<FLandscapeTargetListInfo> Target);
@@ -182,8 +181,8 @@ public:
 
 	void Construct(const FArguments& InArgs);
 
-	virtual FReply OnMouseButtonUp(const FGeometry& MyGeometry, const FPointerEvent& MouseEvent);
-	virtual FReply OnMouseDoubleClick(const FGeometry& MyGeometry, const FPointerEvent& MouseEvent);
+	virtual FReply OnMouseButtonUp(const FGeometry& MyGeometry, const FPointerEvent& MouseEvent) override;
+	virtual FReply OnMouseButtonDoubleClick(const FGeometry& MyGeometry, const FPointerEvent& MouseEvent) override;
 	const FSlateBrush* GetBorder() const;
 
 protected:
diff --git a/Engine/Source/Editor/LandscapeEditor/Private/LandscapeEditorModule.cpp b/Engine/Source/Editor/LandscapeEditor/Private/LandscapeEditorModule.cpp
--- a/Engine/Source/Editor/LandscapeEditor/Private/LandscapeEditorModule.cpp
+++ b/Engine/Source/Editor/LandscapeEditor/Private/LandscapeEditorModule.cpp
@@ -36,6 +36,8 @@
 #include "LandscapeEditorDetailCustomization_CopyPaste.h"
 #include "LandscapeEditorDetailCustomization_ImportLayers.h"
 #include "LandscapeEditorDetailCustomization_TargetLayers.h"
+#include "LandscapeEditorDetailCustomization_EditLayers.h"
+#include "LandscapeEditLayerCustomization.h"
 #include "LandscapeSplineDetails.h"
 #include "LandscapeGrassTypeDetails.h"
 #include "LandscapeProxyUIDetails.h"
@@ -107,6 +109,12 @@ public:
 		PropertyModule.RegisterCustomClassLayout("LandscapeProxy", FOnGetDetailCustomizationInstance::CreateStatic(&FLandscapeProxyUIDetails::MakeInstance));
 		PropertyModule.RegisterCustomClassLayout("LandscapeGrassType", FOnGetDetailCustomizationInstance::CreateStatic(&FLandscapeGrassTypeDetails::MakeInstance));
 
+		// register edit layer customizations (Ex. custom Context Menu)
+		RegisterCustomEditLayerClassLayout("LandscapeEditLayerBase", FOnGetLandscapeEditLayerCustomizationInstance::CreateStatic(&FLandscapeEditLayerContextMenuCustomization_Base::MakeInstance));
+		RegisterCustomEditLayerClassLayout("LandscapeEditLayer", FOnGetLandscapeEditLayerCustomizationInstance::CreateStatic(&FLandscapeEditLayerContextMenuCustomization_Layer::MakeInstance));
+		RegisterCustomEditLayerClassLayout("LandscapeEditLayerPersistent", FOnGetLandscapeEditLayerCustomizationInstance::CreateStatic(&FLandscapeEditLayerContextMenuCustomization_Persistent::MakeInstance));
+		RegisterCustomEditLayerClassLayout("LandscapeEditLayerSplines", FOnGetLandscapeEditLayerCustomizationInstance::CreateStatic(&FLandscapeEditLayerContextMenuCustomization_Splines::MakeInstance));
+
 		// register property sections
 		RegisterPropertySectionMappings();
 		
@@ -200,6 +208,12 @@ public:
 		PropertyModule.UnregisterCustomClassLayout("LandscapeProxy");
 		PropertyModule.UnregisterCustomClassLayout("LandscapeGrassType");
 
+		// unregister edit layer customizations (Ex. custom Context Menu)
+		UnregisterCustomEditLayerClassLayout("LandscapeEditLayerBase");
+		UnregisterCustomEditLayerClassLayout("LandscapeEditLayer");
+		UnregisterCustomEditLayerClassLayout("LandscapeEditLayerPersistent");
+		UnregisterCustomEditLayerClassLayout("LandscapeEditLayerSplines");
+
 		// unregister property sections
 		UnregisterPropertySectionMappings();
 
@@ -302,6 +316,86 @@ public:
 		}
 	}
 
+	void RegisterCustomEditLayerClassLayout(FName ClassName, FOnGetLandscapeEditLayerCustomizationInstance EditLayerLayoutDelegate) override
+	{
+		if (ClassName != NAME_None)
+		{
+			FOnGetLandscapeEditLayerCustomizationInstance Callback;
+			Callback = EditLayerLayoutDelegate;
+
+			ClassNameToEditLayerCustomizationMap.Add(ClassName, Callback);
+		}
+	}
+
+	void UnregisterCustomEditLayerClassLayout(FName ClassName)
+	{
+		if (ClassName.IsValid() && (ClassName != NAME_None))
+		{
+			ClassNameToEditLayerCustomizationMap.Remove(ClassName);
+		}
+	}
+
+	TArray<TSharedRef<IEditLayerCustomization>> QueryCustomEditLayerLayoutRecursive(UStruct* InClass)
+	{
+		TArray<TSharedRef<IEditLayerCustomization>> EditLayerCustomizationClassInstances;
+		if (InClass == nullptr)
+		{
+			return EditLayerCustomizationClassInstances;
+		}
+
+		TSet<UStruct*> ClassesToQuery;
+	
+		// Ensure that the base class and its parents are always queried
+		ClassesToQuery.Add(InClass);
+
+		UStruct* ParentStruct = InClass->GetSuperStruct();
+		while (ParentStruct && ParentStruct->IsA(UClass::StaticClass()) && !ClassesToQuery.Contains(ParentStruct))
+		{
+			ClassesToQuery.Add(ParentStruct);
+			ParentStruct = ParentStruct->GetSuperStruct();
+		}
+
+		// Query class inheritance chain
+		for (auto ParentIt = ClassesToQuery.CreateConstIterator(); ParentIt; ++ParentIt)
+		{
+			if (TSharedPtr<IEditLayerCustomization> CustomizationInstance = QueryCustomEditLayerLayoutForClass(*ParentIt))
+			{
+				EditLayerCustomizationClassInstances.Add(CustomizationInstance.ToSharedRef());
+			} 
+		}
+
+		return EditLayerCustomizationClassInstances;
+	}
+
+	TSharedPtr<IEditLayerCustomization> QueryCustomEditLayerLayoutForClass(UStruct* InClass)
+	{
+		const FOnGetLandscapeEditLayerCustomizationInstance* CustomizationDelegate = ClassNameToEditLayerCustomizationMap.Find(InClass->GetFName());
+
+		if (CustomizationDelegate && CustomizationDelegate->IsBound())
+		{
+			// Create a new instance of the customization for the current edit layer class
+			TSharedPtr<IEditLayerCustomization> CustomizationInstance = CustomizationDelegate->Execute();
+
+			// Caller's responsibility to handle the instance lifetime
+			return CustomizationInstance;
+		}
+
+		return nullptr;
+	}
+
+	void ApplyEditLayerContextMenuCustomizations(ULandscapeEditLayerBase* InEditLayer, const TArray<TSharedRef<IEditLayerCustomization>>& InCustomizations, FEditLayerCategoryToEntryMap& OutContextMenuEntryMap)
+	{
+		if (InEditLayer == nullptr || InCustomizations.IsEmpty())
+		{
+			return;
+		}
+
+		for (TSharedRef<IEditLayerCustomization> CustomizationInstance : InCustomizations)
+		{
+			CustomizationInstance->CustomizeContextMenu(InEditLayer, OutContextMenuEntryMap);
+		}
+	}
+
 	virtual const TCHAR* GetHeightmapImportDialogTypeString() const override;
 	virtual const TCHAR* GetWeightmapImportDialogTypeString() const override;
 
@@ -339,6 +433,9 @@ protected:
 
 	TMultiMap<FName, FName> RegisteredPropertySections;
 	TArray<FName> PropertyTypesToUnregisterOnShutdown;
+
+	/** A mapping of class names to edit layer delegates, called when querying for custom edit layer layouts */
+	FCustomEditLayerLayoutNameMap ClassNameToEditLayerCustomizationMap;
 };
 
 IMPLEMENT_MODULE(FLandscapeEditorModule, LandscapeEditor);
diff --git a/Engine/Source/Editor/LandscapeEditor/Private/LandscapeEditorUtils.cpp b/Engine/Source/Editor/LandscapeEditor/Private/LandscapeEditorUtils.cpp
--- a/Engine/Source/Editor/LandscapeEditor/Private/LandscapeEditorUtils.cpp
+++ b/Engine/Source/Editor/LandscapeEditor/Private/LandscapeEditorUtils.cpp
@@ -22,6 +22,8 @@
 #include "AssetRegistry/AssetData.h"
 #include "AssetRegistry/AssetRegistryModule.h"
 #include "Algo/LevenshteinDistance.h"
+#include "Framework/MultiBox/MultiBoxBuilder.h"
+#include "LandscapeEditLayerCustomization.h"
 
 namespace LandscapeEditorUtils
 {
@@ -144,6 +146,28 @@ namespace LandscapeEditorUtils
 		return BestLayerInfoObj;
 	}
 
+	void BuildContextMenuFromCategoryEntryMap(const FEditLayerCategoryToEntryMap& InMap, FMenuBuilder& OutMenuBuilder)
+	{
+		for (const TPair<FName, FEditLayerMenuBlock>& CategoryPair : InMap)
+		{
+			const FName& CategoryName = CategoryPair.Key;
+			const TArray<FMenuEntryParams>& MenuEntries = CategoryPair.Value.Entries;
+
+			if (MenuEntries.IsEmpty())
+			{
+				continue;
+			}
+
+			check(!CategoryPair.Value.SectionLabel.IsEmpty()); 
+
+			OutMenuBuilder.BeginSection(CategoryName, CategoryPair.Value.SectionLabel);
+			for (const FMenuEntryParams& Entry : MenuEntries)
+			{
+				OutMenuBuilder.AddMenuEntry(Entry);
+			}
+			OutMenuBuilder.EndSection();
+		}
+	}
 }
 
 
diff --git a/Engine/Source/Editor/LandscapeEditor/Public/LandscapeEditLayerCustomization.h b/Engine/Source/Editor/LandscapeEditor/Public/LandscapeEditLayerCustomization.h
new file mode 100644
--- /dev/null
+++ b/Engine/Source/Editor/LandscapeEditor/Public/LandscapeEditLayerCustomization.h
@@ -0,0 +1,30 @@
+// Copyright Epic Games, Inc. All Rights Reserved.
+
+#pragma once
+
+#include "CoreMinimal.h"
+#include "Containers/Map.h"
+class ULandscapeEditLayerBase;
+class FMenuBuilder;
+struct FMenuEntryParams;
+
+struct FEditLayerMenuBlock
+{
+	TArray<FMenuEntryParams> Entries;
+	FText SectionLabel;
+};
+
+/* Map of context menu Category Names to their entries */
+using FEditLayerCategoryToEntryMap = TMap<FName, FEditLayerMenuBlock>;
+
+/**
+ * Interface for Edit Layer classes that expand on customizations not covered by IDetailsView (Context Menu, Dialog, etc.)
+ */
+class IEditLayerCustomization : public TSharedFromThis<IEditLayerCustomization>
+{
+public:
+	virtual ~IEditLayerCustomization() {}
+
+	/** Called when building edit layer context menus */
+	virtual void CustomizeContextMenu(ULandscapeEditLayerBase* InEditLayer, FEditLayerCategoryToEntryMap& OutContextMenuEntryMap) = 0;
+};
diff --git a/Engine/Source/Editor/LandscapeEditor/Public/LandscapeEditorModule.h b/Engine/Source/Editor/LandscapeEditor/Public/LandscapeEditorModule.h
--- a/Engine/Source/Editor/LandscapeEditor/Public/LandscapeEditorModule.h
+++ b/Engine/Source/Editor/LandscapeEditor/Public/LandscapeEditorModule.h
@@ -7,6 +7,7 @@
 #include "Templates/SharedPointer.h"
 #include "Templates/UnrealTypeTraits.h"
 #include "Templates/EnableIf.h"
+#include "Containers/Map.h" 
 
 template<class T>
 class ILandscapeFileFormat;
@@ -16,6 +17,15 @@ using ILandscapeWeightmapFileFormat = ILandscapeFileFormat<uint8>;
 
 class FUICommandList;
 class FLandscapeImageFileCache;
+class IEditLayerCustomization;
+class ULandscapeEditLayerBase;
+struct FEditLayerMenuBlock;
+using FEditLayerCategoryToEntryMap = TMap<FName, FEditLayerMenuBlock>;
+
+/** Delegate called to get customizations for a specific edit layer class */
+DECLARE_DELEGATE_RetVal(TSharedRef<IEditLayerCustomization>, FOnGetLandscapeEditLayerCustomizationInstance);
+
+using FCustomEditLayerLayoutNameMap = TMap<FName, FOnGetLandscapeEditLayerCustomizationInstance>;
 
 /**
  * LandscapeEditor module interface
@@ -35,6 +45,15 @@ public:
 	virtual const TCHAR* GetHeightmapExportDialogTypeString() const = 0;
 	virtual const TCHAR* GetWeightmapExportDialogTypeString() const = 0;
 
+	// Edit layer customizations
+	virtual void RegisterCustomEditLayerClassLayout(FName ClassName, FOnGetLandscapeEditLayerCustomizationInstance EditLayerLayoutDelegate) = 0;
+	virtual void UnregisterCustomEditLayerClassLayout(FName ClassName) = 0;
+
+	virtual TArray<TSharedRef<IEditLayerCustomization>> QueryCustomEditLayerLayoutRecursive(UStruct* InClass) = 0;
+	virtual TSharedPtr<IEditLayerCustomization> QueryCustomEditLayerLayoutForClass(UStruct* InClass) = 0;
+
+	virtual void ApplyEditLayerContextMenuCustomizations(ULandscapeEditLayerBase* InEditLayer, const TArray<TSharedRef<IEditLayerCustomization>>& InCustomizations, FEditLayerCategoryToEntryMap& OutContextMenuEntryMap) = 0;
+
 	// Gets the heightmap/weightmap format associated with a given extension (null if no plugin is registered for this extension)
 	virtual const ILandscapeHeightmapFileFormat* GetHeightmapFormatByExtension(const TCHAR* Extension) const = 0;
 	virtual const ILandscapeWeightmapFileFormat* GetWeightmapFormatByExtension(const TCHAR* Extension) const = 0;
diff --git a/Engine/Source/Editor/LandscapeEditor/Public/LandscapeEditorUtils.h b/Engine/Source/Editor/LandscapeEditor/Public/LandscapeEditorUtils.h
--- a/Engine/Source/Editor/LandscapeEditor/Public/LandscapeEditorUtils.h
+++ b/Engine/Source/Editor/LandscapeEditor/Public/LandscapeEditorUtils.h
@@ -4,10 +4,14 @@
 
 #include "CoreMinimal.h"
 #include "FileHelpers.h"
+#include "Containers/Map.h"
 
 class ALandscapeProxy;
 class ALandscapeStreamingProxy;
 class ULandscapeLayerInfoObject;
+class FMenuBuilder;
+struct FEditLayerMenuBlock;
+using FEditLayerCategoryToEntryMap = TMap<FName, FEditLayerMenuBlock>;
 
 namespace LandscapeEditorUtils
 {
@@ -29,4 +33,7 @@ namespace LandscapeEditorUtils
 
 	TArray<FAssetData> GetLandscapeTargetLayerInfoAssets();
 	TOptional<FAssetData> FindLandscapeTargetLayerInfoAsset(const FName& LayerName, const FString& TargetLayerAssetPackagePath);
+
+	// Edit layer context menu customization
+	void BuildContextMenuFromCategoryEntryMap(const FEditLayerCategoryToEntryMap& InMap, FMenuBuilder& OutMenuBuilder);
 }
diff --git a/Engine/Source/Runtime/Landscape/Classes/LandscapeEditLayer.h b/Engine/Source/Runtime/Landscape/Classes/LandscapeEditLayer.h
--- a/Engine/Source/Runtime/Landscape/Classes/LandscapeEditLayer.h
+++ b/Engine/Source/Runtime/Landscape/Classes/LandscapeEditLayer.h
@@ -47,13 +47,12 @@ class ULandscapeEditLayerBase : public UObject
 	GENERATED_BODY()
 
 public:
-	// Little wrapper class to let edit layers expose some actions generically : 
-	class FEditLayerAction
+	class UE_DEPRECATED(5.7, "Expose custom actions to the editor with IEditLayerCustomization and RegisterCustomEditLayerClassLayout") FEditLayerAction
 	{
 		FEditLayerAction() = delete;
 
 	public:
-		class FExecuteParams
+		class UE_DEPRECATED(5.7, "Expose custom actions to the editor with IEditLayerCustomization and RegisterCustomEditLayerClassLayout") FExecuteParams
 		{
 			FExecuteParams() = delete;
 
@@ -84,7 +83,7 @@ public:
 			ALandscape* Landscape = nullptr;
 		};
 
-		struct FExecuteResult
+		struct UE_DEPRECATED(5.7, "Expose custom actions to the editor with IEditLayerCustomization and RegisterCustomEditLayerClassLayout") FExecuteResult
 		{
 			FExecuteResult() = default;
 			FExecuteResult(bool bInSuccess, const FText& InReason = FText())
@@ -281,7 +280,10 @@ public:
 	/**
 	* @return a list of actions that can be triggered on this edit layer
 	*/
+	UE_DEPRECATED(5.7, "Register custom context menu actions using RegisterCustomEditLayerClassLayout from ILandscapeEditorModule")
+PRAGMA_DISABLE_DEPRECATION_WARNINGS
 	virtual TArray<FEditLayerAction> GetActions() const 
+PRAGMA_ENABLE_DEPRECATION_WARNINGS
 	{ 
 		return {}; 
 	}
@@ -425,7 +427,7 @@ public:
 /** 
 * This is the standard type of edit layer. It can be manually authored (sculpted, painted, etc.) in the landscape editor 
 */
-UCLASS(MinimalAPI)
+UCLASS(MinimalAPI, meta = (ShortToolTip = "Standard edit layer"))
 class ULandscapeEditLayer : public ULandscapeEditLayerPersistent
 {
 	GENERATED_BODY()
@@ -464,7 +466,7 @@ public:
 /** 
 * Procedural edit layer that lets the user manipulate its content using landscape splines (Splines tool in the Manage panel) 
 */
-UCLASS(MinimalAPI)
+UCLASS(MinimalAPI, meta = (ShortToolTip = "Special edit layer for landscape splines"))
 class ULandscapeEditLayerSplines : public ULandscapeEditLayerPersistent
 {
 	GENERATED_BODY()
@@ -477,7 +479,6 @@ public:
 	virtual bool SupportsMultiple() const override { return false; } // only one layer of this type is allowed
 	virtual bool SupportsBeingCollapsedAway() const override { return false; } // this is a procedural and therefore cannot be collapsed 
 	virtual FString GetDefaultName() const override { return TEXT("Splines"); }
-	virtual TArray<FEditLayerAction> GetActions() const override;
 #if WITH_EDITOR
 	virtual ELandscapeBlendMode GetBlendMode() const override { return ELandscapeBlendMode::LSBM_AlphaBlend; };
 	bool SupportsAlphaForTargetType(ELandscapeToolTargetType InType) const override { return false; };
diff --git a/Engine/Source/Runtime/Landscape/Private/LandscapeEditLayer.cpp b/Engine/Source/Runtime/Landscape/Private/LandscapeEditLayer.cpp
--- a/Engine/Source/Runtime/Landscape/Private/LandscapeEditLayer.cpp
+++ b/Engine/Source/Runtime/Landscape/Private/LandscapeEditLayer.cpp
@@ -434,35 +434,4 @@ bool ULandscapeEditLayerSplines::SupportsTargetType(ELandscapeToolTargetType InT
 	return (InType == ELandscapeToolTargetType::Heightmap) || (InType == ELandscapeToolTargetType::Weightmap) || (InType == ELandscapeToolTargetType::Visibility);
 }
 
-TArray<ULandscapeEditLayerSplines::FEditLayerAction> ULandscapeEditLayerSplines::GetActions() const
-{
-	TArray<FEditLayerAction> Actions;
-
-#if WITH_EDITOR
-	// Register an "Update Splines" action :
-	Actions.Add(FEditLayerAction(
-		LOCTEXT("LandscapeEditLayerSplines_UpdateSplines", "Update Splines"),
-		FEditLayerAction::FExecuteDelegate::CreateWeakLambda(this, [](const FEditLayerAction::FExecuteParams& InParams)
-		{
-			InParams.GetLandscape()->UpdateAllLandscapeSplines(FGuid(), /*bForceUpdate =*/true);
-			return FEditLayerAction::FExecuteResult(/*bInSuccess = */true);
-		}),
-		FEditLayerAction::FCanExecuteDelegate::CreateWeakLambda(this, [](const FEditLayerAction::FExecuteParams& InParams, FText& OutReason)
-		{
-			const ULandscapeEditLayerBase* EditLayer = InParams.GetEditLayer();
-			check(EditLayer != nullptr);
-			if (EditLayer->IsLocked())
-			{
-				OutReason = FText::Format(LOCTEXT("LandscapeEditLayerSplines_CannotUpdateSplinesOnLockedLayer", "Cannot update splines on layer '{0}' : the layer is currently locked"), FText::FromName(EditLayer->GetName()));
-				return false;
-			}
-
-			OutReason = LOCTEXT("LandscapeEditLayerSplines_UpdateSplines_Tooltip", "Update Landscape Splines");
-			return true;
-		})));
-#endif // WITH_EDITOR
-
-	return Actions;
-}
-
 #undef LOCTEXT_NAMESPACE
