Landscape Editor - Adding BugHawk feedback to edit layer customization UI change 45797695

diff --git a/Engine/Source/Editor/LandscapeEditor/Private/LandscapeEditorDetailCustomization_EditLayers.cpp b/Engine/Source/Editor/LandscapeEditor/Private/LandscapeEditorDetailCustomization_EditLayers.cpp
--- a/Engine/Source/Editor/LandscapeEditor/Private/LandscapeEditorDetailCustomization_EditLayers.cpp
+++ b/Engine/Source/Editor/LandscapeEditor/Private/LandscapeEditorDetailCustomization_EditLayers.cpp
@@ -170,21 +170,23 @@ void FLandscapeEditLayerCustomizationCommon::CollapseLayer(int32 InLayerIndex)
 
 bool FLandscapeEditLayerCustomizationCommon::CanToggleVisibility(int32 InLayerIndex, FText& OutReason) const
 {
-	FEdModeLandscape* LandscapeEdMode = GetEditorMode();
-
-	const ULandscapeEditLayerBase* EditLayer = LandscapeEdMode->GetEditLayerConst(InLayerIndex);
-	if (EditLayer == nullptr)
+	if (FEdModeLandscape* LandscapeEdMode = GetEditorMode())
 	{
-		return false;
-	}
+		const ULandscapeEditLayerBase* EditLayer = LandscapeEdMode->GetEditLayerConst(InLayerIndex);
+		if (EditLayer == nullptr)
+		{
+			return false;
+		}
 
-	if (EditLayer->IsLocked())
-	{
-		OutReason = LOCTEXT("ToggleVisibility_CantToggleLocked", "Cannot change the visibility of a locked edit layer");
-		return false;
+		if (EditLayer->IsLocked())
+		{
+			OutReason = LOCTEXT("ToggleVisibility_CantToggleLocked", "Cannot change the visibility of a locked edit layer");
+			return false;
+		}
+
+		OutReason = LOCTEXT("ToggleVisibility_CanToggle", "Toggle the visibility of the edit layer");
 	}
 
-	OutReason = LOCTEXT("ToggleVisibility_CanToggle", "Toggle the visibility of the edit layer");
 	return true;
 }
 
diff --git a/Engine/Source/Runtime/Landscape/Classes/LandscapeEditLayer.h b/Engine/Source/Runtime/Landscape/Classes/LandscapeEditLayer.h
--- a/Engine/Source/Runtime/Landscape/Classes/LandscapeEditLayer.h
+++ b/Engine/Source/Runtime/Landscape/Classes/LandscapeEditLayer.h
@@ -277,9 +277,6 @@ public:
 
 #endif // WITH_EDITOR
 
-	/**
-	* @return a list of actions that can be triggered on this edit layer
-	*/
 	UE_DEPRECATED(5.7, "Register custom context menu actions using RegisterCustomEditLayerClassLayout from ILandscapeEditorModule")
 PRAGMA_DISABLE_DEPRECATION_WARNINGS
 	virtual TArray<FEditLayerAction> GetActions() const 
