Fixed landscape custom dialogs buttons

diff --git a/Engine/Source/Editor/LandscapeEditor/Private/LandscapeEditorDetailCustomization_NewLandscape.cpp b/Engine/Source/Editor/LandscapeEditor/Private/LandscapeEditorDetailCustomization_NewLandscape.cpp
--- a/Engine/Source/Editor/LandscapeEditor/Private/LandscapeEditorDetailCustomization_NewLandscape.cpp
+++ b/Engine/Source/Editor/LandscapeEditor/Private/LandscapeEditorDetailCustomization_NewLandscape.cpp
@@ -127,9 +127,18 @@ FReply FLandscapeEditorArrayBuilder_NewTargetLayer::ShowAutoFillTargetLayerDialo
 
 		const TSharedPtr<SCustomDialog> DialogWindow = SNew(SCustomDialog)
 			.Title(LOCTEXT("TargetLayerAutoFillLayers", "Auto-Fill Target Layer Assets"))
+			.HAlignButtonBox(HAlign_Center)
+			.VAlignButtonBox(VAlign_Center)
 			.Buttons({
-				SCustomDialog::FButton(LOCTEXT("OK", "OK")).SetPrimary(true),
-				SCustomDialog::FButton(LOCTEXT("Cancel", "Cancel"))
+				SCustomDialog::FButton(
+					LOCTEXT("OK", "OK"),
+					/*InOnClicked = */FSimpleDelegate::CreateLambda([&bIsCreateAssetForEmptyLayersChecked]()
+						{
+							HandleAutoFillTargetLayers(bIsCreateAssetForEmptyLayersChecked);
+						}),
+					SCustomDialog::EButtonRole::Confirm)
+					.SetPrimary(true),
+				SCustomDialog::FButton(LOCTEXT("Cancel", "Cancel"), /*InOnClicked = */FSimpleDelegate(), SCustomDialog::EButtonRole::Cancel)
 				})
 			.ContentAreaPadding(10.0f)
 			.Content()
@@ -171,11 +180,7 @@ FReply FLandscapeEditorArrayBuilder_NewTargetLayer::ShowAutoFillTargetLayerDialo
 					]
 			];
 
-		// -1 is cancel, 0 all layers
-		if (DialogWindow->ShowModal() == 0)
-		{
-			HandleAutoFillTargetLayers(bIsCreateAssetForEmptyLayersChecked);
-		}
+		DialogWindow->ShowModal();
 	}
 	
 	return FReply::Handled();
diff --git a/Engine/Source/Editor/LandscapeEditor/Private/LandscapeEditorDetailCustomization_TargetLayers.cpp b/Engine/Source/Editor/LandscapeEditor/Private/LandscapeEditorDetailCustomization_TargetLayers.cpp
--- a/Engine/Source/Editor/LandscapeEditor/Private/LandscapeEditorDetailCustomization_TargetLayers.cpp
+++ b/Engine/Source/Editor/LandscapeEditor/Private/LandscapeEditorDetailCustomization_TargetLayers.cpp
@@ -635,10 +635,24 @@ FReply FLandscapeEditorCustomNodeBuilder_TargetLayers::ShowAutoFillTargetLayerDi
 
 	const TSharedPtr<SCustomDialog> DialogWindow = SNew(SCustomDialog)
 		.Title(LOCTEXT("TargetLayerAutoFillLayers", "Auto-Fill Target Layer Assets"))
+		.HAlignButtonBox(HAlign_Center)
+		.VAlignButtonBox(VAlign_Center)
 		.Buttons({
-			SCustomDialog::FButton(LOCTEXT("Unassigned Layers Only", "Unassigned Layers Only")).SetPrimary(true),
-			SCustomDialog::FButton(LOCTEXT("All Layers", "All Layers")),
-			SCustomDialog::FButton(LOCTEXT("Cancel", "Cancel"))
+			SCustomDialog::FButton(
+				LOCTEXT("Unassigned Layers Only", "Unassigned Layers Only"), 
+				/*InOnClicked = */FSimpleDelegate::CreateLambda([this, &bIsCreateNewTargetLayersChecked]()
+						{
+							HandleAutoFillTargetLayers(/*bUpdateAllLayers = */false, bIsCreateNewTargetLayersChecked);
+						}),
+					SCustomDialog::EButtonRole::Confirm)
+					.SetPrimary(true),
+			SCustomDialog::FButton(
+				LOCTEXT("All Layers", "All Layers"),
+				/*InOnClicked = */FSimpleDelegate::CreateLambda([this, &bIsCreateNewTargetLayersChecked]()
+						{
+							HandleAutoFillTargetLayers(/*bUpdateAllLayers = */true, bIsCreateNewTargetLayersChecked);
+						})),
+			SCustomDialog::FButton(LOCTEXT("Cancel", "Cancel"), /*InOnClicked = */FSimpleDelegate(), SCustomDialog::EButtonRole::Cancel)
 			})
 		.ContentAreaPadding(10.0f)
 		.Content()
@@ -680,14 +694,7 @@ FReply FLandscapeEditorCustomNodeBuilder_TargetLayers::ShowAutoFillTargetLayerDi
 				]
 		];
 
-	const int32 ButtonPressedIndex = DialogWindow->ShowModal();
-	
-	// -1 is cancel, 0 empty layers, 1 all layers
-	if (ButtonPressedIndex == 0 || ButtonPressedIndex == 1)
-	{
-		const bool bUpdateAllLayers = (ButtonPressedIndex == 1);
-		HandleAutoFillTargetLayers(bUpdateAllLayers, bIsCreateNewTargetLayersChecked);
-	}
+	DialogWindow->ShowModal();
 
 	return FReply::Handled();
 }
diff --git a/Engine/Source/Editor/LandscapeEditor/Private/SLandscapeLayerListDialog.cpp b/Engine/Source/Editor/LandscapeEditor/Private/SLandscapeLayerListDialog.cpp
--- a/Engine/Source/Editor/LandscapeEditor/Private/SLandscapeLayerListDialog.cpp
+++ b/Engine/Source/Editor/LandscapeEditor/Private/SLandscapeLayerListDialog.cpp
@@ -198,20 +198,18 @@ void SLandscapeLayerListDialog::Construct(const FArguments& InArgs, const TWeakO
 	.OnGenerateRow(this, &SLandscapeLayerListDialog::OnGenerateRow)
 	.ListItemsSource(&WidgetLayerList);
 	
-	FButton CompleteButton(LOCTEXT("CompleteButtonText", "Complete"));
-	CompleteButton.OnClicked.BindSP(this, &SLandscapeLayerListDialog::OnComplete);
-	
 	// Construct custom dialog with list view supporting drag + drop
 	SCustomDialog::Construct(SCustomDialog::FArguments()
 		.Title(FText(LOCTEXT("LandscapeLayerListDialogTitleText", "Insert New Landscape Edit Layer")))
 		.UseScrollBox(false)
+		.HAlignButtonBox(HAlign_Center)
+		.VAlignButtonBox(VAlign_Center)
 		.WindowArguments(SWindow::FArguments()
-			.HasCloseButton(false))
+		.HasCloseButton(false))
+		.ContentAreaPadding(10.0f)
 		.Content()
 		[
 			SNew(SBox)
-			// Right padding of 20 is to offset SCustomDialog NullWidget to the left of content
-			.Padding(FMargin(0, 0, 20, 0))
 			[
 				SNew(SVerticalBox)
 				+ SVerticalBox::Slot()
@@ -237,7 +235,8 @@ void SLandscapeLayerListDialog::Construct(const FArguments& InArgs, const TWeakO
 			]
 		]
 		.Buttons({
-			CompleteButton
+			FButton(LOCTEXT("CompleteButtonText", "Complete"), FSimpleDelegate::CreateSP(this, &SLandscapeLayerListDialog::OnComplete), SCustomDialog::EButtonRole::Confirm)
+				.SetPrimary(true)
 		})
 	);
 }
