[Fix] Win64 - Verify_Activate_Landscape_Mode_With_Hotkey

https://jira.it.epicgames.com/browse/UE-320631

diff --git a/Engine/Source/Editor/LandscapeEditor/Private/Tests/LandscapeCoreWorkflowTests.cpp b/Engine/Source/Editor/LandscapeEditor/Private/Tests/LandscapeCoreWorkflowTests.cpp
--- a/Engine/Source/Editor/LandscapeEditor/Private/Tests/LandscapeCoreWorkflowTests.cpp
+++ b/Engine/Source/Editor/LandscapeEditor/Private/Tests/LandscapeCoreWorkflowTests.cpp
@@ -108,6 +108,7 @@ TEST_CLASS_WITH_FLAGS(LandscapeCoreWorkflow, "Editor.Landscape", EAutomationTest
 		ASSERT_THAT(IsNotNull(ActiveViewport, TEXT("Failed to find active Level Editor viewport.")));
 
 		// Set keyboard focus to the active viewport to ensure it receives input
+		Slate.ClearAllUserFocus(EFocusCause::SetDirectly);
 		const bool bFocusSet = Slate.SetKeyboardFocus(ActiveViewport->AsWidget(), EFocusCause::SetDirectly);
 		ASSERT_THAT(IsTrue(bFocusSet, TEXT("Failed to set keyboard focus to the active viewport.")));
 
