Landscape mode : hide landscape selector when the New tool in the Manage panel is active
+ misc code cleanup

diff --git a/Engine/Source/Editor/LandscapeEditor/Private/LandscapeEditorDetails.cpp b/Engine/Source/Editor/LandscapeEditor/Private/LandscapeEditorDetails.cpp
--- a/Engine/Source/Editor/LandscapeEditor/Private/LandscapeEditorDetails.cpp
+++ b/Engine/Source/Editor/LandscapeEditor/Private/LandscapeEditorDetails.cpp
@@ -126,7 +126,7 @@ void FLandscapeEditorDetails::CustomizeDetails(IDetailLayoutBuilder& DetailBuild
 			];
 	}
 		
-	// Only continue cuztomization if we are in NewLandscape mode or if editing is not disabled
+	// Only continue customization if we are not in NewLandscape mode or if editing is not disabled
 	if (bDisabledEditing && LandscapeEdMode->CurrentTool->GetToolName() != FName("NewLandscape"))
 	{
 		return;
@@ -223,60 +223,10 @@ void FLandscapeEditorDetails::CustomizeDetails(IDetailLayoutBuilder& DetailBuild
 	Customization_TargetLayers->CustomizeDetails(DetailBuilder);
 }
 
-FText FLandscapeEditorDetails::GetLocalizedName(FString Name)
-{
-#define LANDSCAPE_TOOL_LOCTEXT(KEY, SOURCE) { TEXT(KEY), LOCTEXT(KEY, SOURCE) }
-	static const TSortedMap<FString, FText, FDefaultAllocator, FLocKeySortedMapLess> ToolNames = {
-		LANDSCAPE_TOOL_LOCTEXT("ToolSet_NewLandscape", "New Landscape"),
-		LANDSCAPE_TOOL_LOCTEXT("ToolSet_ResizeLandscape", "Change Component Size"),
-		LANDSCAPE_TOOL_LOCTEXT("ToolSet_Sculpt", "Sculpt"),
-		LANDSCAPE_TOOL_LOCTEXT("ToolSet_Erase", "Erase"),
-		LANDSCAPE_TOOL_LOCTEXT("ToolSet_Paint", "Paint"),
-		LANDSCAPE_TOOL_LOCTEXT("ToolSet_Smooth", "Smooth"),
-		LANDSCAPE_TOOL_LOCTEXT("ToolSet_Flatten", "Flatten"),
-		LANDSCAPE_TOOL_LOCTEXT("ToolSet_Ramp", "Ramp"),
-		LANDSCAPE_TOOL_LOCTEXT("ToolSet_Erosion", "Erosion"),
-		LANDSCAPE_TOOL_LOCTEXT("ToolSet_HydraErosion", "HydroErosion"),
-		LANDSCAPE_TOOL_LOCTEXT("ToolSet_Noise", "Noise"),
-		LANDSCAPE_TOOL_LOCTEXT("ToolSet_Visibility", "Visibility"),
-		LANDSCAPE_TOOL_LOCTEXT("ToolSet_BlueprintBrush", "Blueprint Brushes"),
-
-		LANDSCAPE_TOOL_LOCTEXT("ToolSet_Select", "Selection"),
-		LANDSCAPE_TOOL_LOCTEXT("ToolSet_AddComponent", "Add"),
-		LANDSCAPE_TOOL_LOCTEXT("ToolSet_DeleteComponent", "Delete"),
-		LANDSCAPE_TOOL_LOCTEXT("ToolSet_MoveToLevel", "Move to Level"),
-
-		LANDSCAPE_TOOL_LOCTEXT("ToolSet_Mask", "Selection"),
-		LANDSCAPE_TOOL_LOCTEXT("ToolSet_CopyPaste", "Copy/Paste"),
-		LANDSCAPE_TOOL_LOCTEXT("ToolSet_Mirror", "Mirror"),
-
-		LANDSCAPE_TOOL_LOCTEXT("ToolSet_Splines", "Edit Splines"),
-
-		LANDSCAPE_TOOL_LOCTEXT("BrushSet_Circle", "Circle"),
-		LANDSCAPE_TOOL_LOCTEXT("BrushSet_Alpha", "Alpha"),
-		LANDSCAPE_TOOL_LOCTEXT("BrushSet_Pattern", "Pattern"),
-		LANDSCAPE_TOOL_LOCTEXT("BrushSet_Component", "Component"),
-		LANDSCAPE_TOOL_LOCTEXT("BrushSet_Gizmo", "Gizmo"),
-		LANDSCAPE_TOOL_LOCTEXT("BrushSet_Dummy", "NoBrush"),
-		LANDSCAPE_TOOL_LOCTEXT("BrushSet_Splines", "Splines"),
-
-		LANDSCAPE_TOOL_LOCTEXT("Circle_Smooth", "Smooth"),
-		LANDSCAPE_TOOL_LOCTEXT("Circle_Linear", "Linear"),
-		LANDSCAPE_TOOL_LOCTEXT("Circle_Spherical", "Spherical"),
-		LANDSCAPE_TOOL_LOCTEXT("Circle_Tip", "Tip"),
-		LANDSCAPE_TOOL_LOCTEXT("Circle_Dummy", "NoBrush"),
-	};
-#undef LANDSCAPE_TOOL_LOCTEXT
-	
-	FText Result = ToolNames.FindRef(Name);
-	ensure(!Result.IsEmpty());
-	return Result;
-}
-
 EVisibility FLandscapeEditorDetails::GetTargetLandscapeSelectorVisibility()
 {
 	FEdModeLandscape* LandscapeEdMode = GetEditorMode();
-	if (LandscapeEdMode && LandscapeEdMode->GetLandscapeList().Num() > 1)
+	if (LandscapeEdMode && (LandscapeEdMode->GetLandscapeList().Num() > 1) && !IsToolActive("NewLandscape"))
 	{
 		return EVisibility::Visible;
 	}
@@ -360,59 +310,6 @@ void FLandscapeEditorDetails::OnChangeTargetLandscape(TWeakObjectPtr<ULandscapeI
 	}
 }
 
-FText FLandscapeEditorDetails::GetCurrentToolName() const
-{
-	FEdModeLandscape* LandscapeEdMode = GetEditorMode();
-	if (LandscapeEdMode != NULL && LandscapeEdMode->CurrentTool != NULL)
-	{
-		const TCHAR* CurrentToolName = LandscapeEdMode->CurrentTool->GetToolName();
-		return GetLocalizedName(FString("ToolSet_") + CurrentToolName);
-	}
-
-	return LOCTEXT("Unknown", "Unknown");
-}
-
-FSlateIcon FLandscapeEditorDetails::GetCurrentToolIcon() const
-{
-	FEdModeLandscape* LandscapeEdMode = GetEditorMode();
-	if (LandscapeEdMode != NULL && LandscapeEdMode->CurrentTool != NULL)
-	{
-		const TCHAR* CurrentToolName = LandscapeEdMode->CurrentTool->GetToolName();
-		return FLandscapeEditorCommands::Get().NameToCommandMap.FindChecked(*(FString("Tool_") + CurrentToolName))->GetIcon();
-	}
-
-	return FSlateIcon(FAppStyle::GetAppStyleSetName(), "Default");
-}
-
-bool FLandscapeEditorDetails::GetToolSelectorIsVisible() const
-{
-	FEdModeLandscape* LandscapeEdMode = GetEditorMode();
-	if (LandscapeEdMode && LandscapeEdMode->CurrentTool)
-	{
-		if (!LandscapeEdMode->CanEditCurrentTarget())
-		{
-			return false;
-		}
-
-		if (!IsToolActive("NewLandscape") || LandscapeEdMode->GetLandscapeList().Num() > 0)
-		{
-			return true;
-		}
-	}
-
-	return false;
-}
-
-EVisibility FLandscapeEditorDetails::GetToolSelectorVisibility() const
-{
-	if (GetToolSelectorIsVisible())
-	{
-		return EVisibility::Visible;
-	}
-	return EVisibility::Collapsed;
-}
-
-
 FName FLandscapeEditorDetails::GetCurrentBrushFName() const
 {
 	FEdModeLandscape* LandscapeEdMode = GetEditorMode();
@@ -424,35 +321,6 @@ FName FLandscapeEditorDetails::GetCurrentBrushFName() const
 	return NAME_None;
 }
 
-
-FText FLandscapeEditorDetails::GetCurrentBrushName() const
-{
-	FEdModeLandscape* LandscapeEdMode = GetEditorMode();
-	if (LandscapeEdMode != NULL && LandscapeEdMode->CurrentBrush != NULL)
-	{
-		const FName CurrentBrushSetName = LandscapeEdMode->LandscapeBrushSets[LandscapeEdMode->CurrentBrushSetIndex].BrushSetName;
-		return GetLocalizedName(CurrentBrushSetName.ToString());
-	}
-
-	return LOCTEXT("Unknown", "Unknown");
-}
-
-FSlateIcon FLandscapeEditorDetails::GetCurrentBrushIcon() const
-{
-	FEdModeLandscape* LandscapeEdMode = GetEditorMode();
-	if (LandscapeEdMode != NULL && LandscapeEdMode->CurrentBrush != NULL)
-	{
-		const FName CurrentBrushSetName = LandscapeEdMode->LandscapeBrushSets[LandscapeEdMode->CurrentBrushSetIndex].BrushSetName;
-		TSharedPtr<FUICommandInfo> Command = FLandscapeEditorCommands::Get().NameToCommandMap.FindRef(CurrentBrushSetName);
-		if (Command.IsValid())
-		{
-			return Command->GetIcon();
-		}
-	}
-
-	return FSlateIcon(FAppStyle::GetAppStyleSetName(), "Default");
-}
-
 bool FLandscapeEditorDetails::GetBrushSelectorIsVisible() const
 {
 	FEdModeLandscape* LandscapeEdMode = GetEditorMode();
@@ -483,34 +351,6 @@ FName FLandscapeEditorDetails::GetCurrentBrushFalloffFName() const
 	return NAME_None;
 }
 
-FText FLandscapeEditorDetails::GetCurrentBrushFalloffName() const
-{
-	FEdModeLandscape* LandscapeEdMode = GetEditorMode();
-	if (LandscapeEdMode != NULL && LandscapeEdMode->CurrentBrush != NULL && GetBrushFalloffSelectorIsVisible())
-	{
-		const TCHAR* CurrentBrushName = LandscapeEdMode->CurrentBrush->GetBrushName();
-		return GetLocalizedName(CurrentBrushName);
-	}
-
-	return LOCTEXT("Unknown", "Unknown");
-}
-
-FSlateIcon FLandscapeEditorDetails::GetCurrentBrushFalloffIcon() const
-{
-	FEdModeLandscape* LandscapeEdMode = GetEditorMode();
-	if (LandscapeEdMode != NULL && LandscapeEdMode->CurrentBrush != NULL)
-	{
-		const FName CurrentBrushName = LandscapeEdMode->CurrentBrush->GetBrushName();
-		TSharedPtr<FUICommandInfo> Command = FLandscapeEditorCommands::Get().NameToCommandMap.FindRef(CurrentBrushName);
-		if (Command.IsValid())
-		{
-			return Command->GetIcon();
-		}
-	}
-
-	return FSlateIcon(FAppStyle::GetAppStyleSetName(), "Default");
-}
-
 void FLandscapeEditorDetails::SetBrushCommand(FName InBrush)
 {
 	TSharedPtr<FUICommandList> CommandList = GetEditorMode()->GetUICommandList();
diff --git a/Engine/Source/Editor/LandscapeEditor/Private/LandscapeEditorDetails.h b/Engine/Source/Editor/LandscapeEditor/Private/LandscapeEditorDetails.h
--- a/Engine/Source/Editor/LandscapeEditor/Private/LandscapeEditorDetails.h
+++ b/Engine/Source/Editor/LandscapeEditor/Private/LandscapeEditorDetails.h
@@ -46,28 +46,17 @@ public:
 	virtual void CustomizeDetails(IDetailLayoutBuilder& DetailBuilder) override;
 
 protected:
-	static FText GetLocalizedName(FString Name);
-
 	static EVisibility GetTargetLandscapeSelectorVisibility();
 	static FText GetLandscapeDisplayName(ALandscape* InLandscape);
 	static FText GetTargetLandscapeName();
 	static TSharedRef<SWidget> GetTargetLandscapeMenu();
 	static void OnChangeTargetLandscape(TWeakObjectPtr<ULandscapeInfo> LandscapeInfo);
 
-	FText GetCurrentToolName() const;
-	FSlateIcon GetCurrentToolIcon() const;
-	bool GetToolSelectorIsVisible() const;
-	EVisibility GetToolSelectorVisibility() const;
-
 	FName GetCurrentBrushFName() const;
-	FText GetCurrentBrushName() const;
-	FSlateIcon GetCurrentBrushIcon() const;
 	bool GetBrushSelectorIsVisible() const;
 	EVisibility GetBrushSelectorVisibility() const;
 
 	FName GetCurrentBrushFalloffFName() const;
-	FText GetCurrentBrushFalloffName() const;
-	FSlateIcon GetCurrentBrushFalloffIcon() const;
 	bool GetBrushFalloffSelectorIsVisible() const;
 	EVisibility GetBrushFalloffSelectorVisibility() const;
 
