Adjusted comment

#rnx
#skipci

diff --git a/Engine/Source/Editor/LandscapeEditor/Private/LandscapeEditorDetails.cpp b/Engine/Source/Editor/LandscapeEditor/Private/LandscapeEditorDetails.cpp
--- a/Engine/Source/Editor/LandscapeEditor/Private/LandscapeEditorDetails.cpp
+++ b/Engine/Source/Editor/LandscapeEditor/Private/LandscapeEditorDetails.cpp
@@ -126,7 +126,7 @@ void FLandscapeEditorDetails::CustomizeDetails(IDetailLayoutBuilder& DetailBuild
 			];
 	}
 		
-	// Only continue customization if we are not in NewLandscape mode or if editing is not disabled
+	// Only continue customization if we are in NewLandscape mode or if editing is not disabled
 	if (bDisabledEditing && LandscapeEdMode->CurrentTool->GetToolName() != FName("NewLandscape"))
 	{
 		return;
