Fix assert when deleting landscape spline points and segments at the same time.

diff --git a/Engine/Source/Editor/LandscapeEditor/Private/LandscapeEdModeSplineTools.cpp b/Engine/Source/Editor/LandscapeEditor/Private/LandscapeEdModeSplineTools.cpp
--- a/Engine/Source/Editor/LandscapeEditor/Private/LandscapeEdModeSplineTools.cpp
+++ b/Engine/Source/Editor/LandscapeEditor/Private/LandscapeEdModeSplineTools.cpp
@@ -388,7 +388,12 @@ public:
 	{
 		FScopedTransaction Transaction(LOCTEXT("LandscapeSpline_DeleteSegment", "Delete Landscape Spline Segment"));
 
-		ULandscapeSplinesComponent* SplinesComponent = ToDelete->GetOuterULandscapeSplinesComponent();
+		ULandscapeSplinesComponent* SplinesComponent = ToDelete->GetOuterSafe();
+		if (!SplinesComponent)
+		{
+			// Segment is already deleted
+			return;
+		}
 		SplinesComponent->Modify();
 
 		ToDelete->Modify();
@@ -426,7 +431,12 @@ public:
 	{
 		FScopedTransaction Transaction(LOCTEXT("LandscapeSpline_DeleteControlPoint", "Delete Landscape Spline Control Point"));
 
-		ULandscapeSplinesComponent* SplinesComponent = ToDelete->GetOuterULandscapeSplinesComponent();
+		ULandscapeSplinesComponent* SplinesComponent = ToDelete->GetOuterSafe();
+		if (!SplinesComponent)
+		{
+			// Point is already deleted
+			return;
+		}
 		SplinesComponent->Modify();
 
 		ToDelete->Modify();
@@ -1957,6 +1967,7 @@ public:
 	{
 		TArray<ULandscapeSplineControlPoint*> SelectedSplineControlPoints = SplineSelection->GetSelectedSplineControlPoints();
 		TArray<ULandscapeSplineSegment*> SelectedSplineSegments = SplineSelection->GetSelectedSplineSegments();
+		TSet<ULandscapeSplinesComponent*> AffectedSplineComponents;
 
 		if (SelectedSplineControlPoints.Num() > 0 || SelectedSplineSegments.Num() > 0)
 		{
@@ -1964,14 +1975,25 @@ public:
 
 			for (ULandscapeSplineControlPoint* ControlPoint : SelectedSplineControlPoints)
 			{
+				AffectedSplineComponents.Add(ControlPoint->GetOuterSafe());
 				DeleteControlPoint(ControlPoint, bDeleteLooseEnds);
 			}
 			for (ULandscapeSplineSegment* Segment : SelectedSplineSegments)
 			{
+				AffectedSplineComponents.Add(Segment->GetOuterSafe());
 				DeleteSegment(Segment, bDeleteLooseEnds);
 			}
 			SplineSelection->ClearSelection();
 		}
+
+		for (ULandscapeSplinesComponent* Component : AffectedSplineComponents)
+		{
+			// Could have nulls here because deleting control points might have already deleted some segments.
+			if (Component)
+			{
+				Component->CheckSplinesValid();
+			}
+		}
 	}
 
 	void InternalProcessEditCut()
diff --git a/Engine/Source/Runtime/Landscape/Classes/LandscapeSplineControlPoint.h b/Engine/Source/Runtime/Landscape/Classes/LandscapeSplineControlPoint.h
--- a/Engine/Source/Runtime/Landscape/Classes/LandscapeSplineControlPoint.h
+++ b/Engine/Source/Runtime/Landscape/Classes/LandscapeSplineControlPoint.h
@@ -289,9 +289,9 @@ public:
 
 	/** Use GetOuterSafe instead of GetOuterULandscapeSplinesComponent if there is any danger that this point has been deleted (via DeleteSplinePoints).
 	  * This will return nullptr in that case, because the outer is no longer a ULandscapeSplinesComponent */
-	ULandscapeSplinesComponent* GetOuterSafe() const;
+	LANDSCAPE_API ULandscapeSplinesComponent* GetOuterSafe() const;
 	/** Use GetOuterChecked when the call is always expected to succeed.  Will assert if it fails. */
-	ULandscapeSplinesComponent* GetOuterChecked() const;
+	LANDSCAPE_API ULandscapeSplinesComponent* GetOuterChecked() const;
 
 	friend class FLandscapeToolSplines;
 	friend class ULandscapeInfo;
diff --git a/Engine/Source/Runtime/Landscape/Classes/LandscapeSplineSegment.h b/Engine/Source/Runtime/Landscape/Classes/LandscapeSplineSegment.h
--- a/Engine/Source/Runtime/Landscape/Classes/LandscapeSplineSegment.h
+++ b/Engine/Source/Runtime/Landscape/Classes/LandscapeSplineSegment.h
@@ -387,9 +387,9 @@ public:
 
 	/** Use GetOuterSafe instead of GetOuterULandscapeSplinesComponent if there is any danger that this segment has been deleted (via DeleteSplinePoints).
 	  * This will return nullptr in that case, because the outer is no longer a ULandscapeSplinesComponent */
-	ULandscapeSplinesComponent* GetOuterSafe() const;
+	LANDSCAPE_API ULandscapeSplinesComponent* GetOuterSafe() const;
 	/** Use GetOuterChecked when the call is always expected to succeed.  Will assert if it fails. */
-	ULandscapeSplinesComponent* GetOuterChecked() const;
+	LANDSCAPE_API ULandscapeSplinesComponent* GetOuterChecked() const;
 
 	friend class FLandscapeToolSplines;
 	friend class ULandscapeInfo;
diff --git a/Engine/Source/Runtime/Landscape/Classes/LandscapeSplinesComponent.h b/Engine/Source/Runtime/Landscape/Classes/LandscapeSplinesComponent.h
--- a/Engine/Source/Runtime/Landscape/Classes/LandscapeSplinesComponent.h
+++ b/Engine/Source/Runtime/Landscape/Classes/LandscapeSplinesComponent.h
@@ -154,7 +154,9 @@ public:
 
 	LANDSCAPE_API ILandscapeSplineInterface* GetSplineOwner();
 
-	void CheckSplinesValid();
+	// Run some ensures on the internal state of the data
+	LANDSCAPE_API void CheckSplinesValid();
+
 	bool ModifySplines(bool bAlwaysMarkDirty = true);
 
 	const TArray<TObjectPtr<ULandscapeSplineControlPoint>>& GetControlPoints() const { return ControlPoints; }
