Fix intermittent assert when deleting landscape spline points.

diff --git a/Engine/Source/Editor/LandscapeEditor/Private/LandscapeEdModeSplineTools.cpp b/Engine/Source/Editor/LandscapeEditor/Private/LandscapeEdModeSplineTools.cpp
--- a/Engine/Source/Editor/LandscapeEditor/Private/LandscapeEdModeSplineTools.cpp
+++ b/Engine/Source/Editor/LandscapeEditor/Private/LandscapeEdModeSplineTools.cpp
@@ -437,17 +437,24 @@ public:
 			// Point is already deleted
 			return;
 		}
-		SplinesComponent->Modify();
 
+		// If needed, create prompt window before we start editing.  Creating the window can cause the selection code to touch the
+		// spline network when it updates hit proxies.  Don't let that code observe an incomplete state.
+		bool bConnectable = ToDelete->ConnectedSegments.Num() == 2 && ToDelete->ConnectedSegments[0].Segment != ToDelete->ConnectedSegments[1].Segment;
+		int32 ConnectChoice = EAppReturnType::No;
+		if (bConnectable)
+		{
+			ConnectChoice = FMessageDialog::Open(EAppMsgType::YesNoCancel, LOCTEXT("WantToJoinControlPoint", "Control point has two segments attached, do you want to join them?"));
+		}
+
+		SplinesComponent->Modify();
 		ToDelete->Modify();
 		ToDelete->DeleteSplinePoints();
 		SplineSelection->ResetNavigationPath();
 
-		if (ToDelete->ConnectedSegments.Num() == 2
-			&& ToDelete->ConnectedSegments[0].Segment != ToDelete->ConnectedSegments[1].Segment)
+		if (bConnectable)
 		{
-			int32 Result = FMessageDialog::Open(EAppMsgType::YesNoCancel, LOCTEXT("WantToJoinControlPoint", "Control point has two segments attached, do you want to join them?"));
-			switch (Result)
+			switch (ConnectChoice)
 			{
 			case EAppReturnType::Yes:
 			{
