Fixed circular include in landscape editor code

#rnx

diff --git a/Engine/Source/Editor/LandscapeEditor/Private/ActorFactoryLandscape.cpp b/Engine/Source/Editor/LandscapeEditor/Private/ActorFactoryLandscape.cpp
--- a/Engine/Source/Editor/LandscapeEditor/Private/ActorFactoryLandscape.cpp
+++ b/Engine/Source/Editor/LandscapeEditor/Private/ActorFactoryLandscape.cpp
@@ -10,6 +10,7 @@
 #include "EditorModes.h"
 #include "LandscapeProxy.h"
 #include "LandscapeEditorObject.h"
+#include "LandscapeEdMode.h"
 #include "Classes/LandscapePlaceholder.h"
 
 #include UE_INLINE_GENERATED_CPP_BY_NAME(ActorFactoryLandscape)
diff --git a/Engine/Source/Editor/LandscapeEditor/Private/LandscapeEditorObject.cpp b/Engine/Source/Editor/LandscapeEditor/Private/LandscapeEditorObject.cpp
--- a/Engine/Source/Editor/LandscapeEditor/Private/LandscapeEditorObject.cpp
+++ b/Engine/Source/Editor/LandscapeEditor/Private/LandscapeEditorObject.cpp
@@ -8,6 +8,7 @@
 #include "UObject/ConstructorHelpers.h"
 #include "LandscapeEditorModule.h"
 #include "LandscapeEditorPrivate.h"
+#include "LandscapeEdMode.h"
 #include "LandscapeRender.h"
 #include "LandscapeSettings.h"
 #include "LandscapeImportHelper.h"
@@ -522,6 +523,12 @@ void ULandscapeEditorObject::SetGizmoSnapMode(ELandscapeGizmoSnapType InSnapMode
 	}
 }
 
+bool ULandscapeEditorObject::IsWeightmapTarget() const
+{
+	check(ParentMode);
+	return !(ParentMode->CurrentToolTarget.TargetType == ELandscapeToolTargetType::Heightmap || ParentMode->CurrentToolTarget.TargetType == ELandscapeToolTargetType::Visibility);
+}
+
 bool ULandscapeEditorObject::LoadAlphaTextureSourceData(UTexture2D* InTexture, TArray<uint8>& OutSourceData, int32& OutSourceDataSizeX, int32& OutSourceDataSizeY, ELandscapeTextureColorChannel& InOutTextureChannel)
 {
 	check(InTexture != nullptr);
diff --git a/Engine/Source/Editor/LandscapeEditor/Public/LandscapeEditorObject.h b/Engine/Source/Editor/LandscapeEditor/Public/LandscapeEditorObject.h
--- a/Engine/Source/Editor/LandscapeEditor/Public/LandscapeEditorObject.h
+++ b/Engine/Source/Editor/LandscapeEditor/Public/LandscapeEditorObject.h
@@ -6,15 +6,17 @@
 #include "UObject/ObjectMacros.h"
 #include "UObject/Object.h"
 #include "LandscapeProxy.h"
-#include "Editor/LandscapeEditor/Private/LandscapeEdMode.h"
 #include "LandscapeFileFormatInterface.h"
 #include "LandscapeBlueprintBrushBase.h"
 #include "LandscapeImportHelper.h"
 
 #include "LandscapeEditorObject.generated.h"
 
+class ALandscapeGizmoActiveActor;
+class FEdModeLandscape;
 class ULandscapeMaterialInstanceConstant;
 class UTexture2D;
+enum class ELandscapeGizmoSnapType;
 
 UENUM()
 enum class ELandscapeToolFlattenMode : int8
@@ -859,11 +861,7 @@ public:
 
 private:
 
-	bool IsWeightmapTarget() const
-	{
-		check(ParentMode);
-		return !(ParentMode->CurrentToolTarget.TargetType == ELandscapeToolTargetType::Heightmap || ParentMode->CurrentToolTarget.TargetType == ELandscapeToolTargetType::Visibility);
-	}
+	bool IsWeightmapTarget() const;
 
 	static bool LoadAlphaTextureSourceData(UTexture2D* InTexture, TArray<uint8>& OutSourceData, int32& OutSourceDataSizeX, int32& OutSourceDataSizeY, ELandscapeTextureColorChannel& InOutTextureChannel);
 };
