Fixed non-unity build

#rnx

diff --git a/Engine/Source/Editor/LandscapeEditor/Private/LandscapeEditorDetailCustomization_CopyPaste.cpp b/Engine/Source/Editor/LandscapeEditor/Private/LandscapeEditorDetailCustomization_CopyPaste.cpp
--- a/Engine/Source/Editor/LandscapeEditor/Private/LandscapeEditorDetailCustomization_CopyPaste.cpp
+++ b/Engine/Source/Editor/LandscapeEditor/Private/LandscapeEditorDetailCustomization_CopyPaste.cpp
@@ -7,7 +7,8 @@
 #include "SlateOptMacros.h"
 #include "Widgets/Input/SButton.h"
 #include "LandscapeEditorObject.h"
-
+#include "LandscapeEdMode.h"
+#include "LandscapeGizmoActiveActor.h"
 #include "DetailLayoutBuilder.h"
 #include "DetailWidgetRow.h"
 #include "IDetailGroup.h"
@@ -196,16 +197,16 @@ void FLandscapeEditorDetailCustomization_CopyPaste::CustomizeDetails(IDetailLayo
 						{
 							CurrentGuessedDimension = *InItem.Get();
 							
-							TArray<FString> SeperatedValues;
-							CurrentGuessedDimension.ParseIntoArray(SeperatedValues, TEXT(" "));
+							TArray<FString> SeparatedValues;
+							CurrentGuessedDimension.ParseIntoArray(SeparatedValues, TEXT(" "));
 
-							check(SeperatedValues.Num() == 3); // it should contain Number1 x Number2
+							check(SeparatedValues.Num() == 3); // it should contain Number1 x Number2
 
-							int32 FirstNumber = FCString::Atoi(*SeperatedValues[0]);
-							int32 SecondNumber = FCString::Atoi(*SeperatedValues[2]);
+							int32 FirstNumber = FCString::Atoi(*SeparatedValues[0]);
+							int32 SecondNumber = FCString::Atoi(*SeparatedValues[2]);
 
 							FEdModeLandscape* LandscapeEdMode = GetEditorMode();
-							if (LandscapeEdMode != NULL)
+							if (LandscapeEdMode != nullptr)
 							{
 								LandscapeEdMode->UISettings->GizmoImportSize.X = FirstNumber;
 								LandscapeEdMode->UISettings->GizmoImportSize.Y = SecondNumber;
@@ -246,7 +247,7 @@ END_SLATE_FUNCTION_BUILD_OPTIMIZATION
 FReply FLandscapeEditorDetailCustomization_CopyPaste::OnCopyToGizmoButtonClicked()
 {
 	FEdModeLandscape* LandscapeEdMode = GetEditorMode();
-	if (LandscapeEdMode != NULL)
+	if (LandscapeEdMode != nullptr)
 	{
 		LandscapeEdMode->CopyDataToGizmo();
 	}
@@ -257,7 +258,7 @@ FReply FLandscapeEditorDetailCustomization_CopyPaste::OnCopyToGizmoButtonClicked
 FReply FLandscapeEditorDetailCustomization_CopyPaste::OnFitGizmoToSelectionButtonClicked()
 {
 	FEdModeLandscape* LandscapeEdMode = GetEditorMode();
-	if (LandscapeEdMode != NULL)
+	if (LandscapeEdMode != nullptr)
 	{
 		ALandscapeGizmoActiveActor* Gizmo = LandscapeEdMode->UISettings->CurrentGizmoActor.Get();
 		if (Gizmo && Gizmo->TargetLandscapeInfo)
@@ -274,7 +275,7 @@ FReply FLandscapeEditorDetailCustomization_CopyPaste::OnFitGizmoToSelectionButto
 FReply FLandscapeEditorDetailCustomization_CopyPaste::OnFitHeightsToGizmoButtonClicked()
 {
 	FEdModeLandscape* LandscapeEdMode = GetEditorMode();
-	if (LandscapeEdMode != NULL)
+	if (LandscapeEdMode != nullptr)
 	{
 		ALandscapeGizmoActiveActor* Gizmo = LandscapeEdMode->UISettings->CurrentGizmoActor.Get();
 		if (Gizmo && Gizmo->TargetLandscapeInfo)
@@ -291,7 +292,7 @@ FReply FLandscapeEditorDetailCustomization_CopyPaste::OnFitHeightsToGizmoButtonC
 FReply FLandscapeEditorDetailCustomization_CopyPaste::OnClearGizmoDataButtonClicked()
 {
 	FEdModeLandscape* LandscapeEdMode = GetEditorMode();
-	if (LandscapeEdMode != NULL)
+	if (LandscapeEdMode != nullptr)
 	{
 		ALandscapeGizmoActiveActor* Gizmo = LandscapeEdMode->UISettings->CurrentGizmoActor.Get();
 		if (Gizmo && Gizmo->TargetLandscapeInfo)
@@ -306,11 +307,11 @@ FReply FLandscapeEditorDetailCustomization_CopyPaste::OnClearGizmoDataButtonClic
 FReply FLandscapeEditorDetailCustomization_CopyPaste::OnGizmoHeightmapFilenameButtonClicked(TSharedRef<IPropertyHandle> HeightmapPropertyHandle)
 {
 	FEdModeLandscape* LandscapeEdMode = GetEditorMode();
-	if (LandscapeEdMode != NULL)
+	if (LandscapeEdMode != nullptr)
 	{
 		// Prompt the user for the Filenames
 		IDesktopPlatform* DesktopPlatform = FDesktopPlatformModule::Get();
-		if (DesktopPlatform != NULL)
+		if (DesktopPlatform != nullptr)
 		{
 			TArray<FString> OpenFilenames;
 			bool bOpened = DesktopPlatform->OpenFileDialog(
@@ -346,7 +347,7 @@ FReply FLandscapeEditorDetailCustomization_CopyPaste::OnGizmoHeightmapFilenameBu
 bool FLandscapeEditorDetailCustomization_CopyPaste::GetGizmoImportButtonIsEnabled() const
 {
 	FEdModeLandscape* LandscapeEdMode = GetEditorMode();
-	if (LandscapeEdMode != NULL)
+	if (LandscapeEdMode != nullptr)
 	{
 		if (!LandscapeEdMode->UISettings->GizmoHeightmapFilenameString.IsEmpty() && LandscapeEdMode->UISettings->GizmoImportSize != FIntPoint(0, 0))
 		{
@@ -360,7 +361,7 @@ bool FLandscapeEditorDetailCustomization_CopyPaste::GetGizmoImportButtonIsEnable
 bool FLandscapeEditorDetailCustomization_CopyPaste::GetGizmoGuessSizeButtonIsEnabled() const
 {
 	FEdModeLandscape* LandscapeEdMode = GetEditorMode();
-	if (LandscapeEdMode != NULL)
+	if (LandscapeEdMode != nullptr)
 	{
 		if (!LandscapeEdMode->UISettings->GizmoHeightmapFilenameString.IsEmpty())
 		{
@@ -375,7 +376,7 @@ void FLandscapeEditorDetailCustomization_CopyPaste::GenerateGuessDimensionList()
 {
 	FEdModeLandscape* LandscapeEdMode = GetEditorMode();
 
-	if (LandscapeEdMode != NULL)
+	if (LandscapeEdMode != nullptr)
 	{
 		check(!LandscapeEdMode->UISettings->GizmoHeightmapFilenameString.IsEmpty());
 
@@ -460,7 +461,7 @@ void FLandscapeEditorDetailCustomization_CopyPaste::GenerateGuessDimensionList()
 FReply FLandscapeEditorDetailCustomization_CopyPaste::OnGizmoImportButtonClicked()
 {
 	FEdModeLandscape* LandscapeEdMode = GetEditorMode();
-	if (LandscapeEdMode != NULL)
+	if (LandscapeEdMode != nullptr)
 	{
 		ALandscapeGizmoActiveActor* Gizmo = LandscapeEdMode->UISettings->CurrentGizmoActor.Get();
 		if (Gizmo)
@@ -516,7 +517,7 @@ FReply FLandscapeEditorDetailCustomization_CopyPaste::OnGizmoImportButtonClicked
 				}
 			}
 
-			Gizmo->Import(LandscapeEdMode->UISettings->GizmoImportSize.X, LandscapeEdMode->UISettings->GizmoImportSize.Y, (uint16*)Data.GetData(), LayerInfos, LayerDataPtrs.Num() ? LayerDataPtrs.GetData() : NULL);
+			Gizmo->Import(LandscapeEdMode->UISettings->GizmoImportSize.X, LandscapeEdMode->UISettings->GizmoImportSize.Y, (uint16*)Data.GetData(), LayerInfos, LayerDataPtrs.Num() ? LayerDataPtrs.GetData() : nullptr);
 
 			// Make sure gizmo actor is selected
 			GEditor->SelectNone(false, true);
@@ -530,7 +531,7 @@ FReply FLandscapeEditorDetailCustomization_CopyPaste::OnGizmoImportButtonClicked
 FReply FLandscapeEditorDetailCustomization_CopyPaste::OnGizmoExportButtonClicked()
 {
 	FEdModeLandscape* LandscapeEdMode = GetEditorMode();
-	if (LandscapeEdMode != NULL)
+	if (LandscapeEdMode != nullptr)
 	{
 		ALandscapeGizmoActiveActor* Gizmo = LandscapeEdMode->UISettings->CurrentGizmoActor.Get();
 		if (Gizmo && Gizmo->TargetLandscapeInfo && Gizmo->SelectedData.Num())
@@ -669,7 +670,7 @@ FReply FLandscapeEditorStructCustomization_FGizmoImportLayer::OnGizmoImportLayer
 
 	// Prompt the user for the Filenames
 	IDesktopPlatform* DesktopPlatform = FDesktopPlatformModule::Get();
-	if (DesktopPlatform != NULL)
+	if (DesktopPlatform != nullptr)
 	{
 		TArray<FString> OpenFilenames;
 		bool bOpened = DesktopPlatform->OpenFileDialog(
diff --git a/Engine/Source/Editor/LandscapeEditor/Private/LandscapeEditorDetailCustomization_ResizeLandscape.cpp b/Engine/Source/Editor/LandscapeEditor/Private/LandscapeEditorDetailCustomization_ResizeLandscape.cpp
--- a/Engine/Source/Editor/LandscapeEditor/Private/LandscapeEditorDetailCustomization_ResizeLandscape.cpp
+++ b/Engine/Source/Editor/LandscapeEditor/Private/LandscapeEditorDetailCustomization_ResizeLandscape.cpp
@@ -8,6 +8,7 @@
 #include "Widgets/Input/SButton.h"
 #include "Widgets/Input/SComboButton.h"
 #include "LandscapeEditorObject.h"
+#include "LandscapeEdMode.h"
 
 #include "DetailLayoutBuilder.h"
 #include "DetailWidgetRow.h"
@@ -292,7 +293,7 @@ END_SLATE_FUNCTION_BUILD_OPTIMIZATION
 FText FLandscapeEditorDetailCustomization_ResizeLandscape::GetOriginalSectionSize()
 {
 	FEdModeLandscape* LandscapeEdMode = GetEditorMode();
-	if (LandscapeEdMode != NULL)
+	if (LandscapeEdMode != nullptr)
 	{
 		return FText::Format(LOCTEXT("NxNQuads", "{0}x{0} Quads"), FText::AsNumber(LandscapeEdMode->UISettings->ResizeLandscape_Original_QuadsPerSection));
 	}
@@ -302,7 +303,7 @@ FText FLandscapeEditorDetailCustomization_ResizeLandscape::GetOriginalSectionSiz
 
 TSharedRef<SWidget> FLandscapeEditorDetailCustomization_ResizeLandscape::GetSectionSizeMenu(TSharedRef<IPropertyHandle> PropertyHandle)
 {
-	FMenuBuilder MenuBuilder(true, NULL);
+	FMenuBuilder MenuBuilder(true, nullptr);
 
 	for (int32 i = 0; i < UE_ARRAY_COUNT(FLandscapeConfig::SubsectionSizeQuadsValues); i++)
 	{
@@ -334,7 +335,7 @@ FText FLandscapeEditorDetailCustomization_ResizeLandscape::GetSectionSize(TShare
 bool FLandscapeEditorDetailCustomization_ResizeLandscape::IsSectionSizeResetToDefaultVisible(TSharedPtr<IPropertyHandle> PropertyHandle)
 {
 	FEdModeLandscape* LandscapeEdMode = GetEditorMode();
-	if (LandscapeEdMode != NULL)
+	if (LandscapeEdMode != nullptr)
 	{
 		return LandscapeEdMode->UISettings->ResizeLandscape_QuadsPerSection != LandscapeEdMode->UISettings->ResizeLandscape_Original_QuadsPerSection;
 	}
@@ -345,7 +346,7 @@ bool FLandscapeEditorDetailCustomization_ResizeLandscape::IsSectionSizeResetToDe
 void FLandscapeEditorDetailCustomization_ResizeLandscape::OnSectionSizeResetToDefault(TSharedPtr<IPropertyHandle> PropertyHandle)
 {
 	FEdModeLandscape* LandscapeEdMode = GetEditorMode();
-	if (LandscapeEdMode != NULL)
+	if (LandscapeEdMode != nullptr)
 	{
 		LandscapeEdMode->UISettings->ResizeLandscape_QuadsPerSection = LandscapeEdMode->UISettings->ResizeLandscape_Original_QuadsPerSection;
 	}
@@ -354,7 +355,7 @@ void FLandscapeEditorDetailCustomization_ResizeLandscape::OnSectionSizeResetToDe
 FText FLandscapeEditorDetailCustomization_ResizeLandscape::GetOriginalSectionsPerComponent()
 {
 	FEdModeLandscape* LandscapeEdMode = GetEditorMode();
-	if (LandscapeEdMode != NULL)
+	if (LandscapeEdMode != nullptr)
 	{
 		int32 SectionsPerComponent = LandscapeEdMode->UISettings->ResizeLandscape_Original_SectionsPerComponent;
 
@@ -369,7 +370,7 @@ FText FLandscapeEditorDetailCustomization_ResizeLandscape::GetOriginalSectionsPe
 
 TSharedRef<SWidget> FLandscapeEditorDetailCustomization_ResizeLandscape::GetSectionsPerComponentMenu(TSharedRef<IPropertyHandle> PropertyHandle)
 {
-	FMenuBuilder MenuBuilder(true, NULL);
+	FMenuBuilder MenuBuilder(true, nullptr);
 
 	for (int32 i = 0; i < UE_ARRAY_COUNT(FLandscapeConfig::NumSectionValues); i++)
 	{
@@ -407,7 +408,7 @@ FText FLandscapeEditorDetailCustomization_ResizeLandscape::GetSectionsPerCompone
 bool FLandscapeEditorDetailCustomization_ResizeLandscape::IsSectionsPerComponentResetToDefaultVisible(TSharedPtr<IPropertyHandle> PropertyHandle)
 {
 	FEdModeLandscape* LandscapeEdMode = GetEditorMode();
-	if (LandscapeEdMode != NULL)
+	if (LandscapeEdMode != nullptr)
 	{
 		return LandscapeEdMode->UISettings->ResizeLandscape_SectionsPerComponent != LandscapeEdMode->UISettings->ResizeLandscape_Original_SectionsPerComponent;
 	}
@@ -418,7 +419,7 @@ bool FLandscapeEditorDetailCustomization_ResizeLandscape::IsSectionsPerComponent
 void FLandscapeEditorDetailCustomization_ResizeLandscape::OnSectionsPerComponentResetToDefault(TSharedPtr<IPropertyHandle> PropertyHandle)
 {
 	FEdModeLandscape* LandscapeEdMode = GetEditorMode();
-	if (LandscapeEdMode != NULL)
+	if (LandscapeEdMode != nullptr)
 	{
 		LandscapeEdMode->UISettings->ResizeLandscape_SectionsPerComponent = LandscapeEdMode->UISettings->ResizeLandscape_Original_SectionsPerComponent;
 	}
@@ -427,7 +428,7 @@ void FLandscapeEditorDetailCustomization_ResizeLandscape::OnSectionsPerComponent
 FText FLandscapeEditorDetailCustomization_ResizeLandscape::GetOriginalComponentCount()
 {
 	FEdModeLandscape* LandscapeEdMode = GetEditorMode();
-	if (LandscapeEdMode != NULL)
+	if (LandscapeEdMode != nullptr)
 	{
 		return FText::Format(LOCTEXT("NxN", "{0}\u00D7{1}"),
 			FText::AsNumber(LandscapeEdMode->UISettings->ResizeLandscape_Original_ComponentCount.X),
@@ -446,7 +447,7 @@ FText FLandscapeEditorDetailCustomization_ResizeLandscape::GetComponentCount(TSh
 FText FLandscapeEditorDetailCustomization_ResizeLandscape::GetOriginalLandscapeResolution()
 {
 	FEdModeLandscape* LandscapeEdMode = GetEditorMode();
-	if (LandscapeEdMode != NULL)
+	if (LandscapeEdMode != nullptr)
 	{
 		const int32 Original_ComponentSizeQuads = LandscapeEdMode->UISettings->ResizeLandscape_Original_SectionsPerComponent * LandscapeEdMode->UISettings->ResizeLandscape_Original_QuadsPerSection;
 		return FText::Format(LOCTEXT("NxN", "{0}\u00D7{1}"),
@@ -460,7 +461,7 @@ FText FLandscapeEditorDetailCustomization_ResizeLandscape::GetOriginalLandscapeR
 FText FLandscapeEditorDetailCustomization_ResizeLandscape::GetLandscapeResolution()
 {
 	FEdModeLandscape* LandscapeEdMode = GetEditorMode();
-	if (LandscapeEdMode != NULL)
+	if (LandscapeEdMode != nullptr)
 	{
 		const int32 ComponentSizeQuads = LandscapeEdMode->UISettings->ResizeLandscape_SectionsPerComponent * LandscapeEdMode->UISettings->ResizeLandscape_QuadsPerSection;
 		return FText::Format(LOCTEXT("NxN", "{0}\u00D7{1}"),
@@ -474,7 +475,7 @@ FText FLandscapeEditorDetailCustomization_ResizeLandscape::GetLandscapeResolutio
 FText FLandscapeEditorDetailCustomization_ResizeLandscape::GetOriginalTotalComponentCount()
 {
 	FEdModeLandscape* LandscapeEdMode = GetEditorMode();
-	if (LandscapeEdMode != NULL)
+	if (LandscapeEdMode != nullptr)
 	{
 		return FText::AsNumber(LandscapeEdMode->UISettings->ResizeLandscape_Original_ComponentCount.X * LandscapeEdMode->UISettings->ResizeLandscape_Original_ComponentCount.Y);
 	}
@@ -485,7 +486,7 @@ FText FLandscapeEditorDetailCustomization_ResizeLandscape::GetOriginalTotalCompo
 FText FLandscapeEditorDetailCustomization_ResizeLandscape::GetTotalComponentCount()
 {
 	FEdModeLandscape* LandscapeEdMode = GetEditorMode();
-	if (LandscapeEdMode != NULL)
+	if (LandscapeEdMode != nullptr)
 	{
 		return FText::AsNumber(LandscapeEdMode->UISettings->ResizeLandscape_ComponentCount.X * LandscapeEdMode->UISettings->ResizeLandscape_ComponentCount.Y);
 	}
@@ -496,7 +497,7 @@ FText FLandscapeEditorDetailCustomization_ResizeLandscape::GetTotalComponentCoun
 FReply FLandscapeEditorDetailCustomization_ResizeLandscape::OnApplyButtonClicked()
 {
 	FEdModeLandscape* LandscapeEdMode = GetEditorMode();
-	if (LandscapeEdMode != NULL)
+	if (LandscapeEdMode != nullptr)
 	{
 		FScopedTransaction Transaction(LOCTEXT("Undo", "Changing Landscape Component Size"));
 
