Landscape UX improvements
* Updated icons to landscape target layers display order methods
* Added sort type (ascending/descending) for this
* Cleanup of unused icons
* The sort/filter options have moved next to filter box within the Layers property (since they affect only the layers that are visible, these buttons should not be present when the property is collapsed)
* Added number of target layers in the Layers property header, in order to be similar to standard arrays and also the edit layers property above

diff --git a/Engine/Source/Editor/EditorStyle/Private/StarshipStyle.cpp b/Engine/Source/Editor/EditorStyle/Private/StarshipStyle.cpp
--- a/Engine/Source/Editor/EditorStyle/Private/StarshipStyle.cpp
+++ b/Engine/Source/Editor/EditorStyle/Private/StarshipStyle.cpp
@@ -7204,14 +7204,6 @@ void FStarshipEditorStyle::FStyle::SetupLandscapeEditorStyle()
 #if WITH_EDITOR
 	// Landscape Editor
 	{
-		// Modes
-		Set("LandscapeEditor.ManageMode", new IMAGE_BRUSH("Icons/icon_Landscape_Mode_Manage_40x", Icon40x40));
-		Set("LandscapeEditor.SculptMode", new IMAGE_BRUSH("Icons/icon_Landscape_Mode_Sculpt_40x", Icon40x40));
-		Set("LandscapeEditor.PaintMode",  new IMAGE_BRUSH("Icons/icon_Landscape_Mode_Paint_40x",  Icon40x40));
-		Set("LandscapeEditor.ManageMode.Small", new IMAGE_BRUSH("Icons/icon_Landscape_Mode_Manage_20x", Icon20x20));
-		Set("LandscapeEditor.SculptMode.Small", new IMAGE_BRUSH("Icons/icon_Landscape_Mode_Sculpt_20x", Icon20x20));
-		Set("LandscapeEditor.PaintMode.Small",  new IMAGE_BRUSH("Icons/icon_Landscape_Mode_Paint_20x",  Icon20x20));
-
 		{
 			// Tools
 			Set("LandscapeEditor.NewLandscape",					new IMAGE_BRUSH("Icons/Landscape/Landscape_NewLandscape_x40",		Icon20x20));
@@ -7285,7 +7277,6 @@ void FStarshipEditorStyle::FStyle::SetupLandscapeEditorStyle()
 			Set("LandscapeEditor.CircleBrush_Spherical.Small",	new IMAGE_BRUSH("Icons/Landscape/Landscape_FalloffSpherical_x40",  Icon20x20));
 			Set("LandscapeEditor.CircleBrush_Tip",      		new IMAGE_BRUSH("Icons/Landscape/Landscape_FalloffTip_x40",        Icon20x20));
 			Set("LandscapeEditor.CircleBrush_Tip.Small",      	new IMAGE_BRUSH("Icons/Landscape/Landscape_FalloffTip_x40",        Icon20x20));
-
 		}
 
 		Set("LandscapeEditor.Brush.AffectsHeightmap", new IMAGE_BRUSH("Icons/icon_Landscape_Affects_Heightmap_16x", Icon16x16));
@@ -7311,34 +7302,27 @@ void FStarshipEditorStyle::FStyle::SetupLandscapeEditorStyle()
 		Set("LandscapeEditor.TargetList.RowSelected",          new BOX_BRUSH("Common/RoundedSelection_16x", 4.0f / 16.0f, FStyleColors::Select));
 		Set("LandscapeEditor.TargetList.RowSelectedHovered",   new BOX_BRUSH("Common/RoundedSelection_16x", 4.0f / 16.0f, FStyleColors::PrimaryHover));
 
-		Set("LandscapeEditor.Target_Heightmap",  new IMAGE_BRUSH_SVG("Icons/icon_Landscape_Target_Heightmap_64x",  Icon48x48));
-		Set("LandscapeEditor.Target_Weightmap",	 new IMAGE_BRUSH_SVG("Icons/icon_Landscape_Target_Weightmap_64x",  Icon48x48));
-		Set("LandscapeEditor.Target_Visibility", new IMAGE_BRUSH_SVG("Icons/icon_Landscape_Target_Visibility_64x", Icon48x48));
-		Set("LandscapeEditor.Target_Unknown",    new IMAGE_BRUSH_SVG("Icons/icon_Landscape_Target_Unknown_64x",    Icon48x48));
+		Set("LandscapeEditor.Target_Heightmap",  new IMAGE_BRUSH_SVG("Icons/Landscape/Landscape_Target_Heightmap_64x",  Icon48x48));
+		Set("LandscapeEditor.Target_Weightmap",	 new IMAGE_BRUSH_SVG("Icons/Landscape/Landscape_Target_Weightmap_64x",  Icon48x48));
+		Set("LandscapeEditor.Target_Visibility", new IMAGE_BRUSH_SVG("Icons/Landscape/Landscape_Target_Visibility_64x", Icon48x48));
+		Set("LandscapeEditor.Target_Unknown",    new IMAGE_BRUSH_SVG("Icons/Landscape/Landscape_Target_Unknown_64x",    Icon48x48));
 
 		Set("LandscapeEditor.Target_Create",     new IMAGE_BRUSH("Icons/PlusSymbol_12x", Icon12x12));
 		Set("LandscapeEditor.Target_Delete",     new IMAGE_BRUSH("Icons/Cross_12x", Icon12x12));
 
-		Set("LandscapeEditor.Target_DisplayOrder.Default", new IMAGE_BRUSH("Icons/icon_landscape_sort_base", Icon16x16));
-		Set("LandscapeEditor.Target_DisplayOrder.Alphabetical", new IMAGE_BRUSH("Icons/icon_landscape_sort_alphabetical", Icon16x16));
-		Set("LandscapeEditor.Target_DisplayOrder.Custom", new IMAGE_BRUSH("Icons/icon_landscape_sort_custom", Icon16x16));
-		Set("LandscapeEditor.Target_DisplayOrder.ByBlendMethod", new IMAGE_BRUSH("Icons/icon_landscape_sort_byblendmethod", Icon16x16));
+		Set("LandscapeEditor.Target_DisplayOrder.Default.Ascending", new IMAGE_BRUSH_SVG("Icons/Landscape/Landscape_Target_DisplayOrder_Default_Ascending_16x", Icon16x16));
+		Set("LandscapeEditor.Target_DisplayOrder.Default.Descending", new IMAGE_BRUSH_SVG("Icons/Landscape/Landscape_Target_DisplayOrder_Default_Descending_16x", Icon16x16));
+		Set("LandscapeEditor.Target_DisplayOrder.Alphabetical.Ascending", new IMAGE_BRUSH_SVG("Icons/Landscape/Landscape_Target_DisplayOrder_Alphabetical_Ascending_16x", Icon16x16));
+		Set("LandscapeEditor.Target_DisplayOrder.Alphabetical.Descending", new IMAGE_BRUSH_SVG("Icons/Landscape/Landscape_Target_DisplayOrder_Alphabetical_Descending_16x", Icon16x16));
+		Set("LandscapeEditor.Target_DisplayOrder.ByBlendMethod.Ascending", new IMAGE_BRUSH_SVG("Icons/Landscape/Landscape_Target_DisplayOrder_ByBlendMethod_Ascending_16x", Icon16x16));
+		Set("LandscapeEditor.Target_DisplayOrder.ByBlendMethod.Descending", new IMAGE_BRUSH_SVG("Icons/Landscape/Landscape_Target_DisplayOrder_ByBlendMethod_Descending_16x", Icon16x16));
+		Set("LandscapeEditor.Target_DisplayOrder.Custom", new IMAGE_BRUSH_SVG("Icons/Landscape/Landscape_Target_DisplayOrder_Custom_16x", Icon16x16));
 
 		Set("LandscapeEditor.TargetList.DropZone.Above", new BOX_BRUSH("Common/VerticalBoxDropZoneIndicator_Above", FMargin(10.0f / 16.0f, 10.0f / 16.0f, 0, 0), SelectionColor_Subdued));
 		Set("LandscapeEditor.TargetList.DropZone.Below", new BOX_BRUSH("Common/VerticalBoxDropZoneIndicator_Below", FMargin(10.0f / 16.0f, 0, 0, 10.0f / 16.0f), SelectionColor_Subdued));
 
-		Set("LandscapeEditor.Layer.Sync", new IMAGE_BRUSH_SVG("Icons/icon_Landscape_Layers_From_Material_16x", Icon16x16));
-		Set("LandscapeEditor.Layer.AutoFill", new IMAGE_BRUSH_SVG("Icons/icon_Landscape_Layers_AutoFill_16x", Icon16x16));
-
-		Set("LandscapeEditor.SpinBox", FSpinBoxStyle(GetWidgetStyle<FSpinBoxStyle>("SpinBox"))
-			.SetTextPadding(FMargin(0.f))
-			.SetBackgroundBrush(FSlateNoResource())
-			.SetHoveredBackgroundBrush(FSlateNoResource())
-			.SetInactiveFillBrush(FSlateNoResource())
-			.SetActiveFillBrush(FSlateNoResource())
-			.SetForegroundColor(FSlateColor::UseForeground())
-			.SetArrowsImage(FSlateNoResource())
-		);
+		Set("LandscapeEditor.Layer.Sync", new IMAGE_BRUSH_SVG("Icons/Landscape/Landscape_Layers_From_Material_16x", Icon16x16));
+		Set("LandscapeEditor.Layer.AutoFill", new IMAGE_BRUSH_SVG("Icons/Landscape/Landscape_Layers_AutoFill_16x", Icon16x16));
 }
 
 #endif
diff --git a/Engine/Source/Editor/LandscapeEditor/Private/LandscapeEdMode.cpp b/Engine/Source/Editor/LandscapeEditor/Private/LandscapeEdMode.cpp
--- a/Engine/Source/Editor/LandscapeEditor/Private/LandscapeEdMode.cpp
+++ b/Engine/Source/Editor/LandscapeEditor/Private/LandscapeEdMode.cpp
@@ -2668,7 +2668,7 @@ void FEdModeLandscape::UpdateTargetList(bool bRegenerateThumbnails)
 
 			UISettings->TargetDisplayOrder = Landscape->TargetDisplayOrder;
 
-			UpdateTargetLayerDisplayOrder(UISettings->TargetDisplayOrder);
+			UpdateTargetLayerDisplayOrder(UISettings->TargetDisplayOrder, UISettings->bTargetDisplayOrdersAscending);
 		}
 	}
 
@@ -2677,7 +2677,7 @@ void FEdModeLandscape::UpdateTargetList(bool bRegenerateThumbnails)
 	UpdateLayerUsageInformation();
 }
 
-void FEdModeLandscape::UpdateTargetLayerDisplayOrder(ELandscapeLayerDisplayMode InTargetDisplayOrder)
+void FEdModeLandscape::UpdateTargetLayerDisplayOrder(ELandscapeLayerDisplayMode InTargetDisplayOrder, bool bInIsAscendingOrder)
 {
 	if (!CurrentToolTarget.LandscapeInfo.IsValid())
 	{
@@ -2846,6 +2846,11 @@ void FEdModeLandscape::UpdateTargetLayerDisplayOrder(ELandscapeLayerDisplayMode
 		break;
 	}	
 
+	if (!bInIsAscendingOrder)
+	{
+		Algo::Reverse(SavedTargetNameList);
+	}
+
 	if (DetailPanelRefreshRequired)
 	{
 		RefreshDetailPanel();
diff --git a/Engine/Source/Editor/LandscapeEditor/Private/LandscapeEdMode.h b/Engine/Source/Editor/LandscapeEditor/Private/LandscapeEdMode.h
--- a/Engine/Source/Editor/LandscapeEditor/Private/LandscapeEdMode.h
+++ b/Engine/Source/Editor/LandscapeEditor/Private/LandscapeEdMode.h
@@ -533,7 +533,7 @@ public:
 	bool CanEditCurrentTarget(FText* Reason = nullptr) const;
 
 	/** Update Display order list */
-	void UpdateTargetLayerDisplayOrder(ELandscapeLayerDisplayMode InTargetDisplayOrder);
+	void UpdateTargetLayerDisplayOrder(ELandscapeLayerDisplayMode InTargetDisplayOrder, bool bInIsAscendingOrder);
 	void MoveTargetLayerDisplayOrder(int32 IndexToMove, int32 IndexToDestination);
 
 	void RequestUpdateLayerUsageInformation();
diff --git a/Engine/Source/Editor/LandscapeEditor/Private/LandscapeEditorDetailCustomization_EditLayersStack.cpp b/Engine/Source/Editor/LandscapeEditor/Private/LandscapeEditorDetailCustomization_EditLayersStack.cpp
--- a/Engine/Source/Editor/LandscapeEditor/Private/LandscapeEditorDetailCustomization_EditLayersStack.cpp
+++ b/Engine/Source/Editor/LandscapeEditor/Private/LandscapeEditorDetailCustomization_EditLayersStack.cpp
@@ -153,13 +153,14 @@ void FLandscapeEditorCustomNodeBuilder_Layers::GenerateHeaderRowContent(FDetailW
 				.VAlign(VAlign_Center)
 				[
 					SNew(STextBlock)
-						.Text(this, &FLandscapeEditorCustomNodeBuilder_Layers::GetNumLayersText)
+					.MinDesiredWidth(125.0f)
+					.Text(this, &FLandscapeEditorCustomNodeBuilder_Layers::GetNumLayersText)
 				]
 				+ SHorizontalBox::Slot()
 				.AutoWidth()
 				.HAlign(HAlign_Center)
 				.VAlign(VAlign_Center)
-				.Padding(2.0f, 0.0f)
+				.Padding(4.0f, 1.0f, 0.0f, 1.0f)
 				[
 					AddButton
 				]
diff --git a/Engine/Source/Editor/LandscapeEditor/Private/LandscapeEditorDetailCustomization_TargetLayers.cpp b/Engine/Source/Editor/LandscapeEditor/Private/LandscapeEditorDetailCustomization_TargetLayers.cpp
--- a/Engine/Source/Editor/LandscapeEditor/Private/LandscapeEditorDetailCustomization_TargetLayers.cpp
+++ b/Engine/Source/Editor/LandscapeEditor/Private/LandscapeEditorDetailCustomization_TargetLayers.cpp
@@ -195,6 +195,8 @@ void FLandscapeEditorDetailCustomization_TargetLayers::CustomizeDetails(IDetailL
 	TSharedRef<IPropertyHandle> PropertyHandle_TargetLayerAssetFilePath = DetailBuilder.GetProperty(GET_MEMBER_NAME_CHECKED(ULandscapeEditorObject, TargetLayerAssetFilePath));
 	TSharedRef<IPropertyHandle> PropertyHandle_TargetDisplayOrder = DetailBuilder.GetProperty(GET_MEMBER_NAME_CHECKED(ULandscapeEditorObject, TargetDisplayOrder));
 	PropertyHandle_TargetDisplayOrder->MarkHiddenByCustomization();
+	TSharedRef<IPropertyHandle> PropertyHandle_TargetDisplayOrderIsAscendingPropertyHandle = DetailBuilder.GetProperty(GET_MEMBER_NAME_CHECKED(ULandscapeEditorObject, bTargetDisplayOrdersAscending));
+	PropertyHandle_TargetDisplayOrderIsAscendingPropertyHandle->MarkHiddenByCustomization();
 
 	TSharedRef<IPropertyHandle> PropertyHandle_TargetShowUnusedLayers = DetailBuilder.GetProperty(GET_MEMBER_NAME_CHECKED(ULandscapeEditorObject, ShowUnusedLayers));
 	PropertyHandle_TargetShowUnusedLayers->MarkHiddenByCustomization();
@@ -226,7 +228,8 @@ void FLandscapeEditorDetailCustomization_TargetLayers::CustomizeDetails(IDetailL
 				.AutoWrapText(true)
 		];
 
-	TargetsCategory.AddCustomBuilder(MakeShareable(new FLandscapeEditorCustomNodeBuilder_TargetLayers(DetailBuilder.GetThumbnailPool().ToSharedRef(), PropertyHandle_TargetDisplayOrder, PropertyHandle_TargetShowUnusedLayers)));
+	TargetsCategory.AddCustomBuilder(MakeShareable(new FLandscapeEditorCustomNodeBuilder_TargetLayers(DetailBuilder.GetThumbnailPool().ToSharedRef(), 
+		PropertyHandle_TargetDisplayOrder, PropertyHandle_TargetDisplayOrderIsAscendingPropertyHandle, PropertyHandle_TargetShowUnusedLayers)));
 
 	TargetsCategory.AddCustomRow(FText())
 		.Visibility(TAttribute<EVisibility>::Create(TAttribute<EVisibility>::FGetter::CreateStatic(&FLandscapeEditorDetailCustomization_TargetLayers::GetPopulateTargetLayersInfoTipVisibility)))
@@ -397,9 +400,11 @@ FEdModeLandscape* FLandscapeEditorCustomNodeBuilder_TargetLayers::GetEditorMode(
 	return (FEdModeLandscape*)GLevelEditorModeTools().GetActiveMode(FBuiltinEditorModes::EM_Landscape);
 }
 
-FLandscapeEditorCustomNodeBuilder_TargetLayers::FLandscapeEditorCustomNodeBuilder_TargetLayers(TSharedRef<FAssetThumbnailPool> InThumbnailPool, TSharedRef<IPropertyHandle> InTargetDisplayOrderPropertyHandle, TSharedRef<IPropertyHandle> InTargetShowUnusedLayersPropertyHandle)
+FLandscapeEditorCustomNodeBuilder_TargetLayers::FLandscapeEditorCustomNodeBuilder_TargetLayers(TSharedRef<FAssetThumbnailPool> InThumbnailPool, 
+	TSharedRef<IPropertyHandle> InTargetDisplayOrderPropertyHandle, TSharedRef<IPropertyHandle> InTargetDisplayOrderIsAscendingPropertyHandle, TSharedRef<IPropertyHandle> InTargetShowUnusedLayersPropertyHandle)
 	: ThumbnailPool(InThumbnailPool)
 	, TargetDisplayOrderPropertyHandle(InTargetDisplayOrderPropertyHandle)
+	, TargetDisplayOrderIsAscendingPropertyHandle(InTargetDisplayOrderIsAscendingPropertyHandle)
 	, TargetShowUnusedLayersPropertyHandle(InTargetShowUnusedLayersPropertyHandle)
 {
 }
@@ -439,79 +444,19 @@ void FLandscapeEditorCustomNodeBuilder_TargetLayers::GenerateHeaderRowContent(FD
 		NodeRow.ValueWidget
 		[
 			SNew(SHorizontalBox)
-			
-			+ SHorizontalBox::Slot()
-			.AutoWidth()
+			+SHorizontalBox::Slot()
+			.FillWidth(1) // Fill the entire width if possible
 			.VAlign(VAlign_Center)
-			.Padding(0.0f, 0.0f, 0.0f, 0.0f)
 			[
-				SNew(SComboButton)
-				.ComboButtonStyle(FAppStyle::Get(), "SimpleComboButtonWithIcon")
-				.ForegroundColor(FSlateColor::UseForeground())
-				.HasDownArrow(true)
-				.VAlign(VAlign_Center)
-				.HAlign(HAlign_Center)
-				.ToolTipText(LOCTEXT("TargetLayerSortButtonTooltip", "Define how we want to sort the displayed layers"))
-				.OnGetMenuContent(this, &FLandscapeEditorCustomNodeBuilder_TargetLayers::GetTargetLayerDisplayOrderButtonMenuContent)
-				.IsEnabled(this, &FLandscapeEditorCustomNodeBuilder_TargetLayers::HasWeightmapTargetLayers)
-				.ButtonContent()
-				[
-					SNew(SHorizontalBox)
-					+ SHorizontalBox::Slot()
-					.AutoWidth()
-					.VAlign(VAlign_Center)
-					[
-						SNew( SOverlay )
-						+SOverlay::Slot()
-						[
-							SNew(SImage)
-							.Image(FAppStyle::GetBrush("LandscapeEditor.Target_DisplayOrder.Default"))
-						]	
-						+SOverlay::Slot()
-						[
-							SNew(SImage)
-							.Image(this, &FLandscapeEditorCustomNodeBuilder_TargetLayers::GetTargetLayerDisplayOrderBrush)
-						]
-					]
-				]
-			]
-
-			+ SHorizontalBox::Slot()
-			.AutoWidth()
-			.VAlign(VAlign_Center)
-			.Padding(5.0f, 0.0f, 0.0f, 0.0f)
-			[
-				SNew(SComboButton)
-				.ComboButtonStyle(FAppStyle::Get(), "SimpleComboButtonWithIcon")
-				.ForegroundColor(FSlateColor::UseForeground())
-				.HasDownArrow(true)
-				.VAlign(VAlign_Center)
-				.HAlign(HAlign_Center)
-				.ToolTipText(LOCTEXT("TargetLayerUnusedLayerButtonTooltip", "Define if we want to display unused layers"))
-				.OnGetMenuContent(this, &FLandscapeEditorCustomNodeBuilder_TargetLayers::GetTargetLayerShowUnusedButtonMenuContent)
-				.IsEnabled(this, &FLandscapeEditorCustomNodeBuilder_TargetLayers::HasWeightmapTargetLayers)
-				.ButtonContent()
-				[
-					SNew(SHorizontalBox)
-					+ SHorizontalBox::Slot()
-					.AutoWidth()
-					.VAlign(VAlign_Center)
-					[
-						SNew(SBox)
-						.WidthOverride(16.0f)
-						.HeightOverride(16.0f)
-						[
-							SNew(SImage)
-								.Image(this, &FLandscapeEditorCustomNodeBuilder_TargetLayers::GetShowUnusedBrush)
-						]
-					]
-				]
+				SNew(STextBlock)
+				.MinDesiredWidth(125.f)
+				.Text(this, &FLandscapeEditorCustomNodeBuilder_TargetLayers::GetNumWeightmapTargetLayersText)
 			]
 
 			+ SHorizontalBox::Slot()
 			.AutoWidth()
 			.VAlign(VAlign_Center)
-			.Padding(5.0f, 0.0f, 0.0f, 0.0f)
+			.Padding(4.0f, 1.0f, 0.0f, 1.0f)
 			[
 				PropertyCustomizationHelpers::MakeAddButton(FSimpleDelegate::CreateRaw(this, &FLandscapeEditorCustomNodeBuilder_TargetLayers::HandleCreateLayer), NSLOCTEXT("Landscape", "CreateLayer", "Create Layer"))
 			]
@@ -519,7 +464,7 @@ void FLandscapeEditorCustomNodeBuilder_TargetLayers::GenerateHeaderRowContent(FD
 			+ SHorizontalBox::Slot()
 			.AutoWidth()
 			.VAlign(VAlign_Center)
-			.Padding(5.0f, 0.0f, 0.0f, 0.0f)
+			.Padding(4.0f, 1.0f, 0.0f, 1.0f)
 			[
 				SNew(SButton)
 				.ButtonStyle( FAppStyle::Get(), "SimpleButton" )
@@ -535,7 +480,7 @@ void FLandscapeEditorCustomNodeBuilder_TargetLayers::GenerateHeaderRowContent(FD
 			+ SHorizontalBox::Slot()
 			.AutoWidth()
 			.VAlign(VAlign_Center)
-			.Padding(5.0f, 0.0f, 0.0f, 0.0f)
+			.Padding(4.0f, 1.0f, 0.0f, 1.0f)
 			[
 				SNew(SButton)
 				.ButtonStyle(FAppStyle::Get(), "SimpleButton")
@@ -544,13 +489,25 @@ void FLandscapeEditorCustomNodeBuilder_TargetLayers::GenerateHeaderRowContent(FD
 				.IsEnabled(this, &FLandscapeEditorCustomNodeBuilder_TargetLayers::HasWeightmapTargetLayers)
 				[
 					SNew(SImage)
-						.Image(FAppStyle::GetBrush("LandscapeEditor.Layer.AutoFill"))
+					.Image(FAppStyle::GetBrush("LandscapeEditor.Layer.AutoFill"))
 				]
 			]
 		];
 	}
 }
 
+FText FLandscapeEditorCustomNodeBuilder_TargetLayers::GetNumWeightmapTargetLayersText() const
+{
+	FEdModeLandscape* LandscapeEdMode = GetEditorMode();
+	ALandscape* Landscape = LandscapeEdMode ? LandscapeEdMode->GetLandscape() : nullptr;
+	if (Landscape)
+	{
+		return FText::Format(LOCTEXT("NumTargetLayersText", "{0} Target {0}|plural(one=Layer, other=Layers)"), GetWeightmapTargetLayerCount());
+	}
+
+	return FText();
+}
+
 const TArray<const TSharedRef<FLandscapeTargetListInfo>> FLandscapeEditorCustomNodeBuilder_TargetLayers::GetUnassignedTargetLayersFromMaterial() const
 {
 	TArray<const TSharedRef<FLandscapeTargetListInfo>> TargetLayerList;
@@ -773,7 +730,7 @@ TSharedRef<SWidget> FLandscapeEditorCustomNodeBuilder_TargetLayers::GetTargetLay
 {
 	FMenuBuilder MenuBuilder(/*bInShouldCloseWindowAfterMenuSelection=*/true, nullptr, nullptr, /*bCloseSelfOnly=*/ true);
 
-	MenuBuilder.BeginSection("TargetLayerSortType", LOCTEXT("SortTypeHeading", "Sort Type"));
+	MenuBuilder.BeginSection("TargetLayerSortBy", LOCTEXT("SortByHeading", "Sort By"));
 	{
 		MenuBuilder.AddMenuEntry(
 			LOCTEXT("TargetLayerDisplayOrderDefault", "Default"),
@@ -829,6 +786,36 @@ TSharedRef<SWidget> FLandscapeEditorCustomNodeBuilder_TargetLayers::GetTargetLay
 	}
 	MenuBuilder.EndSection();
 
+	MenuBuilder.BeginSection("TargetLayerSortType", LOCTEXT("SortTypeHeading", "Sort Type"));
+	{
+		MenuBuilder.AddMenuEntry(
+			LOCTEXT("TargetLayerDisplayOrderAscending", "Ascending"),
+			LOCTEXT("TargetLayerDisplayOrderAscendingToolTip", "Sort the items in Ascending order"),
+			FSlateIcon(),
+			FUIAction(
+				FExecuteAction::CreateSP(this, &FLandscapeEditorCustomNodeBuilder_TargetLayers::SetAscendingDisplayOrder, true),
+				FCanExecuteAction::CreateSP(this, &FLandscapeEditorCustomNodeBuilder_TargetLayers::CanChangeDisplayOrderSortType),
+				FIsActionChecked::CreateSP(this, &FLandscapeEditorCustomNodeBuilder_TargetLayers::IsAscendingDisplayOrder)
+			),
+			NAME_None,
+			EUserInterfaceActionType::RadioButton
+		);
+
+		MenuBuilder.AddMenuEntry(
+			LOCTEXT("TargetLayerDisplayOrderDescending", "Descending"),
+			LOCTEXT("TargetLayerDisplayOrderDescendingToolTip", "Sort the items in Descending order"),
+			FSlateIcon(),
+			FUIAction(
+				FExecuteAction::CreateSP(this, &FLandscapeEditorCustomNodeBuilder_TargetLayers::SetAscendingDisplayOrder, false),
+				FCanExecuteAction::CreateSP(this, &FLandscapeEditorCustomNodeBuilder_TargetLayers::CanChangeDisplayOrderSortType),
+				FIsActionChecked::CreateSP(this, &FLandscapeEditorCustomNodeBuilder_TargetLayers::IsDescendingDisplayOrder)
+			),
+			NAME_None,
+			EUserInterfaceActionType::RadioButton
+		);
+	}
+	MenuBuilder.EndSection();
+
 	return MenuBuilder.MakeWidget();
 }
 
@@ -900,6 +887,22 @@ void FLandscapeEditorCustomNodeBuilder_TargetLayers::SetSelectedDisplayOrder(ELa
 	TargetDisplayOrderPropertyHandle->SetValue((uint8)InDisplayOrder);	
 }
 
+void FLandscapeEditorCustomNodeBuilder_TargetLayers::SetAscendingDisplayOrder(bool bInIsAscending)
+{
+	TargetDisplayOrderIsAscendingPropertyHandle->SetValue(bInIsAscending);
+}
+
+bool FLandscapeEditorCustomNodeBuilder_TargetLayers::CanChangeDisplayOrderSortType() const 
+{
+	FEdModeLandscape* LandscapeEdMode = GetEditorMode();
+	if (LandscapeEdMode != nullptr)
+	{
+		return LandscapeEdMode->UISettings->TargetDisplayOrder != ELandscapeLayerDisplayMode::UserSpecific;
+	}
+
+	return false;
+}
+
 bool FLandscapeEditorCustomNodeBuilder_TargetLayers::IsSelectedDisplayOrder(ELandscapeLayerDisplayMode InDisplayOrder) const
 {
 	FEdModeLandscape* LandscapeEdMode = GetEditorMode();
@@ -911,18 +914,46 @@ bool FLandscapeEditorCustomNodeBuilder_TargetLayers::IsSelectedDisplayOrder(ELan
 	return false;
 }
 
+bool FLandscapeEditorCustomNodeBuilder_TargetLayers::IsAscendingDisplayOrder() const
+{
+	FEdModeLandscape* LandscapeEdMode = GetEditorMode();
+	if (LandscapeEdMode != nullptr)
+	{
+		return LandscapeEdMode->UISettings->bTargetDisplayOrdersAscending;
+	}
+
+	return false;
+}
+
+bool FLandscapeEditorCustomNodeBuilder_TargetLayers::IsDescendingDisplayOrder() const
+{
+	FEdModeLandscape* LandscapeEdMode = GetEditorMode();
+	if (LandscapeEdMode != nullptr)
+	{
+		return !LandscapeEdMode->UISettings->bTargetDisplayOrdersAscending;
+	}
+
+	return false;
+}
+
 const FSlateBrush* FLandscapeEditorCustomNodeBuilder_TargetLayers::GetTargetLayerDisplayOrderBrush() const
 {
 	FEdModeLandscape* LandscapeEdMode = GetEditorMode();
 	if (LandscapeEdMode != nullptr)
 	{
+		auto AppendDisplayOrderSuffix = [bTargetDisplayOrdersAscending = LandscapeEdMode->UISettings->bTargetDisplayOrdersAscending](const TCHAR* InBaseBrushName) -> FName
+			{
+				return FName(FString::Format(TEXT("{0}.{1}"), { InBaseBrushName, bTargetDisplayOrdersAscending ? TEXT("Ascending") : TEXT("Descending")}));
+			};
+
 		switch (LandscapeEdMode->UISettings->TargetDisplayOrder)
 		{
-			case ELandscapeLayerDisplayMode::Alphabetical: return FAppStyle::Get().GetBrush("LandscapeEditor.Target_DisplayOrder.Alphabetical");
+			case ELandscapeLayerDisplayMode::Default: return FAppStyle::Get().GetBrush(AppendDisplayOrderSuffix(TEXT("LandscapeEditor.Target_DisplayOrder.Default")));
+			case ELandscapeLayerDisplayMode::Alphabetical: return FAppStyle::Get().GetBrush(AppendDisplayOrderSuffix(TEXT("LandscapeEditor.Target_DisplayOrder.Alphabetical")));
 			case ELandscapeLayerDisplayMode::UserSpecific: return FAppStyle::Get().GetBrush("LandscapeEditor.Target_DisplayOrder.Custom");
-			case ELandscapeLayerDisplayMode::ByBlendMethod: return FAppStyle::Get().GetBrush("LandscapeEditor.Target_DisplayOrder.ByBlendMethod");
+			case ELandscapeLayerDisplayMode::ByBlendMethod: return FAppStyle::Get().GetBrush(AppendDisplayOrderSuffix(TEXT("LandscapeEditor.Target_DisplayOrder.ByBlendMethod")));
 			default:
-				return nullptr;
+				checkNoEntry();
 		}
 	}
 
@@ -992,6 +1023,21 @@ EVisibility FLandscapeEditorCustomNodeBuilder_TargetLayers::GetLayersFilterVisib
 	return HasWeightmapTargetLayers() ? EVisibility::Visible : EVisibility::Collapsed;
 }
 
+EVisibility FLandscapeEditorCustomNodeBuilder_TargetLayers::GetLayersDisplayOptionsVisibility() const
+{
+	if (const FEdModeLandscape* LandscapeEdMode = GetEditorMode(); (LandscapeEdMode != nullptr) && LandscapeEdMode->CurrentToolMode)
+	{
+		// ELandscapeToolTargetType::Invalid means "weightmap with no valid paint layer" so we still want to display that property if it has been marked to be displayed in Weightmap target type, to be consistent 
+		if ((LandscapeEdMode->CurrentToolTarget.TargetType == ELandscapeToolTargetType::Weightmap)
+			|| (LandscapeEdMode->CurrentToolTarget.TargetType == ELandscapeToolTargetType::Invalid))
+		{
+			return EVisibility::Visible;
+		}
+	}
+
+	return EVisibility::Collapsed;
+}
+
 FText FLandscapeEditorCustomNodeBuilder_TargetLayers::GetLayersFilterText() const
 {
 	if (FEdModeLandscape* LandscapeEdMode = GetEditorMode())
@@ -1057,22 +1103,96 @@ void FLandscapeEditorCustomNodeBuilder_TargetLayers::GenerateChildContent(IDetai
 					.AutoHeight()
 					.VAlign(VAlign_Top)
 					.HAlign(HAlign_Fill)
-					.Padding(FMargin(2))
+					.Padding(FMargin(0.0f, 2.0f, 0.0f, 2.0f))
 					[
-						SAssignNew(LayersFilterSearchBox, SSearchBox)
+						SNew(SHorizontalBox)
+						+ SHorizontalBox::Slot()
+						.AutoWidth()
+						.VAlign(VAlign_Center)
+						.Padding(4.0f, 1.0f, 0.0f, 1.0f)
+						[
+							SNew(SComboButton)
+							.ComboButtonStyle(FAppStyle::Get(), "SimpleComboButtonWithIcon")
+							.ForegroundColor(FSlateColor::UseForeground())
+							.HasDownArrow(true)
+							.VAlign(VAlign_Center)
+							.HAlign(HAlign_Center)
+							.ToolTipText(LOCTEXT("TargetLayerSortButtonTooltip", "Define how we want to sort the displayed layers"))
+							.OnGetMenuContent(this, &FLandscapeEditorCustomNodeBuilder_TargetLayers::GetTargetLayerDisplayOrderButtonMenuContent)
+							.IsEnabled(this, &FLandscapeEditorCustomNodeBuilder_TargetLayers::HasWeightmapTargetLayers)
+							.Visibility(this, &FLandscapeEditorCustomNodeBuilder_TargetLayers::GetLayersDisplayOptionsVisibility)
+							.ButtonContent()
+							[
+								SNew(SHorizontalBox)
+								+ SHorizontalBox::Slot()
+								.AutoWidth()
+								.VAlign(VAlign_Center)
+								[
+									SNew(SBox)
+									.WidthOverride(16.0f)
+									.HeightOverride(16.0f)
+									[
+										SNew(SImage)
+										.Image(this, &FLandscapeEditorCustomNodeBuilder_TargetLayers::GetTargetLayerDisplayOrderBrush)
+									]
+								]
+							]
+						]
+
+						+ SHorizontalBox::Slot()
+						.AutoWidth()
+						.VAlign(VAlign_Center)
+						.Padding(4.0f, 1.0f, 0.0f, 1.0f)
+						[
+							SNew(SComboButton)
+							.ComboButtonStyle(FAppStyle::Get(), "SimpleComboButtonWithIcon")
+							.ForegroundColor(FSlateColor::UseForeground())
+							.HasDownArrow(true)
+							.VAlign(VAlign_Center)
+							.HAlign(HAlign_Center)
+							.ToolTipText(LOCTEXT("TargetLayerUnusedLayerButtonTooltip", "Define if we want to display unused layers"))
+							.OnGetMenuContent(this, &FLandscapeEditorCustomNodeBuilder_TargetLayers::GetTargetLayerShowUnusedButtonMenuContent)
+							.IsEnabled(this, &FLandscapeEditorCustomNodeBuilder_TargetLayers::HasWeightmapTargetLayers)
+							.Visibility(this, &FLandscapeEditorCustomNodeBuilder_TargetLayers::GetLayersDisplayOptionsVisibility)
+							.ButtonContent()
+							[
+								SNew(SHorizontalBox)
+								+ SHorizontalBox::Slot()
+								.AutoWidth()
+								.VAlign(VAlign_Center)
+								[
+									SNew(SBox)
+									.WidthOverride(16.0f)
+									.HeightOverride(16.0f)
+									[
+										SNew(SImage)
+										.Image(this, &FLandscapeEditorCustomNodeBuilder_TargetLayers::GetShowUnusedBrush)
+									]
+								]
+							]
+						]
+
+						+ SHorizontalBox::Slot()
+						//.AutoWidth()
+						.VAlign(VAlign_Center)
+						.HAlign(HAlign_Fill)
+						.Padding(5.0f, 0.0f, 0.0f, 0.0f)
+						[
+							SAssignNew(LayersFilterSearchBox, SSearchBox)
 							.InitialText(this, &FLandscapeEditorCustomNodeBuilder_TargetLayers::GetLayersFilterText)
 							.SelectAllTextWhenFocused(true)
 							.HintText(LOCTEXT("LayersSearch", "Filter Target Layers"))
 							.OnTextChanged(this, &FLandscapeEditorCustomNodeBuilder_TargetLayers::OnFilterTextChanged)
 							.OnTextCommitted(this, &FLandscapeEditorCustomNodeBuilder_TargetLayers::OnFilterTextCommitted)
 							.Visibility(this, &FLandscapeEditorCustomNodeBuilder_TargetLayers::GetLayersFilterVisibility)
+						]
 					]
 
 					+ SVerticalBox::Slot()
 					.AutoHeight()
 					.VAlign(VAlign_Top)
 					.HAlign(HAlign_Fill)
-					.Padding(0.0f, 0.0f)
+					.Padding(2)
 					[
 						TargetLayerList.ToSharedRef()
 					]
@@ -1190,14 +1310,14 @@ TSharedPtr<SWidget> FLandscapeEditorCustomNodeBuilder_TargetLayers::GenerateRow(
 					.Padding(FMargin(2))
 					[
 						SNew(SBox)
-							.Visibility_Static(&FLandscapeEditorCustomNodeBuilder_TargetLayers::GetDebugModeLayerUsageVisibility, Target)
-							.WidthOverride(48.0f)
-							.HeightOverride(48.0f)
-							[
-								SNew(SImage)
-									.Image(FCoreStyle::Get().GetBrush("WhiteBrush"))
-									.ColorAndOpacity_Static(&FLandscapeEditorCustomNodeBuilder_TargetLayers::GetLayerUsageDebugColor, Target)
-							]
+						.Visibility_Static(&FLandscapeEditorCustomNodeBuilder_TargetLayers::GetDebugModeLayerUsageVisibility, Target)
+						.WidthOverride(48.0f)
+						.HeightOverride(48.0f)
+						[
+							SNew(SImage)
+							.Image(FCoreStyle::Get().GetBrush("WhiteBrush"))
+							.ColorAndOpacity_Static(&FLandscapeEditorCustomNodeBuilder_TargetLayers::GetLayerUsageDebugColor, Target)
+						]
 					]
 
 					+ SHorizontalBox::Slot()
@@ -2058,9 +2178,9 @@ EVisibility FLandscapeEditorCustomNodeBuilder_TargetLayers::GetLayersSubstractiv
 	if (LandscapeEdMode != nullptr && Target->TargetType == ELandscapeToolTargetType::Weightmap && Target->LayerInfoObj.IsValid())
 	{
 		if (ULandscapeEditLayerBase* EditLayer = LandscapeEdMode->GetEditLayer(LandscapeEdMode->GetSelectedEditLayerIndex()))
-		{
-			return EVisibility::Visible;
-		}
+	    {
+		    return EVisibility::Visible;
+	    }
 	}
 
 	return EVisibility::Collapsed;
diff --git a/Engine/Source/Editor/LandscapeEditor/Private/LandscapeEditorDetailCustomization_TargetLayers.h b/Engine/Source/Editor/LandscapeEditor/Private/LandscapeEditorDetailCustomization_TargetLayers.h
--- a/Engine/Source/Editor/LandscapeEditor/Private/LandscapeEditorDetailCustomization_TargetLayers.h
+++ b/Engine/Source/Editor/LandscapeEditor/Private/LandscapeEditorDetailCustomization_TargetLayers.h
@@ -65,7 +65,8 @@ protected:
 class FLandscapeEditorCustomNodeBuilder_TargetLayers : public IDetailCustomNodeBuilder, public TSharedFromThis<FLandscapeEditorCustomNodeBuilder_TargetLayers>
 {
 public:
-	FLandscapeEditorCustomNodeBuilder_TargetLayers(TSharedRef<FAssetThumbnailPool> ThumbnailPool, TSharedRef<IPropertyHandle> InTargetDisplayOrderPropertyHandle, TSharedRef<IPropertyHandle> InTargetShowUnusedLayersPropertyHandle);
+	FLandscapeEditorCustomNodeBuilder_TargetLayers(TSharedRef<FAssetThumbnailPool> ThumbnailPool, 
+		TSharedRef<IPropertyHandle> InTargetDisplayOrderPropertyHandle, TSharedRef<IPropertyHandle> InTargetDisplayOrderIsAscendingPropertyHandle, TSharedRef<IPropertyHandle> InTargetShowUnusedLayersPropertyHandle);
 	~FLandscapeEditorCustomNodeBuilder_TargetLayers();
 
 	virtual void SetOnRebuildChildren( FSimpleDelegate InOnRegenerateChildren ) override;
@@ -81,6 +82,7 @@ public:
 protected:
 	TSharedRef<FAssetThumbnailPool> ThumbnailPool;
 	TSharedRef<IPropertyHandle> TargetDisplayOrderPropertyHandle;
+	TSharedRef<IPropertyHandle> TargetDisplayOrderIsAscendingPropertyHandle;
 	TSharedRef<IPropertyHandle> TargetShowUnusedLayersPropertyHandle;
 	TSharedPtr<SSearchBox> LayersFilterSearchBox;
 	/** Widgets for displaying and editing the target layer name */
@@ -141,7 +143,11 @@ protected:
 	const FSlateBrush* GetTargetLayerDisplayOrderBrush() const;
 	TSharedRef<SWidget> GetTargetLayerDisplayOrderButtonMenuContent();
 	void SetSelectedDisplayOrder(ELandscapeLayerDisplayMode InDisplayOrder);
+	void SetAscendingDisplayOrder(bool bInIsAscending);
 	bool IsSelectedDisplayOrder(ELandscapeLayerDisplayMode InDisplayOrder) const;
+	bool IsAscendingDisplayOrder() const;
+	bool IsDescendingDisplayOrder() const;
+	bool CanChangeDisplayOrderSortType() const;
 
 	TSharedRef<SWidget> GetTargetLayerShowUnusedButtonMenuContent();
 	const FSlateBrush* GetShowUnusedBrush() const;
@@ -154,7 +160,9 @@ protected:
 	int32 GetWeightmapTargetLayerCount() const;
 	bool HasWeightmapTargetLayers() const;
 	EVisibility GetLayersFilterVisibility() const;
+	EVisibility GetLayersDisplayOptionsVisibility() const;
 	FText GetLayersFilterText() const;
+	FText GetNumWeightmapTargetLayersText() const;
 };
 
 class SLandscapeEditorSelectableBorder : public SBorder
diff --git a/Engine/Source/Editor/LandscapeEditor/Private/LandscapeEditorObject.cpp b/Engine/Source/Editor/LandscapeEditor/Private/LandscapeEditorObject.cpp
--- a/Engine/Source/Editor/LandscapeEditor/Private/LandscapeEditorObject.cpp
+++ b/Engine/Source/Editor/LandscapeEditor/Private/LandscapeEditorObject.cpp
@@ -129,6 +129,7 @@ ULandscapeEditorObject::ULandscapeEditorObject(const FObjectInitializer& ObjectI
 
 	, BrushComponentSize(1)
 	, TargetDisplayOrder(ELandscapeLayerDisplayMode::Default)
+	, bTargetDisplayOrdersAscending(true)
 	, ShowUnusedLayers(true)
 {
 	// Structure to hold one-time initialization
@@ -199,8 +200,9 @@ void ULandscapeEditorObject::PostEditChangeProperty(FPropertyChangedEvent& Prope
 		UpdateComponentLayerAllowList();
 	}
 
-	if (PropertyChangedEvent.MemberProperty == nullptr ||
-		PropertyChangedEvent.MemberProperty->GetFName() == GET_MEMBER_NAME_CHECKED(ULandscapeEditorObject, TargetDisplayOrder))
+	if (PropertyChangedEvent.MemberProperty == nullptr 
+		|| (PropertyChangedEvent.MemberProperty->GetFName() == GET_MEMBER_NAME_CHECKED(ULandscapeEditorObject, TargetDisplayOrder))
+		|| (PropertyChangedEvent.MemberProperty->GetFName() == GET_MEMBER_NAME_CHECKED(ULandscapeEditorObject, bTargetDisplayOrdersAscending)))
 	{
 		UpdateTargetLayerDisplayOrder();
 	}
@@ -1003,7 +1005,7 @@ void ULandscapeEditorObject::UpdateTargetLayerDisplayOrder()
 {
 	if (ParentMode != nullptr)
 	{
-		ParentMode->UpdateTargetLayerDisplayOrder(TargetDisplayOrder);
+		ParentMode->UpdateTargetLayerDisplayOrder(TargetDisplayOrder, bTargetDisplayOrdersAscending);
 	}
 }
 
diff --git a/Engine/Source/Editor/LandscapeEditor/Public/LandscapeEditorObject.h b/Engine/Source/Editor/LandscapeEditor/Public/LandscapeEditorObject.h
--- a/Engine/Source/Editor/LandscapeEditor/Public/LandscapeEditorObject.h
+++ b/Engine/Source/Editor/LandscapeEditor/Public/LandscapeEditorObject.h
@@ -748,6 +748,9 @@ public:
 	UPROPERTY(Category = "Target Layers", EditAnywhere)
 	ELandscapeLayerDisplayMode TargetDisplayOrder;
 
+	UPROPERTY(Category = "Target Layers", EditAnywhere)
+	bool bTargetDisplayOrdersAscending;
+
 	// Show layers that are not currently being painted anywhere
 	UPROPERTY(Category = "Target Layers", EditAnywhere)
 	bool ShowUnusedLayers;	
