Landscape Editor - Added CollapseAllEditLayers function exposed via Edit Layer context menu
* Destructive collapse (Keeps final weightmap/heightmap data and removes all Blueprint brushes). Restores "Toggle Edit Layers" feature used to "flatten" edit layer landscape data into a single layer

diff --git a/Engine/Source/Editor/LandscapeEditor/Private/LandscapeEditorDetailCustomization_EditLayers.cpp b/Engine/Source/Editor/LandscapeEditor/Private/LandscapeEditorDetailCustomization_EditLayers.cpp
--- a/Engine/Source/Editor/LandscapeEditor/Private/LandscapeEditorDetailCustomization_EditLayers.cpp
+++ b/Engine/Source/Editor/LandscapeEditor/Private/LandscapeEditorDetailCustomization_EditLayers.cpp
@@ -71,7 +71,14 @@ bool FLandscapeEditLayerCustomizationCommon::CanCollapseLayer(int32 InLayerIndex
 {
 	FEdModeLandscape* LandscapeEdMode = GetEditorMode();
 	ALandscape* Landscape = LandscapeEdMode ? LandscapeEdMode->GetLandscape() : nullptr;
-	if (!Landscape || Landscape->GetEditLayersConst().Num() <= 1)
+
+	if (Landscape == nullptr)
+	{
+		OutReason = LOCTEXT("Landscape_CollapseLayer_Reason_InvalidLandscape", "Landscape is invalid");
+		return false;
+	}
+
+	if (Landscape->GetEditLayersConst().Num() <= 1)
 	{
 		OutReason = LOCTEXT("Landscape_CollapseLayer_Reason_NotEnoughLayersToCollapse", "Not enough layers to do collapse");
 		return false;
@@ -156,7 +163,7 @@ void FLandscapeEditLayerCustomizationCommon::CollapseLayer(int32 InLayerIndex)
 
 		if (Landscape && Layer && BaseLayer)
 		{
-			EAppReturnType::Type Result = FMessageDialog::Open(EAppMsgType::YesNo, FText::Format(LOCTEXT("Landscape_CollapseLayer_Message", "The layer {0} will be collapsed into layer {1}.  Continue?"), FText::FromName(Layer->GetName()), FText::FromName(BaseLayer->GetName())));
+			EAppReturnType::Type Result = FMessageDialog::Open(EAppMsgType::YesNo, FText::Format(LOCTEXT("Landscape_CollapseLayer_Message", "The layer '{0}' will be collapsed into layer '{1}'.  Continue?"), FText::FromName(Layer->GetName()), FText::FromName(BaseLayer->GetName())));
 			if (Result == EAppReturnType::Yes)
 			{
 				const FScopedTransaction Transaction(LOCTEXT("Landscape_Layers_Collapse", "Collapse Layer"));
@@ -168,6 +175,81 @@ void FLandscapeEditLayerCustomizationCommon::CollapseLayer(int32 InLayerIndex)
 	}
 }
 
+bool FLandscapeEditLayerCustomizationCommon::CanCollapseAllLayers(FText& OutReason) const
+{
+	FEdModeLandscape* LandscapeEdMode = GetEditorMode();
+	ALandscape* Landscape = LandscapeEdMode ? LandscapeEdMode->GetLandscape() : nullptr;
+	if (Landscape == nullptr)
+	{
+		OutReason = LOCTEXT("Landscape_CollapseLayer_Reason_InvalidLandscape", "Landscape is invalid");
+		return false;
+	}
+
+	if (Landscape->GetEditLayersConst().Num() <= 1)
+	{
+		OutReason = LOCTEXT("Landscape_CollapseLayer_Reason_NotEnoughLayersToCollapse", "Not enough layers to do collapse");
+		return false;
+	}
+
+	const ULandscapeEditLayerBase* BaseEditLayer = Landscape->GetEditLayerConst(0);
+	check(BaseEditLayer);
+
+	if (!BaseEditLayer->SupportsBeingCollapsedAway())
+	{
+		OutReason = FText::Format(LOCTEXT("Landscape_CollapseAllLayers_Reason_DefaultLayerDoesntSupportCollapsing", "Cannot collapse all layers onto layer '{0}'. The type of layer '{1}' doesn't support being collapsed away"),
+			FText::FromName(BaseEditLayer->GetName()), BaseEditLayer->GetClass()->GetDisplayNameText());
+		return false;
+	}
+
+	TArray<FString> InvalidEditLayerNames;
+	for (const ULandscapeEditLayerBase* EditLayer : Landscape->GetEditLayersConst())
+	{
+		check(EditLayer != nullptr);
+
+		if (EditLayer->IsLocked())
+		{
+			InvalidEditLayerNames.Add(EditLayer->GetName().ToString());
+		}
+	}
+
+	if (InvalidEditLayerNames.Num() > 0)
+	{
+		FFormatNamedArguments Args;
+		Args.Add(TEXT("Count"), InvalidEditLayerNames.Num());
+		Args.Add(TEXT("NameList"), FText::FromString(FString::Join(InvalidEditLayerNames, TEXT(", "))));
+		OutReason = FText::Format(LOCTEXT("Landscape_CollapseAllLayers_Reason_LayerLocked", 
+			"Cannot collapse all layers. '{NameList}' {Count}|plural(one=is,other=are) included in the collapse but {Count}|plural(one=is,other=are) locked."), Args);
+		return false;
+	}
+	
+	OutReason = FText::Format(LOCTEXT("Landscape_CollapseAllLayers_Reason_Collapse", "Collapse all layers into '{0}', preserving data while removing all Blueprint Brushes and hidden edit layer data."), 
+		FText::FromName(BaseEditLayer->GetName()));
+	return true;
+}
+
+void FLandscapeEditLayerCustomizationCommon::CollapseAllEditLayers()
+{
+	if (FEdModeLandscape* LandscapeEdMode = GetEditorMode())
+	{
+		ALandscape* Landscape = LandscapeEdMode->GetLandscape();
+		const ULandscapeEditLayerBase* BaseLayer = LandscapeEdMode->GetEditLayerConst(0);
+
+		if (Landscape && BaseLayer)
+		{
+			EAppReturnType::Type Result = FMessageDialog::Open(EAppMsgType::YesNo, FText::Format(LOCTEXT("Landscape_CollapseAllLayer_Message", "All edit layer data will be collapsed into layer '{0}'. Blueprint brushes and hidden edit layer data will be removed. Continue?"),
+				FText::FromName(BaseLayer->GetName())));
+
+			if (Result == EAppReturnType::Yes)
+			{
+				const FScopedTransaction Transaction(LOCTEXT("Landscape_Layers_Collapse_All", "Collapse All Layers"));
+				Landscape->CollapseAllEditLayers();
+				OnLayerSelectionChanged(0);
+				LandscapeEdMode->RefreshDetailPanel();
+			}
+		}
+	}
+}
+
 bool FLandscapeEditLayerCustomizationCommon::CanToggleVisibility(int32 InLayerIndex, FText& OutReason) const
 {
 	if (FEdModeLandscape* LandscapeEdMode = GetEditorMode())
@@ -368,6 +450,16 @@ void FLandscapeEditLayerContextMenuCustomization_Base::CustomizeContextMenu(ULan
 		CollapseLayerEntry.LabelOverride = LOCTEXT("CollapseLayer", "Collapse");
 
 		OutContextMenuEntryMap.FindOrAdd(EditLayerCategory).Entries.Add(CollapseLayerEntry);
+
+		// Collapse All layers
+		FMenuEntryParams CollapseAllLayersEntry;
+		CollapseAllLayersEntry.DirectActions = FUIAction(
+			FExecuteAction::CreateSP(this, &FLandscapeEditLayerContextMenuCustomization_Base::CollapseAllEditLayers),
+			FCanExecuteAction::CreateSPLambda(this, [this, EditLayerIndex] { FText Reason; return CanCollapseAllLayers(Reason); }));
+		CollapseAllLayersEntry.ToolTipOverride = MakeAttributeSPLambda(this, [this, EditLayerIndex] { FText Reason; CanCollapseAllLayers(Reason); return Reason; });
+		CollapseAllLayersEntry.LabelOverride = LOCTEXT("CollapseAllLayers", "Collapse All Layers");
+
+		OutContextMenuEntryMap.FindOrAdd(EditLayerCategory).Entries.Add(CollapseAllLayersEntry);
 		
 		// Visibility Section (new category)
 		const FName VisibilityCategory = "Visibility";
diff --git a/Engine/Source/Editor/LandscapeEditor/Private/LandscapeEditorDetailCustomization_EditLayers.h b/Engine/Source/Editor/LandscapeEditor/Private/LandscapeEditorDetailCustomization_EditLayers.h
--- a/Engine/Source/Editor/LandscapeEditor/Private/LandscapeEditorDetailCustomization_EditLayers.h
+++ b/Engine/Source/Editor/LandscapeEditor/Private/LandscapeEditorDetailCustomization_EditLayers.h
@@ -24,6 +24,9 @@ protected:
 	bool CanCollapseLayer(int32 InLayerIndex, FText& OutReason) const;
 	void CollapseLayer(int32 InLayerIndex);
 
+	bool CanCollapseAllLayers(FText& OutReason) const;
+	void CollapseAllEditLayers();
+
 	bool CanToggleVisibility(int32 InLayerIndex, FText& OutReason) const;
 	FReply OnToggleVisibility(int32 InLayerIndex);
 	const FSlateBrush* GetVisibilityBrushForLayer(int32 InLayerIndex) const;
diff --git a/Engine/Source/Runtime/Landscape/Classes/Landscape.h b/Engine/Source/Runtime/Landscape/Classes/Landscape.h
--- a/Engine/Source/Runtime/Landscape/Classes/Landscape.h
+++ b/Engine/Source/Runtime/Landscape/Classes/Landscape.h
@@ -502,9 +502,9 @@ PRAGMA_ENABLE_DEPRECATION_WARNINGS
 	
 	LANDSCAPE_API bool DeleteLayer(int32 InLayerIndex);
 	LANDSCAPE_API void CollapseLayer(int32 InLayerIndex);
-	UE_DEPRECATED(5.7, "Non-edit layer landscapes are deprecated. Use CollapseEditLayers to replace all existing layers with a new default layer.")
+	UE_DEPRECATED(5.7, "Non-edit layer landscapes are deprecated. Use CollapseAllEditLayers to flatten data to the base edit layer.")
 	LANDSCAPE_API void DeleteLayers();
-	LANDSCAPE_API void CollapseEditLayers();
+	LANDSCAPE_API void CollapseAllEditLayers();
 	LANDSCAPE_API void SetEditingLayer(const FGuid& InLayerGuid = FGuid());
 	LANDSCAPE_API void SetGrassUpdateEnabled(bool bInGrassUpdateEnabled);
 	LANDSCAPE_API const FGuid& GetEditingLayer() const;
diff --git a/Engine/Source/Runtime/Landscape/Private/LandscapeEditLayers.cpp b/Engine/Source/Runtime/Landscape/Private/LandscapeEditLayers.cpp
--- a/Engine/Source/Runtime/Landscape/Private/LandscapeEditLayers.cpp
+++ b/Engine/Source/Runtime/Landscape/Private/LandscapeEditLayers.cpp
@@ -1965,6 +1965,7 @@ bool ALandscape::PrepareTextureResources(bool bInWaitForStreaming)
 
 void ALandscape::DeleteUnusedLayers()
 {
+	// TODO [jared.ritchie] Check there is not a layer update pending before deleting unused layers
 	ULandscapeInfo* LandscapeInfo = GetLandscapeInfo();
 
 	if (LandscapeInfo == nullptr)
@@ -8318,23 +8319,36 @@ void ALandscape::DeleteLayers()
 {
 }
 
-void ALandscape::CollapseEditLayers()
+void ALandscape::CollapseAllEditLayers()
 {
 	if (IsTemplate())
 	{
 		return;
 	}
 
-	// Remove all existing edit layers
-	for (int32 LayerIndex = LandscapeEditLayers.Num() - 1; LayerIndex >= 0; --LayerIndex)
+	CopyOldDataToDefaultLayer();
+	SetSelectedEditLayerIndex(0);
+
+	// Delete all brushes and keep only the base edit layer
+	for (int32 EditLayerIndex = LandscapeEditLayers.Num() - 1; EditLayerIndex >= 0; --EditLayerIndex)
 	{
-		DeleteLayer(LayerIndex);
+		for (int32 BrushIndex = LandscapeEditLayers[EditLayerIndex].Brushes.Num() - 1; BrushIndex >= 0; --BrushIndex)
+		{
+			ALandscapeBlueprintBrushBase* Brush = LandscapeEditLayers[EditLayerIndex].Brushes[BrushIndex].GetBrush();
+			RemoveBrushFromLayer(EditLayerIndex, Brush);
+		}
+
+		if (EditLayerIndex > 0)
+		{
+			DeleteLayer(EditLayerIndex);
+		}
 	}
+	 
+	// Force an update before DeleteUnusedLayers runs to ensure newly copied data is not accidentally deleted
+	ForceUpdateLayersContent();
 
-	// Create a new empty layer since all landscapes require >= 1 layer
-	CreateDefaultLayer();
-	CopyOldDataToDefaultLayer();
-	check(SelectedEditLayerIndex == 0);
+	// After removing individual edit layer data, clear any empty weightmap allocations from the components' final data
+	DeleteUnusedLayers();
 }
 
 bool ALandscape::DeleteLayer(int32 InLayerIndex)
@@ -8416,12 +8430,12 @@ void ALandscape::CollapseLayer(int32 InLayerIndex)
 {
 	check(InLayerIndex >= 1);
 
-	FScopedSlowTask SlowTask(static_cast<float>(GetLandscapeInfo()->XYtoComponentMap.Num()), LOCTEXT("Landscape_CollapseLayer_SlowWork", "Collapsing Layer..."));
-	SlowTask.MakeDialog();
-
 	ULandscapeInfo* LandscapeInfo = GetLandscapeInfo();
 	check(LandscapeInfo);
 
+	FScopedSlowTask SlowTask(static_cast<float>(LandscapeInfo->XYtoComponentMap.Num()), LOCTEXT("Landscape_CollapseLayer_SlowWork", "Collapsing Layer..."));
+	SlowTask.MakeDialog();
+
 	// Only render the two layers being combined
 	TBitArray<> EditLayerVis;
 	for (int32 i = 0; i < LandscapeEditLayers.Num(); ++i)
@@ -8500,9 +8514,9 @@ void ALandscape::CollapseLayer(int32 InLayerIndex)
 
 	// Note that this code theoretically supports brushes, to handle the non-brush persistent layer data that brush-containing layers
 	// are allowed to hold. But the UI doesn't allow CollapseLayer to be used if either layer contains brushes.
-	for (int32 i = 0; i < LandscapeEditLayers[InLayerIndex].Brushes.Num(); ++i)
+	for (int32 BrushIndex = LandscapeEditLayers[InLayerIndex].Brushes.Num() - 1; BrushIndex >= 0; --BrushIndex)
 	{
-		ALandscapeBlueprintBrushBase* Brush = LandscapeEditLayers[InLayerIndex].Brushes[i].GetBrush();
+		ALandscapeBlueprintBrushBase* Brush = LandscapeEditLayers[InLayerIndex].Brushes[BrushIndex].GetBrush();
 		RemoveBrushFromLayer(InLayerIndex, Brush);
 		AddBrushToLayer(InLayerIndex - 1, Brush);
 	}
