Fix assert when deleting many landscape spline control points at once.

diff --git a/Engine/Source/Editor/LandscapeEditor/Private/LandscapeEdModeSplineTools.cpp b/Engine/Source/Editor/LandscapeEditor/Private/LandscapeEdModeSplineTools.cpp
--- a/Engine/Source/Editor/LandscapeEditor/Private/LandscapeEdModeSplineTools.cpp
+++ b/Engine/Source/Editor/LandscapeEditor/Private/LandscapeEdModeSplineTools.cpp
@@ -500,10 +500,16 @@ public:
 			Connection.Segment->DeleteSplinePoints();
 
 			// Get the control point at the *other* end of the segment and remove it from it
+			// Note: ULandscapeSplineControlPoint::UpdateSplinePoints is called for this case, but not for the "Yes, joint segments" case
+			// above.  It depends on the connected socket names only, not on the connected segment directions/values.
 			ULandscapeSplineControlPoint* OtherEnd = Connection.GetFarConnection().ControlPoint;
 			OtherEnd->Modify();
 			OtherEnd->ConnectedSegments.Remove(FLandscapeSplineConnection(Connection.Segment, 1 - Connection.End));
-			OtherEnd->UpdateSplinePoints();
+			// OtherEnd might already be DeleteSplinePoints-ed when batch-deleting multiple points.
+			if (OtherEnd->GetOuterSafe() != nullptr)
+			{
+				OtherEnd->UpdateSplinePoints();
+			}
 
 			SplinesComponent->Segments.Remove(Connection.Segment);
 
