Several fixes for the landscape add component tool :
* Properly generate extrapolated height data on each persistent edit layer, so that there's continuity on the corners on each of them
* Don't generate collisions on component creation because at this point, the height data is set to all zeroes (i.e. the minimum height, not the median, default height value) and doesn't correspond to what comes out of the the subsequent edit layers merge, whose readback would generate an identical heightmap hash, and would therefore not re-generate collisions. Delaying collision generation until the merge fixes this, since that's the first time the collision values would be valid anyway
* Prevent existing neighboring components to be dirtied when adding a new component (previously they would get dirty by the simple act of caching/interpolating their height data)
* Automatically apply splines after component creation, so that new components are immediately affected by the splines that overlap with them
* Improved visualizer for the component about to be added

diff --git a/Engine/Source/Editor/LandscapeEditor/Private/LandscapeEdModeComponentTools.cpp b/Engine/Source/Editor/LandscapeEditor/Private/LandscapeEdModeComponentTools.cpp
--- a/Engine/Source/Editor/LandscapeEditor/Private/LandscapeEdModeComponentTools.cpp
+++ b/Engine/Source/Editor/LandscapeEditor/Private/LandscapeEdModeComponentTools.cpp
@@ -543,13 +543,11 @@ class FLandscapeToolStrokeAddComponent : public FLandscapeToolStrokeBase
 public:
 	FLandscapeToolStrokeAddComponent(FEdModeLandscape* InEdMode, FEditorViewportClient* InViewportClient, const FLandscapeToolTarget& InTarget)
 		: Super(InEdMode, InViewportClient, InTarget)
-		, HeightCache(InTarget)
 	{
 	}
 
 	virtual ~FLandscapeToolStrokeAddComponent()
 	{
-		HeightCache.Flush();
 	}
 
 	virtual void Apply(FEditorViewportClient* ViewportClient, FLandscapeBrush* Brush, const ULandscapeEditorObject* UISettings, const TArray<FLandscapeToolInteractorPosition>& InteractorPositions)
@@ -575,20 +573,48 @@ public:
 			int32 ComponentIndexX1, ComponentIndexY1, ComponentIndexX2, ComponentIndexY2;
 			ALandscape::CalcComponentIndicesNoOverlap(X1, Y1, X2, Y2, LandscapeInfo->ComponentSizeQuads, ComponentIndexX1, ComponentIndexY1, ComponentIndexX2, ComponentIndexY2);
 
-			// expand the area by one vertex in each direction to ensure normals are calculated correctly
+			// expand the area by one vertex in each direction to capture the height of the corner vertices these components will attach to
 			X1 -= 1;
 			Y1 -= 1;
 			X2 += 1;
 			Y2 += 1;
 
-			TArray<uint16> Data;
-			TArray<FVector> XYOffsetData;
-			HeightCache.CacheData(X1, Y1, X2, Y2);
-			HeightCache.GetCachedData(X1, Y1, X2, Y2, Data);
-
 			UWorld* World = ViewportClient->GetScene()->GetWorld();
 			
-			TArray<ALandscapeProxy*> ModifiedProxies;
+			ALandscape* Landscape = LandscapeInfo->LandscapeActor.Get();
+			check(Landscape != nullptr);
+
+			TMap<FGuid, TArray<uint16>> PersistentEditLayerData;
+			// Copy the persistent edit layer data using FLandscapeHeightCache on the region. That will extrapolate the data for the missing components : 
+			for (const ULandscapeEditLayerBase* LandscapeEditLayer : Landscape->GetEditLayersConst())
+			{
+				check(LandscapeEditLayer != nullptr);
+				if (!LandscapeEditLayer->NeedsPersistentTextures())
+				{
+					continue;
+				}
+
+				TArray<uint16>& Data = PersistentEditLayerData.FindOrAdd(LandscapeEditLayer->GetGuid());
+
+				// No need to copy the persistent textures if they are not manually edited (e.g. fully procedural layers with persistent textures : e.g. the splines layer) :
+				//  those will have a heightmap created and initialized via AddDefaultLayerData but we won't copy any data into it :
+				if (LandscapeEditLayer->SupportsEditingTools())
+				{
+					Data.AddZeroed((1 + X2 - X1) * (1 + Y2 - Y1));
+					// Use FLandscapeEditDataInterface to cache / interpolate the existing edit layer's height data, that allows to use SetShouldDirtyPackage(false)
+					FLandscapeEditDataInterface LandscapeEdit(LandscapeInfo, LandscapeEditLayer->GetGuid(), /*bInUploadTextureChangesToGPU = */false);
+					// Ideally, we wouldn't even need to use FLandscapeEditDataInterface to simply cache the data, since that is intended to mutate the data 
+					//  but it's the class that allows extrapolating the data for missing components... At least, we can force it not to dirty the package when doing so : 
+					LandscapeEdit.SetShouldDirtyPackage(false);
+
+					int32 TempX1 = X1; // GetHeightData overwrites its input min/max x/y
+					int32 TempX2 = X2;
+					int32 TempY1 = Y1;
+					int32 TempY2 = Y2;
+					LandscapeEdit.GetHeightData(TempX1, TempY1, TempX2, TempY2, Data.GetData(), /*Stride = */0);
+				}
+			}
+
 			TArray<ULandscapeComponent*> NewComponents;
 			LandscapeInfo->Modify();
 			ULandscapeSubsystem* LandscapeSubsystem = World->GetSubsystem<ULandscapeSubsystem>();
@@ -613,59 +639,71 @@ public:
 								LandscapeProxy->SubsectionSizeQuads
 							);
 
-							TArray<FColor> HeightData;
-							const int32 ComponentVerts = (LandscapeComponent->SubsectionSizeQuads + 1) * LandscapeComponent->NumSubsections;
-							HeightData.Empty(FMath::Square(ComponentVerts));
-							HeightData.AddZeroed(FMath::Square(ComponentVerts));
-							LandscapeComponent->InitHeightmapData(HeightData, true);
-							LandscapeComponent->UpdateMaterialInstances();
-
+							// Let the component be known to the ULandscapeInfo. This needs to be done before restoring the edit layers data via FLandscapeHeightCache 
 							LandscapeInfo->XYtoComponentMap.Add(FIntPoint(ComponentIndexX, ComponentIndexY), LandscapeComponent);
 							LandscapeInfo->XYtoAddCollisionMap.Remove(FIntPoint(ComponentIndexX, ComponentIndexY));
 
-							ModifiedProxies.AddUnique(LandscapeProxy);
-						}
-					}
-				}
-			}
-
-			// Need to register to use general height data update
-			for (int32 Idx = 0; Idx < NewComponents.Num(); Idx++)
-			{
-				NewComponents[Idx]->RegisterComponent();
-			}
+							// First create the component's final heightmap, it's needed for AddDefaultLayerData... 
+							TArray<FColor> HeightData;
+							const int32 ComponentVerts = (LandscapeComponent->SubsectionSizeQuads + 1) * LandscapeComponent->NumSubsections;
+							HeightData.AddZeroed(FMath::Square(ComponentVerts));
+							LandscapeComponent->InitHeightmapData(HeightData, /*bUpdateCollision = */false);
 
-			HeightCache.SetCachedData(X1, Y1, X2, Y2, Data);
-			HeightCache.Flush();
+							// For each layer that needs some persistent height data
+							for (TPair<FGuid, TArray<uint16>> LayerDataPair : PersistentEditLayerData)
+							{
+								const FGuid& LayerGuid = LayerDataPair.Key;
+								TMap<UTexture2D*, UTexture2D*> CreatedTextures;
+								LandscapeComponent->AddDefaultLayerData(LayerGuid, { LandscapeComponent }, CreatedTextures);
 
-			if (ALandscape* Landscape = LandscapeInfo->LandscapeActor.Get())
-			{
-				Landscape->RequestLayersInitialization();
+								TArray<uint16>& CachedExtrapolatedData = LayerDataPair.Value;
+								if (!CachedExtrapolatedData.IsEmpty())
+								{
+									// Finally, copy our (previously extrapolated) edit layer data onto this new component :
+									FScopedSetLandscapeEditingLayer LayerScope(Landscape, LayerGuid);
+									// We have cached the data for the entire X1, Y1, X2, Y2 region, which is larger than this new component's area so we need
+									//  to extract the component's data into another array : 
+									TArray<uint16> ComponentExtrapolatedCachedData;
+									{
+										const int32 SrcDataStride = X2 - X1 + 1;
+										const int32 DstDataStride = ComponentVerts;
+										ComponentExtrapolatedCachedData.AddZeroed(FMath::Square(ComponentVerts));
+										uint16* DstDataPtr = ComponentExtrapolatedCachedData.GetData();
+										for (int32 Y = ComponentBase.Y; Y < ComponentBase.Y + ComponentVerts; ++Y)
+										{
+											const FIntPoint SrcDataOffset(ComponentBase.X - X1, Y - Y1);
+											check((SrcDataOffset.X >= 0) && (SrcDataOffset.Y >= 0));
+											uint16* SrcDataPtr = CachedExtrapolatedData.GetData() + SrcDataOffset.X + SrcDataOffset.Y * SrcDataStride;
+											check(SrcDataPtr + DstDataStride < SrcDataPtr + SrcDataStride);
+											FMemory::Memcpy(DstDataPtr, SrcDataPtr, DstDataStride * sizeof(uint16));
+											DstDataPtr += DstDataStride;
+										}
+									}
 
-				for (ULandscapeComponent* NewComponent : NewComponents)
-				{
-					TArray<ULandscapeComponent*> ComponentsUsingHeightmap;
-					ComponentsUsingHeightmap.Add(NewComponent);
+									FLandscapeEditDataInterface LandscapeEdit(LandscapeInfo, LayerGuid);
+									LandscapeEdit.SetHeightData(ComponentBase.X, ComponentBase.Y, ComponentBase.X + ComponentVerts - 1, ComponentBase.Y + ComponentVerts - 1, ComponentExtrapolatedCachedData.GetData(), /*InStride = */0, /*InCalcNormals = */false);
+								}
+							}
 
-					for (const ULandscapeEditLayerBase* EditLayer : Landscape->GetEditLayersConst())
-					{
-						// Since we do not share heightmap when adding new component, we will provide the required array, but they will only be used for 1 component
-						TMap<UTexture2D*, UTexture2D*> CreatedHeightmapTextures;
-						NewComponent->AddDefaultLayerData(EditLayer->GetGuid(), ComponentsUsingHeightmap, CreatedHeightmapTextures);
+							LandscapeComponent->UpdateMaterialInstances();
+							// We need to update the bounds before registering the component, otherwise, the transform verification that are done on register will fail
+							LandscapeComponent->UpdateCachedBounds();
+							LandscapeComponent->UpdateBounds();
+							LandscapeComponent->RegisterComponent();
+							// The collisions will be generated after the edit layers merge readback has been done but we still to create the new collision component here
+							//  so that it's part of the transaction and can be automatically deleted upon undo. UpdateCollisionData handles that :
+							LandscapeComponent->UpdateCollisionData();
+						}
 					}
-					
-					// Update Collision
-					NewComponent->UpdateCachedBounds();
-					NewComponent->UpdateBounds();
-					NewComponent->MarkRenderStateDirty();
 				}
-				GEngine->BroadcastOnActorMoved(Landscape);
 			}
+
+			// We need to invalidate the splines in case there's one that interacts with one of the newly added components : 
+			Landscape->RequestSplineLayerUpdate();
+			Landscape->RequestLayersInitialization();
+			GEngine->BroadcastOnActorMoved(Landscape);
 		}
 	}
-
-protected:
-	FLandscapeHeightCache HeightCache;
 };
 
 class FLandscapeToolAddComponent : public FLandscapeToolBase<FLandscapeToolStrokeAddComponent>
@@ -728,13 +766,11 @@ public:
 		{
 			const FColor LineColor = CanToolBeActivated() ? FColor(0, 255, 128) : FColor(255, 0, 64);
 
-			PDI->DrawLine(AddCollision->Corners[0], AddCollision->Corners[3], LineColor, SDPG_Foreground);
 			PDI->DrawLine(AddCollision->Corners[3], AddCollision->Corners[1], LineColor, SDPG_Foreground);
 			PDI->DrawLine(AddCollision->Corners[1], AddCollision->Corners[0], LineColor, SDPG_Foreground);
 
 			PDI->DrawLine(AddCollision->Corners[0], AddCollision->Corners[2], LineColor, SDPG_Foreground);
 			PDI->DrawLine(AddCollision->Corners[2], AddCollision->Corners[3], LineColor, SDPG_Foreground);
-			PDI->DrawLine(AddCollision->Corners[3], AddCollision->Corners[0], LineColor, SDPG_Foreground);
 		}
 	}
 
diff --git a/Engine/Source/Runtime/Landscape/Classes/Landscape.h b/Engine/Source/Runtime/Landscape/Classes/Landscape.h
--- a/Engine/Source/Runtime/Landscape/Classes/Landscape.h
+++ b/Engine/Source/Runtime/Landscape/Classes/Landscape.h
@@ -742,14 +742,9 @@ private:
 
 	FLandscapeEdModeInfo LandscapeEdModeInfo;
 
+	// TODO [jonathan.bard] : remove : this isn't really used anymore
 	UPROPERTY(Transient)
-	bool bLandscapeLayersAreInitialized;
-
-	UPROPERTY(Transient)
-	bool bLandscapeLayersForceResourceReset;
-
-	UPROPERTY(Transient)
-	bool bLandscapeLayersAreUsingLocalMerge;
+	bool bLandscapeLayersAreInitialized; 
 
 	UPROPERTY(Transient)
 	bool WasCompilingShaders;
diff --git a/Engine/Source/Runtime/Landscape/Classes/LandscapeComponent.h b/Engine/Source/Runtime/Landscape/Classes/LandscapeComponent.h
--- a/Engine/Source/Runtime/Landscape/Classes/LandscapeComponent.h
+++ b/Engine/Source/Runtime/Landscape/Classes/LandscapeComponent.h
@@ -1156,7 +1156,7 @@ public:
 
 	/** Updates collision component height data for the entire component, locking and unlocking heightmap textures
 	 */
-	void UpdateCollisionData(bool bInUpdateHeightfieldRegion = true);
+	LANDSCAPE_API void UpdateCollisionData(bool bInUpdateHeightfieldRegion = true);
 
 	/** Cumulates component's dirtied collision region that will need to be updated (used by Layer System)*/
 	void UpdateDirtyCollisionHeightData(FIntRect Region);
diff --git a/Engine/Source/Runtime/Landscape/Private/Landscape.cpp b/Engine/Source/Runtime/Landscape/Private/Landscape.cpp
--- a/Engine/Source/Runtime/Landscape/Private/Landscape.cpp
+++ b/Engine/Source/Runtime/Landscape/Private/Landscape.cpp
@@ -1706,7 +1706,6 @@ ALandscape::ALandscape(const FObjectInitializer& ObjectInitializer)
 	LayerContentUpdateModes = 0;
 	bSplineLayerUpdateRequested = false;
 	bLandscapeLayersAreInitialized = false;
-	bLandscapeLayersForceResourceReset = true;
 	LandscapeEdMode = nullptr;
 	bGrassUpdateEnabled = true;
 	bIsSpatiallyLoaded = false;
@@ -2777,6 +2776,11 @@ void ULandscapeComponent::AddLayerData(const FGuid& InLayerGuid, const FLandscap
 	Data = InData;
 }
 
+// TODO [jonathan.bard] : When removing texture sharing, this can be simplified : 
+//  - InComponentsUsingHeightmap is only ever used in 1 case : ALandscapeProxy::InitializeLayerWithEmptyContent. This is the only place where texture sharing should actually occur. 
+//   In all other places, we pass a single component to this array : the component itself
+//  - InOutCreatedHeightmapTextures isn't being used anywhere and can be removed safely as well
+//  - We should also assert or early-out in the cases where !LandscapeEditLayer->NeedsPersistentTextures() : in this case, we don't need any texture at all for this component on this edit layer
 void ULandscapeComponent::AddDefaultLayerData(const FGuid& InLayerGuid, const TArray<ULandscapeComponent*>& InComponentsUsingHeightmap, TMap<UTexture2D*, UTexture2D*>& InOutCreatedHeightmapTextures)
 {
 	Modify();
diff --git a/Engine/Source/Runtime/Landscape/Private/LandscapeEditLayers.cpp b/Engine/Source/Runtime/Landscape/Private/LandscapeEditLayers.cpp
--- a/Engine/Source/Runtime/Landscape/Private/LandscapeEditLayers.cpp
+++ b/Engine/Source/Runtime/Landscape/Private/LandscapeEditLayers.cpp
@@ -6747,7 +6747,6 @@ void ALandscape::RequestLayersInitialization(bool bInRequestContentUpdate, bool
 	}
 
 	bLandscapeLayersAreInitialized = false;
-	bLandscapeLayersForceResourceReset |= bInForceLayerResourceReset;
 	LandscapeSplinesAffectedComponents.Empty();
 
 	if (bInRequestContentUpdate)
