Landscape Editor - Updated Toggle Splines context menu feature to toggle between splines mode and the last active landscape editor mode
* Added Double-click tooltip to context menu

diff --git a/Engine/Source/Editor/LandscapeEditor/Private/LandscapeEdMode.cpp b/Engine/Source/Editor/LandscapeEditor/Private/LandscapeEdMode.cpp
--- a/Engine/Source/Editor/LandscapeEditor/Private/LandscapeEdMode.cpp
+++ b/Engine/Source/Editor/LandscapeEditor/Private/LandscapeEdMode.cpp
@@ -4,6 +4,7 @@
 
 #include "Algo/Accumulate.h"
 #include "Algo/Transform.h"
+#include "Algo/Find.h"
 #include "MaterialDomain.h"
 #include "Materials/Material.h"
 #include "SceneView.h"
@@ -2101,25 +2102,25 @@ void FEdModeLandscape::SetCurrentToolMode(FName ToolModeName, bool bRestoreCurre
 {
 	if (CurrentToolMode == nullptr || ToolModeName != CurrentToolMode->ToolModeName)
 	{
-		for (int32 i = 0; i < LandscapeToolModes.Num(); ++i)
+		if (FLandscapeToolMode* NewToolMode = Algo::FindBy(LandscapeToolModes, ToolModeName, &FLandscapeToolMode::ToolModeName))
 		{
-			if (LandscapeToolModes[i].ToolModeName == ToolModeName)
+			// Save current tool mode before changing
+			LastActiveToolModeName = CurrentToolMode ? CurrentToolMode->ToolModeName : FName("ToolMode_Sculpt");
+
+			CurrentToolMode = NewToolMode;
+
+			if (bRestoreCurrentTool)
 			{
-				CurrentToolMode = &LandscapeToolModes[i];
-				if (bRestoreCurrentTool)
-				{
-					if (CurrentToolMode->CurrentToolName == NAME_None)
-					{
-						CurrentToolMode->CurrentToolName = CurrentToolMode->ValidTools[0];
-						CurrentToolMode->CurrentTargetLayerName = NAME_None;
-					}
-					SetCurrentTool(CurrentToolMode->CurrentToolName, CurrentToolMode->CurrentTargetLayerName);
-				}
-				if (Toolkit.IsValid())
+				if (CurrentToolMode->CurrentToolName.IsNone())
 				{
-					Toolkit->SetCurrentPalette(ToolModeName);
+					CurrentToolMode->CurrentToolName = CurrentToolMode->ValidTools[0];
+					CurrentToolMode->CurrentTargetLayerName = NAME_None;
 				}
-				break;
+				SetCurrentTool(CurrentToolMode->CurrentToolName, CurrentToolMode->CurrentTargetLayerName);
+			}
+			if (Toolkit.IsValid())
+			{
+				Toolkit->SetCurrentPalette(ToolModeName);
 			}
 		}
 	}
@@ -2172,6 +2173,9 @@ void FEdModeLandscape::SetCurrentTargetLayer(FName TargetLayerName, TWeakObjectP
 
 void FEdModeLandscape::SetCurrentTool(int32 ToolIndex, FName TargetLayerName)
 {
+	// Save current tool name before changing
+	LastActiveToolName = CurrentTool ? CurrentTool->GetToolName() : FName("Sculpt");
+
 	if (CurrentTool)
 	{
 		CurrentTool->PreviousBrushIndex = CurrentBrushSetIndex;
@@ -2180,24 +2184,22 @@ void FEdModeLandscape::SetCurrentTool(int32 ToolIndex, FName TargetLayerName)
 	}
 	CurrentToolIndex = LandscapeTools.IsValidIndex(ToolIndex) ? ToolIndex : 0;
 	FLandscapeTool* NewTool = LandscapeTools[CurrentToolIndex].Get();
+
 	if (!CurrentToolMode->ValidTools.Contains(NewTool->GetToolName()))
 	{
-		// if tool isn't valid for this mode then automatically switch modes
-		// this mostly happens with shortcut keys
-		bool bFoundValidMode = false;
-		for (int32 i = 0; i < LandscapeToolModes.Num(); ++i)
-		{
-			if (LandscapeToolModes[i].ValidTools.Contains(NewTool->GetToolName()))
+		// if tool isn't valid for this mode then automatically switch modes. this mostly happens with shortcut keys
+		FLandscapeToolMode* BackupToolMode = LandscapeToolModes.FindByPredicate([NewToolName = NewTool->GetToolName()](const FLandscapeToolMode& ToolMode)
 			{
-				SetCurrentToolMode(LandscapeToolModes[i].ToolModeName, false);
-				bFoundValidMode = true;
-				break;
-			}
+				return ToolMode.ValidTools.Contains(NewToolName);
+			});
+
+		// if no valid mode is found, keep the mode and fallback to the first valid tool
+		if (BackupToolMode != nullptr)
+		{
+			SetCurrentToolMode(BackupToolMode->ToolModeName, false);
 		}
-		
-		// default to first valid tool of current mode
-		if (!bFoundValidMode)
-		{		
+		else
+		{
 			SetCurrentTool(CurrentToolMode->ValidTools[0]);
 			return;
 		}
@@ -4033,16 +4035,17 @@ void FEdModeLandscape::ToggleSplinesTool(const ULandscapeEditLayerBase* InEditLa
 	{
 		return;
 	}
-
-	if (InEditLayer->IsA<ULandscapeEditLayerSplines>() && CurrentTool->GetToolName() != FName("Splines"))
+ 
+	// Enters spline mode or toggles to the last active landscape edit mode (LastActive defaults to Sculpt mode/tool)
+	if (InEditLayer->IsA<ULandscapeEditLayerSplines>() && CurrentTool && CurrentTool->GetToolName() != FName("Splines"))
 	{
 		SetCurrentToolMode("ToolMode_Manage", false);
 		SetCurrentTool(FName("Splines"));
 	}
 	else
 	{
-		SetCurrentToolMode("ToolMode_Sculpt", false);
-		SetCurrentTool(FName("Sculpt"));
+		SetCurrentToolMode(LastActiveToolModeName, false);
+		SetCurrentTool(LastActiveToolName);
 	}
 
 	SetSelectedEditLayer(Landscape->GetLayerIndex(InEditLayer->GetGuid()));
diff --git a/Engine/Source/Editor/LandscapeEditor/Private/LandscapeEdMode.h b/Engine/Source/Editor/LandscapeEditor/Private/LandscapeEdMode.h
--- a/Engine/Source/Editor/LandscapeEditor/Private/LandscapeEdMode.h
+++ b/Engine/Source/Editor/LandscapeEditor/Private/LandscapeEdMode.h
@@ -675,4 +675,8 @@ private:
 	/** The landscape mode panel has the ability to display the properties of certain UObjects if needed (to bypass the fact that the standard details 
 	 panel is not able to display anything else than actors and components) */
 	TArray<TWeakObjectPtr<UObject>> InspectedObjects;
+
+	// Stores last active tool info. Used to toggle between current and last editor modes (Eg. Toggle Splines Layer)
+	FName LastActiveToolModeName = FName("ToolMode_Sculpt");
+	FName LastActiveToolName = FName("Sculpt");
 };
\ No newline at end of file
diff --git a/Engine/Source/Editor/LandscapeEditor/Private/LandscapeEditorDetailCustomization_EditLayers.cpp b/Engine/Source/Editor/LandscapeEditor/Private/LandscapeEditorDetailCustomization_EditLayers.cpp
--- a/Engine/Source/Editor/LandscapeEditor/Private/LandscapeEditorDetailCustomization_EditLayers.cpp
+++ b/Engine/Source/Editor/LandscapeEditor/Private/LandscapeEditorDetailCustomization_EditLayers.cpp
@@ -682,7 +682,8 @@ void FLandscapeEditLayerContextMenuCustomization_Splines::CustomizeContextMenu(U
 
 		FMenuEntryParams ToggleSplineEntry;
 		ToggleSplineEntry.LabelOverride = LOCTEXT("LandscapeEditorToggleSplinesTool", "Toggle Splines Tool"),
-		ToggleSplineEntry.ToolTipOverride = LOCTEXT("LandscapeEditorToggleSplinesToolTip", "Toggles landscape splines mode");
+		ToggleSplineEntry.ToolTipOverride = LOCTEXT("LandscapeEditorToggleSplinesToolTip", "Toggles between Landscape Splines and the last Landscape Edit mode");
+		ToggleSplineEntry.InputBindingOverride = LOCTEXT("LandscapeEditorToggleSplinesToolTipInputBinding", "Dbl-Click");
 		ToggleSplineEntry.DirectActions = FUIAction(FExecuteAction::CreateLambda([LandscapeEdMode, InEditLayer]()
 		{
 			if (LandscapeEdMode && InEditLayer)
