Landscape Editor - Fixed invalid ULandscapeInfo GC reference caused by the LandscapeGizmoActor and FLandscapeBrush
* ULandscapeInfo objects are cleaned up when levels are unloaded (Eg. Level instance is unloaded). LandscapeGizmo held a bad reference to old landscape info
* FLandscapeBrush holds reference to deleted ULandscapeComponents
* Fixes a crash observed after exiting a level instance with LandscapeEdMode open

diff --git a/Engine/Source/Editor/LandscapeEditor/Private/LandscapeEdMode.cpp b/Engine/Source/Editor/LandscapeEditor/Private/LandscapeEdMode.cpp
--- a/Engine/Source/Editor/LandscapeEditor/Private/LandscapeEdMode.cpp
+++ b/Engine/Source/Editor/LandscapeEditor/Private/LandscapeEdMode.cpp
@@ -1515,7 +1515,7 @@ EEditAction::Type FEdModeLandscape::GetActionEditCopy()
 		{
 			if (GLandscapeEditRenderMode & ELandscapeEditRenderMode::Gizmo || GLandscapeEditRenderMode & (ELandscapeEditRenderMode::Select))
 			{
-				if (UISettings->CurrentGizmoActor.IsValid() && GizmoBrush && UISettings->CurrentGizmoActor->TargetLandscapeInfo)
+				if (UISettings->CurrentGizmoActor.IsValid() && GizmoBrush && UISettings->CurrentGizmoActor->TargetLandscapeInfo.IsValid())
 				{
 					Result = EEditAction::Process;
 				}
@@ -1541,7 +1541,7 @@ EEditAction::Type FEdModeLandscape::GetActionEditPaste()
 		{
 			if (GLandscapeEditRenderMode & ELandscapeEditRenderMode::Gizmo || GLandscapeEditRenderMode & (ELandscapeEditRenderMode::Select))
 			{
-				if (UISettings->CurrentGizmoActor.IsValid() && GizmoBrush && UISettings->CurrentGizmoActor->TargetLandscapeInfo)
+				if (UISettings->CurrentGizmoActor.IsValid() && GizmoBrush && UISettings->CurrentGizmoActor->TargetLandscapeInfo.IsValid())
 				{
 					Result = EEditAction::Process;
 				}
@@ -1642,7 +1642,7 @@ bool FEdModeLandscape::ProcessEditCopy()
 
 			if (!CurrentlySelectedBPBrush)
 			{
-				bool IsSlowTask = IsSlowSelect(UISettings->CurrentGizmoActor->TargetLandscapeInfo);
+				bool IsSlowTask = IsSlowSelect(UISettings->CurrentGizmoActor->TargetLandscapeInfo.Get());
 				if (IsSlowTask)
 				{
 					GWarn->BeginSlowTask(LOCTEXT("BeginFitGizmoAndCopy", "Fit Gizmo to Selected Region and Copy Data..."), true);
@@ -1711,7 +1711,7 @@ bool FEdModeLandscape::ProcessEditPaste()
 
 			if (!CurrentlySelectedBPBrush)
 			{
-				bool IsSlowTask = IsSlowSelect(UISettings->CurrentGizmoActor->TargetLandscapeInfo);
+				bool IsSlowTask = IsSlowSelect(UISettings->CurrentGizmoActor->TargetLandscapeInfo.Get());
 				if (IsSlowTask)
 				{
 					GWarn->BeginSlowTask(LOCTEXT("BeginPasteGizmoDataTask", "Paste Gizmo Data..."), true);
@@ -3134,10 +3134,10 @@ FVector FEdModeLandscape::GetWidgetLocation() const
 
 	if (UISettings->CurrentGizmoActor.IsValid() && (GLandscapeEditRenderMode & ELandscapeEditRenderMode::Gizmo) && UISettings->CurrentGizmoActor->IsSelected())
 	{
-		if (UISettings->CurrentGizmoActor->TargetLandscapeInfo && UISettings->CurrentGizmoActor->TargetLandscapeInfo->GetLandscapeProxy())
+		if (UISettings->CurrentGizmoActor->TargetLandscapeInfo.IsValid() && UISettings->CurrentGizmoActor->TargetLandscapeInfo->GetLandscapeProxy())
 		{
 			// Apply Landscape transformation when it is available
-			ULandscapeInfo* LandscapeInfo = UISettings->CurrentGizmoActor->TargetLandscapeInfo;
+			ULandscapeInfo* LandscapeInfo = UISettings->CurrentGizmoActor->TargetLandscapeInfo.Get();
 			return UISettings->CurrentGizmoActor->GetActorLocation()
 				+ FQuatRotationMatrix(LandscapeInfo->GetLandscapeProxy()->GetActorQuat()).TransformPosition(FVector(0, 0, UISettings->CurrentGizmoActor->GetLength()));
 		}
diff --git a/Engine/Source/Editor/LandscapeEditor/Private/LandscapeEdModeBrushes.cpp b/Engine/Source/Editor/LandscapeEditor/Private/LandscapeEdModeBrushes.cpp
--- a/Engine/Source/Editor/LandscapeEditor/Private/LandscapeEdModeBrushes.cpp
+++ b/Engine/Source/Editor/LandscapeEditor/Private/LandscapeEdModeBrushes.cpp
@@ -3,6 +3,7 @@
 #include "CoreMinimal.h"
 #include "InputCoreTypes.h"
 #include "Materials/Material.h"
+#include "Algo/Transform.h"
 #include "Materials/MaterialInstanceDynamic.h"
 #include "Materials/MaterialInstanceConstant.h"
 #include "EngineUtils.h"
@@ -120,9 +121,13 @@ public:
 
 		for (const ULandscapeComponent* Component : InAffectedComponents)
 		{
-			const ALandscapeProxy* LandscapeProxy = Component->GetLandscapeProxy();
-			const ULandscapeLayerInfoObject* LayerInfo = EdMode->CurrentToolTarget.LayerInfo.Get();
+			const ALandscapeProxy* LandscapeProxy = Component ? Component->GetLandscapeProxy() : nullptr;
+			if (Component == nullptr || LandscapeProxy == nullptr)
+			{
+				return false;
+			}
 
+			const ULandscapeLayerInfoObject* LayerInfo = EdMode->CurrentToolTarget.LayerInfo.Get();
 			if (((EdMode->CurrentToolTarget.TargetType == ELandscapeToolTargetType::Weightmap) || (EdMode->CurrentToolTarget.TargetType == ELandscapeToolTargetType::Visibility))
 				&& (EdMode->UISettings->PaintingRestriction != ELandscapeLayerPaintingRestriction::None))
 			{
@@ -183,14 +188,10 @@ public:
 			bHasUnloadedComponents = LandscapeInfo->HasUnloadedComponentsInRegion(Bounds.Min.X, Bounds.Min.Y, Bounds.Max.X - 1, Bounds.Max.Y - 1);
 		}
 		
+		// Store components from the previous tick
 		TSet<ULandscapeComponent*> PreviousComponents;
-		for (TWeakObjectPtr<ULandscapeComponent> BrushMaterialComponent : BrushMaterialComponents)
-		{
-			if (ULandscapeComponent* Component = BrushMaterialComponent.Get())
-			{
-				PreviousComponents.Add(Component);
-			}
-		}
+		PreviousComponents.Reserve(BrushMaterialComponents.Num());
+		Algo::TransformIf(BrushMaterialComponents, PreviousComponents, [](TWeakObjectPtr<ULandscapeComponent> Component) { return Component.IsValid(); }, [](TWeakObjectPtr<ULandscapeComponent> Component) { return Component.Get(); });
 
 		// We're dealing with weak object ptrs. Now is a good time to remove the null entries (weak pointers can be invalidated) from our containers if any : 
 		// First, remove all components (the map's keys) that might have become stale (the map's value are also weak object ptrs to materials but they will be handled right after) :
@@ -252,12 +253,8 @@ public:
 			AddedComponent->UpdateEditToolRenderData();
 		}
 
-		BrushMaterialComponents.Empty();
-		for (ULandscapeComponent* BrushMaterialComponent : NewComponents)
-		{
-			BrushMaterialComponents.Add(TWeakObjectPtr<ULandscapeComponent>(BrushMaterialComponent));
-		}
-
+		BrushMaterialComponents.Empty(NewComponents.Num());
+		Algo::Transform(NewComponents, BrushMaterialComponents, [](ULandscapeComponent* Component) { return TWeakObjectPtr<ULandscapeComponent>(Component); });
 
 		// Set params for brush material.
 		FVector WorldLocation = Proxy->LandscapeActorToWorld().TransformPosition(FVector(LastMousePosition.X, LastMousePosition.Y, 0));
@@ -429,7 +426,7 @@ class FLandscapeBrushComponent : public FLandscapeBrush
 {
 	using Super = FLandscapeBrush;
 
-	TSet<TObjectPtr<ULandscapeComponent>> BrushMaterialComponents;
+	TSet<TWeakObjectPtr<ULandscapeComponent>, TWeakObjectPtrSetKeyFuncs<TWeakObjectPtr<ULandscapeComponent>>> BrushMaterialComponents;
 
 	virtual const TCHAR* GetBrushName() override { return TEXT("Component"); }
 	virtual FText GetDisplayName() override { return NSLOCTEXT("UnrealEd", "LandscapeMode_Brush_Component", "Component"); };
@@ -453,7 +450,6 @@ public:
 	// FGCObject interface
 	virtual void AddReferencedObjects(FReferenceCollector& Collector) override
 	{
-		Collector.AddReferencedObjects(BrushMaterialComponents);
 		Collector.AddReferencedObject(BrushMaterial);
 	}
 
@@ -480,7 +476,7 @@ public:
 
 	virtual void Tick(FEditorViewportClient* ViewportClient, float DeltaTime) override
 	{
-		TSet<ULandscapeComponent*> NewComponents;
+		TSet<TWeakObjectPtr<ULandscapeComponent>, TWeakObjectPtrSetKeyFuncs<TWeakObjectPtr<ULandscapeComponent>>> NewComponents;
 
 		// Adjusting the brush may use the same keybind as moving the camera as they can be user-set, so we need this second check.
 		if (!ViewportClient->IsMovingCamera() || EdMode->IsAdjustingBrush(ViewportClient))
@@ -521,26 +517,29 @@ public:
 				}
 
 				// Set brush material for components in new region
-				for (ULandscapeComponent* NewComponent : NewComponents)
+				for (TWeakObjectPtr<ULandscapeComponent> NewComponent : NewComponents)
 				{
-					NewComponent->EditToolRenderData.ToolMaterial = BrushMaterial;
-					NewComponent->UpdateEditToolRenderData();
+					if (NewComponent.IsValid())
+					{
+						NewComponent->EditToolRenderData.ToolMaterial = BrushMaterial;
+						NewComponent->UpdateEditToolRenderData();
+					}
 				}
 			}
 		}
 
 		// Remove the material from any old components that are no longer in the region
-		TSet<ULandscapeComponent*> RemovedComponents = ObjectPtrDecay(BrushMaterialComponents).Difference(NewComponents);
-		for (ULandscapeComponent* RemovedComponent : RemovedComponents)
+		TSet<TWeakObjectPtr<ULandscapeComponent>, TWeakObjectPtrSetKeyFuncs<TWeakObjectPtr<ULandscapeComponent>>> RemovedComponents = BrushMaterialComponents.Difference(NewComponents);
+		for (TWeakObjectPtr<ULandscapeComponent> RemovedComponent : RemovedComponents)
 		{
-			if (RemovedComponent != nullptr)
+			if (RemovedComponent.IsValid())
 			{
 				RemovedComponent->EditToolRenderData.ToolMaterial = nullptr;
 				RemovedComponent->UpdateEditToolRenderData();
 			}
 		}
 
-		BrushMaterialComponents = ObjectPtrWrap(MoveTemp(NewComponents));
+		BrushMaterialComponents = MoveTemp(NewComponents);
 	}
 
 	virtual void MouseMove(float LandscapeX, float LandscapeY) override
@@ -583,9 +582,9 @@ public:
 			}
 
 			// Get extent for all components
-			for (ULandscapeComponent* Component : BrushMaterialComponents)
+			for (TWeakObjectPtr<ULandscapeComponent> Component : BrushMaterialComponents)
 			{
-				if (ensure(Component))
+				if (ensure(Component.IsValid()))
 				{
 					Component->GetComponentExtent(ExclusiveBounds.Min.X, ExclusiveBounds.Min.Y, ExclusiveBounds.Max.X, ExclusiveBounds.Max.Y);
 				}
@@ -642,7 +641,7 @@ class FLandscapeBrushGizmo : public FLandscapeBrush
 {
 	using Super = FLandscapeBrush;
 
-	TSet<TObjectPtr<ULandscapeComponent>> BrushMaterialComponents;
+	TSet<TWeakObjectPtr<ULandscapeComponent>, TWeakObjectPtrSetKeyFuncs<TWeakObjectPtr<ULandscapeComponent>>> BrushMaterialComponents;
 
 	const TCHAR* GetBrushName() override { return TEXT("Gizmo"); }
 	virtual FText GetDisplayName() override { return NSLOCTEXT("UnrealEd", "LandscapeMode_Brush_Gizmo", "Gizmo"); };
@@ -663,7 +662,6 @@ public:
 	// FGCObject interface
 	virtual void AddReferencedObjects(FReferenceCollector& Collector) override
 	{
-		Collector.AddReferencedObjects(BrushMaterialComponents);
 		Collector.AddReferencedObject(BrushMaterial);
 	}
 
@@ -699,9 +697,9 @@ public:
 		{
 			ALandscapeGizmoActiveActor* Gizmo = EdMode->UISettings->CurrentGizmoActor.Get();
 
-			if (Gizmo && Gizmo->TargetLandscapeInfo && (Gizmo->TargetLandscapeInfo == EdMode->CurrentToolTarget.LandscapeInfo.Get()) && Gizmo->GizmoTexture && Gizmo->GetRootComponent())
+			if (Gizmo && Gizmo->TargetLandscapeInfo.IsValid() && (Gizmo->TargetLandscapeInfo == EdMode->CurrentToolTarget.LandscapeInfo.Get()) && Gizmo->GizmoTexture && Gizmo->GetRootComponent())
 			{
-				ULandscapeInfo* LandscapeInfo = Gizmo->TargetLandscapeInfo;
+				ULandscapeInfo* LandscapeInfo = Gizmo->TargetLandscapeInfo.Get();
 				if (LandscapeInfo && LandscapeInfo->GetLandscapeProxy())
 				{
 					float ScaleXY = static_cast<float>(FMath::Abs(LandscapeInfo->DrawScale.X));
@@ -748,18 +746,23 @@ public:
 						NewComponent->UpdateEditToolRenderData();
 					}
 
+					// Convert NewComponents to TWeakObjectPtr TSet
+					TSet<TWeakObjectPtr<ULandscapeComponent>, TWeakObjectPtrSetKeyFuncs<TWeakObjectPtr<ULandscapeComponent>>> NewComponentsWeakSet;
+					NewComponentsWeakSet.Reserve(NewComponents.Num());
+					Algo::Transform(NewComponents, NewComponentsWeakSet, [](ULandscapeComponent* Component) { return TWeakObjectPtr<ULandscapeComponent>(Component); });
+
 					// Remove the material from any old components that are no longer in the region
-					TSet<ULandscapeComponent*> RemovedComponents = ObjectPtrDecay(BrushMaterialComponents).Difference(NewComponents);
-					for (ULandscapeComponent* RemovedComponent : RemovedComponents)
+					TSet<TWeakObjectPtr<ULandscapeComponent>, TWeakObjectPtrSetKeyFuncs<TWeakObjectPtr<ULandscapeComponent>>> RemovedComponents = BrushMaterialComponents.Difference(NewComponentsWeakSet);
+					for (TWeakObjectPtr<ULandscapeComponent> RemovedComponent : RemovedComponents)
 					{
-						if (RemovedComponent != nullptr)
+						if (RemovedComponent.IsValid())
 						{
 							RemovedComponent->EditToolRenderData.GizmoMaterial = nullptr;
 							RemovedComponent->UpdateEditToolRenderData();
 						}
 					}
 
-					BrushMaterialComponents = ObjectPtrWrap(MoveTemp(NewComponents));
+					BrushMaterialComponents = MoveTemp(NewComponentsWeakSet);
 				}
 			}
 		}
@@ -811,9 +814,9 @@ public:
 
 		// Get extent for all components
 		FIntRect ExclusiveBounds(INT_MAX, INT_MAX, INT_MIN, INT_MIN);
-		for (ULandscapeComponent* Component : BrushMaterialComponents)
+		for (TWeakObjectPtr<ULandscapeComponent> Component : BrushMaterialComponents)
 		{
-			if (ensure(Component))
+			if (ensure(Component.IsValid()))
 			{
 				Component->GetComponentExtent(ExclusiveBounds.Min.X, ExclusiveBounds.Min.Y, ExclusiveBounds.Max.X, ExclusiveBounds.Max.Y);
 			}
diff --git a/Engine/Source/Editor/LandscapeEditor/Private/LandscapeEditorDetailCustomization_CopyPaste.cpp b/Engine/Source/Editor/LandscapeEditor/Private/LandscapeEditorDetailCustomization_CopyPaste.cpp
--- a/Engine/Source/Editor/LandscapeEditor/Private/LandscapeEditorDetailCustomization_CopyPaste.cpp
+++ b/Engine/Source/Editor/LandscapeEditor/Private/LandscapeEditorDetailCustomization_CopyPaste.cpp
@@ -261,7 +261,7 @@ FReply FLandscapeEditorDetailCustomization_CopyPaste::OnFitGizmoToSelectionButto
 	if (LandscapeEdMode != nullptr)
 	{
 		ALandscapeGizmoActiveActor* Gizmo = LandscapeEdMode->UISettings->CurrentGizmoActor.Get();
-		if (Gizmo && Gizmo->TargetLandscapeInfo)
+		if (Gizmo && Gizmo->TargetLandscapeInfo.IsValid())
 		{
 			FScopedTransaction Transaction(LOCTEXT("LandscapeGizmo_FitToSelection", "Fit gizmo size to selection"));
 			Gizmo->Modify();
@@ -278,7 +278,7 @@ FReply FLandscapeEditorDetailCustomization_CopyPaste::OnFitHeightsToGizmoButtonC
 	if (LandscapeEdMode != nullptr)
 	{
 		ALandscapeGizmoActiveActor* Gizmo = LandscapeEdMode->UISettings->CurrentGizmoActor.Get();
-		if (Gizmo && Gizmo->TargetLandscapeInfo)
+		if (Gizmo && Gizmo->TargetLandscapeInfo.IsValid())
 		{
 			FScopedTransaction Transaction(LOCTEXT("LandscapeGizmo_FitMinMaxHeight", "Set gizmo height to fix contained data"));
 			Gizmo->Modify();
@@ -295,7 +295,7 @@ FReply FLandscapeEditorDetailCustomization_CopyPaste::OnClearGizmoDataButtonClic
 	if (LandscapeEdMode != nullptr)
 	{
 		ALandscapeGizmoActiveActor* Gizmo = LandscapeEdMode->UISettings->CurrentGizmoActor.Get();
-		if (Gizmo && Gizmo->TargetLandscapeInfo)
+		if (Gizmo && Gizmo->TargetLandscapeInfo.IsValid())
 		{
 			Gizmo->ClearGizmoData();
 		}
@@ -534,10 +534,10 @@ FReply FLandscapeEditorDetailCustomization_CopyPaste::OnGizmoExportButtonClicked
 	if (LandscapeEdMode != nullptr)
 	{
 		ALandscapeGizmoActiveActor* Gizmo = LandscapeEdMode->UISettings->CurrentGizmoActor.Get();
-		if (Gizmo && Gizmo->TargetLandscapeInfo && Gizmo->SelectedData.Num())
+		if (Gizmo && Gizmo->TargetLandscapeInfo.IsValid() && Gizmo->SelectedData.Num())
 		{
 			int32 TargetIndex = -1;
-			ULandscapeInfo* LandscapeInfo = Gizmo->TargetLandscapeInfo;
+			ULandscapeInfo* LandscapeInfo = Gizmo->TargetLandscapeInfo.Get();
 			TArray<FString> Filenames;
 
 			// Local set for export
diff --git a/Engine/Source/Editor/LandscapeEditor/Private/LandscapeEditorObject.cpp b/Engine/Source/Editor/LandscapeEditor/Private/LandscapeEditorObject.cpp
--- a/Engine/Source/Editor/LandscapeEditor/Private/LandscapeEditorObject.cpp
+++ b/Engine/Source/Editor/LandscapeEditor/Private/LandscapeEditorObject.cpp
@@ -511,7 +511,7 @@ void ULandscapeEditorObject::SetGizmoSnapMode(ELandscapeGizmoSnapType InSnapMode
 	if (SnapMode != ELandscapeGizmoSnapType::None)
 	{
 		// Ignore gizmo updates if landscape info is invalid (Ex. landscape is deleted after using Copy Tool)
-		if (CurrentGizmoActor.IsValid() && CurrentGizmoActor->TargetLandscapeInfo)
+		if (CurrentGizmoActor.IsValid() && CurrentGizmoActor->TargetLandscapeInfo.IsValid())
 		{
 			const FVector WidgetLocation = CurrentGizmoActor->GetActorLocation();
 			const FRotator WidgetRotation = CurrentGizmoActor->GetActorRotation();
diff --git a/Engine/Source/Runtime/Landscape/Classes/LandscapeGizmoActor.h b/Engine/Source/Runtime/Landscape/Classes/LandscapeGizmoActor.h
--- a/Engine/Source/Runtime/Landscape/Classes/LandscapeGizmoActor.h
+++ b/Engine/Source/Runtime/Landscape/Classes/LandscapeGizmoActor.h
@@ -9,6 +9,7 @@
 #include "LandscapeGizmoActor.generated.h"
 
 class UBillboardComponent;
+class ULandscapeInfo;
 
 UCLASS(notplaceable, MinimalAPI, NotBlueprintable)
 class ALandscapeGizmoActor : public AActor
@@ -35,7 +36,7 @@ class ALandscapeGizmoActor : public AActor
 	float RelativeScaleZ;
 
 	UPROPERTY(EditAnywhere, transient, Category=Gizmo)
-	TObjectPtr<class ULandscapeInfo> TargetLandscapeInfo;
+	TWeakObjectPtr<ULandscapeInfo> TargetLandscapeInfo;
 
 private:
 	UPROPERTY()
diff --git a/Engine/Source/Runtime/Landscape/Private/LandscapeGizmoActor.cpp b/Engine/Source/Runtime/Landscape/Private/LandscapeGizmoActor.cpp
--- a/Engine/Source/Runtime/Landscape/Private/LandscapeGizmoActor.cpp
+++ b/Engine/Source/Runtime/Landscape/Private/LandscapeGizmoActor.cpp
@@ -16,6 +16,7 @@
 #include "LandscapeLayerInfoObject.h"
 #include "LandscapeInfoMap.h"
 #include "LandscapeDataAccess.h"
+#include "LandscapePrivate.h"
 #include "LandscapeRender.h"
 #include "LandscapeGizmoActiveActor.h"
 #include "LandscapeGizmoRenderComponent.h"
@@ -160,7 +161,7 @@ public:
 			&& !FMath::IsNearlyZero(Gizmo->CachedWidth) 
 			&& !FMath::IsNearlyZero(Gizmo->CachedHeight))
 		{
-			ULandscapeInfo* LandscapeInfo = Gizmo->TargetLandscapeInfo;
+			ULandscapeInfo* LandscapeInfo = Gizmo->TargetLandscapeInfo.Get();
  			if (LandscapeInfo && LandscapeInfo->GetLandscapeProxy() && !FMath::IsNearlyZero(LandscapeInfo->DrawScale.X))
 			{
 				const float ScaleXY = static_cast<float>(LandscapeInfo->DrawScale.X);
@@ -414,7 +415,7 @@ FBoxSphereBounds ULandscapeGizmoRenderComponent::CalcBounds(const FTransform& Lo
 	ALandscapeGizmoActiveActor* Gizmo = Cast<ALandscapeGizmoActiveActor>(GetOwner());
 	if (Gizmo)
 	{
-		ULandscapeInfo* LandscapeInfo = Gizmo->TargetLandscapeInfo;
+		ULandscapeInfo* LandscapeInfo = Gizmo->TargetLandscapeInfo.Get();
 		if (LandscapeInfo && LandscapeInfo->GetLandscapeProxy())
 		{
 			FTransform LToW = LandscapeInfo->GetLandscapeProxy()->LandscapeActorToWorld();
@@ -568,7 +569,7 @@ void ALandscapeGizmoActiveActor::PostEditChangeProperty(FPropertyChangedEvent& P
 	}
 	else if ( PropertyName == FName(TEXT("TargetLandscapeInfo")) )
 	{
-		SetTargetLandscape(TargetLandscapeInfo);
+		SetTargetLandscape(TargetLandscapeInfo.Get());
 	}
 	// AActor::PostEditChange will ForceUpdateComponents()
 	Super::PostEditChangeProperty(PropertyChangedEvent);
@@ -706,13 +707,13 @@ ALandscapeGizmoActor* ALandscapeGizmoActiveActor::SpawnGizmoActor()
 
 void ALandscapeGizmoActiveActor::SetTargetLandscape(ULandscapeInfo* LandscapeInfo)
 {
-	ULandscapeInfo* PrevInfo = TargetLandscapeInfo;
+	ULandscapeInfo* PrevInfo = TargetLandscapeInfo.Get();
 	if (!LandscapeInfo || LandscapeInfo->HasAnyFlags(RF_BeginDestroyed))
 	{
 		TargetLandscapeInfo = nullptr;
-		if (GetWorld())
+		if (const UWorld* CurrentWorld = GetWorld())
 		{
-			for (const auto& InfoMapPair : ULandscapeInfoMap::GetLandscapeInfoMap(GetWorld()).Map)
+			for (const TPair<FGuid, TObjectPtr<ULandscapeInfo>>& InfoMapPair : ULandscapeInfoMap::GetLandscapeInfoMap(CurrentWorld).Map)
 			{
 				ULandscapeInfo* CandidateInfo = InfoMapPair.Value;
 				if (CandidateInfo && CandidateInfo->SupportsLandscapeEditing() && !CandidateInfo->HasAnyFlags(RF_BeginDestroyed) && CandidateInfo->GetLandscapeProxy() != nullptr)
@@ -728,10 +729,10 @@ void ALandscapeGizmoActiveActor::SetTargetLandscape(ULandscapeInfo* LandscapeInf
 		TargetLandscapeInfo = LandscapeInfo;
 	}
 
-	if (TargetLandscapeInfo != PrevInfo)
+	if (TargetLandscapeInfo.Get() != PrevInfo)
 	{
 		// if there's no copied data, try to move somewhere useful
-		if (TargetLandscapeInfo && DataType == LGT_None)
+		if (TargetLandscapeInfo.IsValid() && DataType == LGT_None)
 		{
 			MarginZ = static_cast<float>(TargetLandscapeInfo->DrawScale.Z * 3);
 			Width = Height = static_cast<float>(TargetLandscapeInfo->DrawScale.X * (TargetLandscapeInfo->ComponentSizeQuads + 1));
@@ -767,7 +768,7 @@ void ALandscapeGizmoActiveActor::ClearGizmoData()
 
 void ALandscapeGizmoActiveActor::FitToSelection()
 {
-	if (TargetLandscapeInfo)
+	if (TargetLandscapeInfo.IsValid())
 	{
 		// Find fit size
 		int32 MinX = MAX_int32, MinY = MAX_int32;
@@ -795,7 +796,7 @@ void ALandscapeGizmoActiveActor::FitToSelection()
 
 void ALandscapeGizmoActiveActor::FitMinMaxHeight()
 {
-	if (TargetLandscapeInfo)
+	if (TargetLandscapeInfo.IsValid())
 	{
 		FVector::FReal MinZ = UE_OLD_HALF_WORLD_MAX, MaxZ = -UE_OLD_HALF_WORLD_MAX;
 		// Change MinRelativeZ and RelativeZScale to fit Gizmo Box
@@ -817,7 +818,7 @@ void ALandscapeGizmoActiveActor::FitMinMaxHeight()
 
 float ALandscapeGizmoActiveActor::GetNormalizedHeight(uint16 LandscapeHeight) const
 {
-	if (TargetLandscapeInfo)
+	if (TargetLandscapeInfo.IsValid())
 	{
 		ALandscapeProxy* Proxy = TargetLandscapeInfo->GetLandscapeProxy();
 		if (Proxy)
@@ -838,7 +839,7 @@ float ALandscapeGizmoActiveActor::GetNormalizedHeight(uint16 LandscapeHeight) co
 
 float ALandscapeGizmoActiveActor::GetWorldHeight(float NormalizedHeight) const
 {
-	if (TargetLandscapeInfo)
+	if (TargetLandscapeInfo.IsValid())
 	{
 		ALandscapeProxy* Proxy = TargetLandscapeInfo->GetLandscapeProxy();
 		if (Proxy)
@@ -856,7 +857,7 @@ float ALandscapeGizmoActiveActor::GetWorldHeight(float NormalizedHeight) const
 
 float ALandscapeGizmoActiveActor::GetLandscapeHeight(float NormalizedHeight) const
 {
-	if (TargetLandscapeInfo)
+	if (TargetLandscapeInfo.IsValid())
 	{
 		NormalizedHeight = (NormalizedHeight - MinRelativeZ) * RelativeScaleZ;
 		const float ScaleZ = static_cast<float>(TargetLandscapeInfo->DrawScale.Z);
@@ -906,7 +907,7 @@ void ALandscapeGizmoActiveActor::CalcNormal()
 
 void ALandscapeGizmoActiveActor::SampleData(int32 SizeX, int32 SizeY)
 {
-	if (TargetLandscapeInfo && GizmoTexture)
+	if (TargetLandscapeInfo.IsValid() && GizmoTexture)
 	{
 		// Rasterize rendering Texture...
 		int32 TexSizeX = FMath::Min(ALandscapeGizmoActiveActor::DataTexSize, SizeX);
@@ -1041,7 +1042,7 @@ void ALandscapeGizmoActiveActor::Export(int32 Index, TArray<FString>& Filenames)
 
 	bool bExportOneTarget = (Filenames.Num() == 1);
 
-	if (TargetLandscapeInfo)
+	if (TargetLandscapeInfo.IsValid())
 	{
 		int32 MinX = MAX_int32, MinY = MAX_int32;
 		int32 MaxX = MIN_int32, MaxY = MIN_int32;
@@ -1125,7 +1126,7 @@ void ALandscapeGizmoActiveActor::Export(int32 Index, TArray<FString>& Filenames)
 
 void ALandscapeGizmoActiveActor::ExportToClipboard()
 {
-	if (TargetLandscapeInfo && DataType != LGT_None)
+	if (TargetLandscapeInfo.IsValid() && DataType != LGT_None)
 	{
 		//GWarn->BeginSlowTask( TEXT("Exporting Gizmo Data From Clipboard"), true);
 
