Fix crash when deleting middle landscape spline point and selecting cancel.

diff --git a/Engine/Source/Editor/LandscapeEditor/Private/LandscapeEdModeSplineTools.cpp b/Engine/Source/Editor/LandscapeEditor/Private/LandscapeEdModeSplineTools.cpp
--- a/Engine/Source/Editor/LandscapeEditor/Private/LandscapeEdModeSplineTools.cpp
+++ b/Engine/Source/Editor/LandscapeEditor/Private/LandscapeEdModeSplineTools.cpp
@@ -445,6 +445,11 @@ public:
 		if (bConnectable)
 		{
 			ConnectChoice = FMessageDialog::Open(EAppMsgType::YesNoCancel, LOCTEXT("WantToJoinControlPoint", "Control point has two segments attached, do you want to join them?"));
+
+			if (ConnectChoice == EAppReturnType::Cancel)
+			{
+				return;
+			}
 		}
 
 		SplinesComponent->Modify();
@@ -488,9 +493,9 @@ public:
 			case EAppReturnType::No:
 				// Use the "delete all segments" code below
 				break;
-			case EAppReturnType::Cancel:
-				// Do nothing
-				return;
+			default:
+				check(false); //unreachable
+				break;
 			}
 		}
 
