Fix crash or visual errors when importing landscape weightmap image of a different size than the landscape.
Make weightmap import consistent with heightmap import for Original and SubRegion modes, including using the gizmo offset, regardless of whether the image file matches the landscape resolution.

diff --git a/Engine/Source/Editor/LandscapeEditor/Private/LandscapeEdMode.cpp b/Engine/Source/Editor/LandscapeEditor/Private/LandscapeEdMode.cpp
--- a/Engine/Source/Editor/LandscapeEditor/Private/LandscapeEdMode.cpp
+++ b/Engine/Source/Editor/LandscapeEditor/Private/LandscapeEdMode.cpp
@@ -3326,11 +3326,11 @@ void ImportDataInternal(ULandscapeInfo* LandscapeInfo, const FString& Filename,
 	TiledImage.Load<T>(*Filename);
 	FIntPoint ImportResolution = TiledImage.GetResolution();
 
-	bool bResolutionMismatch = false;
+	bool bTransformInputData = false;
 
 	if ((ImportResolution.X != ImportRegionVerts.Width() || ImportResolution.Y != ImportRegionVerts.Height()) && TransformType != ELandscapeImportTransformType::Subregion)
 	{
-		bResolutionMismatch = true;
+		bTransformInputData = true;
 
 		FFormatNamedArguments Args;
 		Args.Add(TEXT("ImportSizeX"), ImportResolution.X);
@@ -3347,41 +3347,53 @@ void ImportDataInternal(ULandscapeInfo* LandscapeInfo, const FString& Filename,
 		}
 	}
 
-	TArray<T> ImportData;
-	if (TransformType == ELandscapeImportTransformType::Subregion)
+	bool bReadUnmodified = bTransformInputData ||
+		TransformType == ELandscapeImportTransformType::None ||
+		TransformType == ELandscapeImportTransformType::Subregion;
+
+	FIntRect DestRegionVerts = ImportRegionVerts;
+
+	// None/Original mode doesn't handle resolution mismatch at all.  Tailor the destination region using the offset and image resolution.
+	if (TransformType == ELandscapeImportTransformType::None)
 	{
-		TiledImage.ReadRegion<T>(ImportRegionVerts, ImportData, bFlipYAxis);
+		DestRegionVerts.Min = Offset;
+		DestRegionVerts.Max = Offset + ImportResolution;
 	}
-	else if (bResolutionMismatch)
+	else if (TransformType == ELandscapeImportTransformType::Subregion)
+	{
+		// Subregion mode depends on TransformImportData to re-align data to the offset even if it starts at the correct resolution.
+		bTransformInputData = true;
+	}
+
+	TArray<T> ImportData;
+	if (bReadUnmodified)
 	{
 		TiledImage.Read<T>(ImportData, bFlipYAxis);
 	}
 	else
 	{
-		const FIntRect RegionToLoad(0, 0, ImportRegionVerts.Width(), ImportRegionVerts.Height());
+		const FIntRect RegionToLoad(0, 0, DestRegionVerts.Width(), DestRegionVerts.Height());
 		TiledImage.ReadRegion<T>(RegionToLoad, ImportData, bFlipYAxis);
 	}
 
-	// Expand if necessary...
 	{
 		TArray<T> FinalData;
-		if (bResolutionMismatch)
+		if (bTransformInputData)
 		{
 			FLandscapeImportHelper::TransformImportData<T>(ImportData, FinalData, FLandscapeImportResolution(ImportResolution.X, ImportResolution.Y), FLandscapeImportResolution(ImportRegionVerts.Width(), ImportRegionVerts.Height()), TransformType, Offset);
 		}
-		else if (TransformType == ELandscapeImportTransformType::Subregion)
-		{
-			FinalData = MoveTemp(ImportData);
-		}
 		else
 		{
 			FinalData = MoveTemp(ImportData);
 		}
 
+		// SetDataFunc (FHeightmapAccessor::SetData / TAlphamapAccessor::SetData) expects the data buffer to exactly match the region dimensions.  Important
+		// for memory access safety and stride.
+		check(DestRegionVerts.Area() == FinalData.Num());
+
 		// Set Data is in Quads (so remove 1)
-		SetDataFunc(ImportRegionVerts.Min.X, ImportRegionVerts.Min.Y, ImportRegionVerts.Max.X - 1, ImportRegionVerts.Max.Y - 1, FinalData);
+		SetDataFunc(DestRegionVerts.Min.X, DestRegionVerts.Min.Y, DestRegionVerts.Max.X - 1, DestRegionVerts.Max.Y - 1, FinalData);
 	}
-	
 }
 
 void FEdModeLandscape::ImportHeightData(ULandscapeInfo* LandscapeInfo, const FGuid& LayerGuid, const FString& Filename, const FIntRect& ImportRegionVerts, ELandscapeImportTransformType TransformType, FIntPoint Offset, const ELandscapeLayerPaintingRestriction& PaintRestriction, bool bFlipYAxis)
diff --git a/Engine/Source/Editor/LandscapeEditor/Private/LandscapeEdModeComponentTools.cpp b/Engine/Source/Editor/LandscapeEditor/Private/LandscapeEdModeComponentTools.cpp
--- a/Engine/Source/Editor/LandscapeEditor/Private/LandscapeEdModeComponentTools.cpp
+++ b/Engine/Source/Editor/LandscapeEditor/Private/LandscapeEdModeComponentTools.cpp
@@ -2477,7 +2477,7 @@ public:
 		const int32 Height = EdMode->UISettings->ImportType == ELandscapeImportTransformType::None ? EdMode->UISettings->ImportLandscape_Height : (LandscapeExtent.Height()+1);
 		const int32 Width = EdMode->UISettings->ImportType == ELandscapeImportTransformType::None ? EdMode->UISettings->ImportLandscape_Width: (LandscapeExtent.Width()+1);
 		const int32 ImportHeight = EdMode->UISettings->ImportType != ELandscapeImportTransformType::Resample ? EdMode->UISettings->ImportLandscape_Height : (LandscapeExtent.Height() + 1);
-		const int32 ImportWidth = EdMode->UISettings->ImportType != ELandscapeImportTransformType::Resample ? EdMode->UISettings->ImportLandscape_Height : (LandscapeExtent.Height() + 1);
+		const int32 ImportWidth = EdMode->UISettings->ImportType != ELandscapeImportTransformType::Resample ? EdMode->UISettings->ImportLandscape_Width : (LandscapeExtent.Width() + 1);
 
 		const FTransform LandscapeTransfo = FTransform(FRotator(0, 0, 0), MinLocation, LandscapeInfo->DrawScale);
 
diff --git a/Engine/Source/Editor/LandscapeEditor/Private/LandscapeEditorDetailCustomization_ImportExport.cpp b/Engine/Source/Editor/LandscapeEditor/Private/LandscapeEditorDetailCustomization_ImportExport.cpp
--- a/Engine/Source/Editor/LandscapeEditor/Private/LandscapeEditorDetailCustomization_ImportExport.cpp
+++ b/Engine/Source/Editor/LandscapeEditor/Private/LandscapeEditorDetailCustomization_ImportExport.cpp
@@ -589,10 +589,14 @@ FReply FLandscapeEditorDetailCustomization_ImportExport::OnImportExportButtonCli
 			{
 				ImportOffset = LocalGizmoPoint - FIntPoint(LandscapeExtent.Min.X, LandscapeExtent.Min.Y);
 			}
-			else if (TransformType == ELandscapeImportTransformType::None)
+			else if (TransformType == ELandscapeImportTransformType::None ||
+					 TransformType == ELandscapeImportTransformType::Subregion)
 			{
-				ImportRegion = FIntRect(LocalGizmoPoint.X, LocalGizmoPoint.Y, LocalGizmoPoint.X + LandscapeEdMode->UISettings->ImportLandscape_Width, LocalGizmoPoint.Y + LandscapeEdMode->UISettings->ImportLandscape_Height);
+				// We only know the image size here if a heightmap is being imported (which requires that all weightmaps match when imported at the same 
+				// time).  Pass in the correct offset consistently and let ImportHeightData/ImportWeightData figure out what that means for the import region.
+				ImportOffset = { LocalGizmoPoint.X, LocalGizmoPoint.Y };
 			}
+			// Note:  Resample mode is always aligned to the landscape and uses no offset.
 			
 			check(LandscapeEdMode->UISettings->ImportLandscape_HeightmapImportResult != ELandscapeImportResult::Error);
 			const bool bIsWorldPartition = LandscapeEdMode->GetWorld()->GetSubsystem<ULandscapeSubsystem>()->IsGridBased();
@@ -642,8 +646,8 @@ FReply FLandscapeEditorDetailCustomization_ImportExport::OnImportExportButtonCli
 
 					return !Progress.ShouldCancel();
 				};
-					
-				LandscapeRegionUtils::ForEachRegion_LoadProcessUnload(LandscapeInfo, ImportRegion, World, RegionImporter);
+
+				LandscapeRegionUtils::ForEachRegion_LoadProcessUnload(LandscapeInfo, World, RegionImporter);
 
 				FLandscapeImageFileCache& LandscapeImageFileCache = FModuleManager::GetModuleChecked<ILandscapeEditorModule>("LandscapeEditor").GetImageFileCache();
 				LandscapeImageFileCache.Clear();
@@ -658,10 +662,6 @@ FReply FLandscapeEditorDetailCustomization_ImportExport::OnImportExportButtonCli
 					Progress.EnterProgressFrame(1.0f, LOCTEXT("ImportingLandscapeHeight", "Importing Landscape Height"));
 					LandscapeEdMode->ImportHeightData(LandscapeInfo, CurrentLayerGuid, LandscapeEdMode->UISettings->ImportLandscape_HeightmapFilename, ImportRegion, TransformType, ImportOffset, PaintRestriction, LandscapeEdMode->UISettings->bFlipYAxis);
 				}
-				else
-				{
-					ImportRegion = LandscapeExtent;
-				}
 
 				for (const FLandscapeImportLayer& ImportLayer : LandscapeEdMode->UISettings->ImportLandscape_Layers)
 				{
@@ -756,7 +756,7 @@ FReply FLandscapeEditorDetailCustomization_ImportExport::OnImportExportButtonCli
 					
 				};
 
-				LandscapeRegionUtils::ForEachRegion_LoadProcessUnload(LandscapeInfo, LandscapeExtent, World, Exporter);
+				LandscapeRegionUtils::ForEachRegion_LoadProcessUnload(LandscapeInfo, World, Exporter);
 			
 
 				// For multiple file export, export each landscape proxy individually :
diff --git a/Engine/Source/Editor/LandscapeEditor/Private/LandscapeImportHelper.cpp b/Engine/Source/Editor/LandscapeEditor/Private/LandscapeImportHelper.cpp
--- a/Engine/Source/Editor/LandscapeEditor/Private/LandscapeImportHelper.cpp
+++ b/Engine/Source/Editor/LandscapeEditor/Private/LandscapeImportHelper.cpp
@@ -308,9 +308,59 @@ ELandscapeImportResult GetImportDescriptorInternal(const FString& FilePath, bool
 	return ELandscapeImportResult::Success;
 }
 
+// Copy the pixels (no resampling) from InSrcRegion of InData into the same size region in OutData, aligned to InDestOffset.  Outside this region will be reset to 0.
+// InDestSize is the desired overall dimensions of the dest buffer.
+// Note:  InSrcRegion is exlcusive/half-open, possibly unlike the other functions in this class.
+// TODO [alan.roach]: move this to LandscapeConfigHelper.h, alongside the old CopySubregion.  Kept out of the public header for hotfix compatibility.
+template<typename T>
+static void TransformImportDataCopySubregion(TConstArrayView<T> InData, TArray<T>& OutData, FIntRect InSrcRegion, int32 InSrcDataPitch, FIntPoint InDestOffset, FIntPoint InDestSize)
+{
+	FIntRect DestTotalRegion(0, 0, InDestSize.X, InDestSize.Y);
+	FIntRect DestRegion(InDestOffset, InDestOffset + InSrcRegion.Size());  // copy region in dest buffer
+	DestRegion.Clip(DestTotalRegion);  // limit copy region to overall bounds of dest buffer
+	FIntPoint SrcOffsetAdjustment = DestRegion.Min - InDestOffset + InSrcRegion.Min;  // after clipping, that might affect the source region
+	int32 DestDataPitch = InDestSize.X;
+	int32 DestDataSize = InDestSize.X * InDestSize.Y;
+
+	// Set the output buffer size and initialize to 0.
+	check(InDestSize.X > 0 && InDestSize.Y > 0);
+	OutData.Empty(DestDataSize);
+	OutData.AddZeroed(DestDataSize);
+
+	if (DestRegion.Area() <= 0)
+	{
+		return;  // Nothing to copy
+	}
+
+	check(InSrcRegion.Min.X >= 0 && InSrcRegion.Min.Y >= 0);
+	check(InSrcRegion.Max.X <= InSrcDataPitch);
+	check(InData.Num() >= ((InSrcRegion.Max.Y - 1) * InSrcDataPitch) + InSrcRegion.Max.X);  // input buffer must be large enough to contain InSrcRegion
+	check(DestDataSize >= ((DestRegion.Max.Y - 1) * DestDataPitch) + DestRegion.Max.X);
+
+
+	const int32 XEnd = DestRegion.Width();
+	const int32 YEnd = DestRegion.Height();
+
+	for (int32 Y = 0; Y < YEnd; ++Y)
+	{
+		for (int32 X = 0; X < XEnd; ++X)
+		{
+			FIntPoint CopyCoords = FIntPoint(X, Y);
+			FIntPoint SrcCoords = CopyCoords + SrcOffsetAdjustment;
+			FIntPoint DestCoords = CopyCoords + DestRegion.Min;
+
+			const int32 SrcIndex = SrcCoords.X + SrcCoords.Y * InSrcDataPitch;
+			const int32 DstIndex = DestCoords.X + DestCoords.Y * DestDataPitch;
+
+			OutData[DstIndex] = InData[SrcIndex];
+		}
+	}
+}
+
 template<class T>
 void TransformImportDataInternal(const TArray<T>& InData, TArray<T>& OutData, const FLandscapeImportResolution& CurrentResolution, const FLandscapeImportResolution& RequiredResolution, ELandscapeImportTransformType TransformType, FIntPoint Offset)
 {
+	// TODO [alan.roach]:  the (-1) applied to Width and Height here is fishy.
 	check(InData.Num() == CurrentResolution.Width * CurrentResolution.Height);
 	if (TransformType == ELandscapeImportTransformType::Resample)
 	{
@@ -337,6 +387,12 @@ void TransformImportDataInternal(const TArray<T>& InData, TArray<T>& OutData, co
 		const FIntRect DestRegion(-OffsetX, -OffsetY, RequiredResolution.Width - OffsetX - 1, RequiredResolution.Height - OffsetY - 1);
 		FLandscapeConfigHelper::ExpandData<T>(InData, OutData, SrcRegion, DestRegion, OffsetX != 0 || OffsetY != 0);
 	}
+	else if (TransformType == ELandscapeImportTransformType::Subregion)
+	{
+		const FIntRect SrcRegion(0, 0, CurrentResolution.Width, CurrentResolution.Height);
+		const FIntPoint DestSize(RequiredResolution.Width, RequiredResolution.Height);
+		TransformImportDataCopySubregion<T>(InData, OutData, SrcRegion, CurrentResolution.Width, Offset, DestSize);
+	}
 	else
 	{
 		OutData = InData;
diff --git a/Engine/Source/Editor/LandscapeEditor/Private/LandscapeRegionUtils.cpp b/Engine/Source/Editor/LandscapeEditor/Private/LandscapeRegionUtils.cpp
--- a/Engine/Source/Editor/LandscapeEditor/Private/LandscapeRegionUtils.cpp
+++ b/Engine/Source/Editor/LandscapeEditor/Private/LandscapeRegionUtils.cpp
@@ -97,11 +97,8 @@ void ForEachComponentByRegion(int32 RegionSize, const TArray<FIntPoint>& Compone
 	}
 }
 
-void ForEachRegion_LoadProcessUnload(ULandscapeInfo* InLandscapeInfo, const FIntRect& InDomain, UWorld* InWorld, TFunctionRef<bool(const FBox&, const TArray<ALandscapeProxy*>)> InRegionFn)
+void ForEachRegion_LoadProcessUnload(ULandscapeInfo* InLandscapeInfo, UWorld* InWorld, TFunctionRef<bool(const FBox&, const TArray<ALandscapeProxy*>)> InRegionFn)
 {
-	const int32 RegionSizeInTexels = InLandscapeInfo->ComponentSizeQuads * InLandscapeInfo->RegionSizeInComponents + 1;
-	const FIntRect RegionCoordinates = InDomain / RegionSizeInTexels;
-
 	TArray<AActor*> Children;
 	InLandscapeInfo->LandscapeActor->GetAttachedActors(Children);
 	TArray<ALocationVolume*> LandscapeRegions;
diff --git a/Engine/Source/Editor/LandscapeEditor/Private/LandscapeRegionUtils.h b/Engine/Source/Editor/LandscapeEditor/Private/LandscapeRegionUtils.h
--- a/Engine/Source/Editor/LandscapeEditor/Private/LandscapeRegionUtils.h
+++ b/Engine/Source/Editor/LandscapeEditor/Private/LandscapeRegionUtils.h
@@ -23,7 +23,7 @@ namespace LandscapeRegionUtils
 	void ForEachComponentByRegion(int32 RegionSize, const TArray<FIntPoint>& ComponentCoordinates, TFunctionRef<bool(const FIntPoint&, const TArray<FIntPoint>&)> RegionFn);
 
 	// Load, Process (call RegionFn) and Unload the Region
-	void ForEachRegion_LoadProcessUnload(ULandscapeInfo* LandscapeInfo, const FIntRect& Domain, UWorld* World, TFunctionRef<bool(const FBox&, const TArray<ALandscapeProxy*>)> RegionFn);
+	void ForEachRegion_LoadProcessUnload(ULandscapeInfo* LandscapeInfo, UWorld* World, TFunctionRef<bool(const FBox&, const TArray<ALandscapeProxy*>)> RegionFn);
 
 	// Number of landscape regions in this landscape.
 	int32 NumLandscapeRegions(ULandscapeInfo* InLandscapeInfo);
