Simplified landscape brush code

#rnx

diff --git a/Engine/Source/Editor/LandscapeEditor/Private/LandscapeEdModePaintTools.cpp b/Engine/Source/Editor/LandscapeEditor/Private/LandscapeEdModePaintTools.cpp
--- a/Engine/Source/Editor/LandscapeEditor/Private/LandscapeEdModePaintTools.cpp
+++ b/Engine/Source/Editor/LandscapeEditor/Private/LandscapeEdModePaintTools.cpp
@@ -78,9 +78,12 @@ public:
 		bDisablePaintingStartupSlowdown = Settings->bDisablePaintingStartupSlowdown;
 	}
 
-	float GetStrength(const ULandscapeEditorObject* UISettings, bool bInIsHeightMap) const
+	float GetStrength(const ULandscapeEditorObject* UISettings) const
 	{
-		return bInIsHeightMap ? UISettings->GetCurrentToolStrength() : FMath::Pow(UISettings->GetCurrentToolStrength(), PaintStrengthGamma);
+		// Apply paint strength gamma for paint strokes only : 
+		return (Target.TargetType == ELandscapeToolTargetType::Heightmap)
+			? UISettings->GetCurrentToolStrength() 
+			: FMath::Pow(UISettings->GetCurrentToolStrength(), PaintStrengthGamma);
 	}
 
 protected:
@@ -219,7 +222,7 @@ public:
 
 		// Adjust strength based on brush size and drawscale, so strength 1 = one hemisphere
 		const float AdjustedStrength = ToolTarget::StrengthMultiplier(this->LandscapeInfo, UISettings->GetCurrentToolBrushRadius());
-		float PaintStrength = this->GetStrength(UISettings, /* bInIsHeightMap = */ false) * Pressure * AdjustedStrength;
+		float PaintStrength = this->GetStrength(UISettings) * Pressure * AdjustedStrength;
 		
 		ValueType DestValue = FWeightmapToolTarget::CacheClass::ClampValue(static_cast<int32>(255.0f * UISettings->WeightTargetValue));
 
@@ -655,12 +658,10 @@ class FLandscapeToolStrokeSmooth : public FLandscapeToolStrokePaintBase<ToolTarg
 {
 	using Super = FLandscapeToolStrokePaintBase<ToolTarget>;
 	using ValueType = typename ToolTarget::CacheClass::DataType;
-	bool bTargetIsHeightmap;
 	FLandscapeLayerDataCache<ToolTarget> LayerDataCache;
 public:
 	FLandscapeToolStrokeSmooth(FEdModeLandscape* InEdMode, FEditorViewportClient* InViewportClient, const FLandscapeToolTarget& InTarget)
 		: Super(InEdMode, InViewportClient, InTarget)
-		, bTargetIsHeightmap(InTarget.TargetType == ELandscapeToolTargetType::Heightmap)
 		, LayerDataCache(InTarget, this->Cache)
 	{
 	}
@@ -677,6 +678,7 @@ public:
 		if (!this->LandscapeInfo) return;
 
 		ALandscape* Landscape = this->LandscapeInfo->LandscapeActor.Get();
+		const bool bTargetIsHeightmap = (this->Target.TargetType == ELandscapeToolTargetType::Heightmap);
 		const bool bCombinedLayerOperation = bTargetIsHeightmap && UISettings->bCombinedLayersOperation && Landscape;
 
 		// Get list of verts to update
@@ -708,7 +710,7 @@ public:
 		LayerDataCache.Read(X1, Y1, X2, Y2, Data);
 		const TArray<ValueType> ReadData { Data };
 		
-		const float ToolStrength = FMath::Clamp<float>(this->GetStrength(UISettings, bTargetIsHeightmap) * Pressure, 0.0f, 1.0f);
+		const float ToolStrength = FMath::Clamp<float>(this->GetStrength(UISettings) * Pressure, 0.0f, 1.0f);
 
 		// Apply the brush
 		if (UISettings->bDetailSmooth)
@@ -809,17 +811,15 @@ class FLandscapeToolStrokeFlatten : public FLandscapeToolStrokePaintBase<ToolTar
 	FVector FlattenNormal;
 	float FlattenPlaneDist;
 	bool bInitializedFlattenValue;
-	bool bTargetIsHeightmap;
 	FLandscapeLayerDataCache<ToolTarget> LayerDataCache;
 
 public:
 	FLandscapeToolStrokeFlatten(FEdModeLandscape* InEdMode, FEditorViewportClient* InViewportClient, const FLandscapeToolTarget& InTarget)
 		: FLandscapeToolStrokePaintBase<ToolTarget>(InEdMode, InViewportClient, InTarget)
 		, bInitializedFlattenValue(false)
-		, bTargetIsHeightmap(InTarget.TargetType == ELandscapeToolTargetType::Heightmap)
 		, LayerDataCache(InTarget, this->Cache)
 	{
-		if (InEdMode->UISettings->bUseFlattenTarget && bTargetIsHeightmap)
+		if (InEdMode->UISettings->bUseFlattenTarget && (InTarget.TargetType == ELandscapeToolTargetType::Heightmap))
 		{
 			FTransform LocalToWorld = InTarget.LandscapeInfo->GetLandscapeProxy()->ActorToWorld();
 			float Height = static_cast<float>((InEdMode->UISettings->GetFlattenTarget(/*bInReturnPreviewValueIfActive = */false) - LocalToWorld.GetTranslation().Z) / LocalToWorld.GetScale3D().Z);
@@ -838,6 +838,7 @@ public:
 		if (!this->LandscapeInfo) return;
 
 		ALandscape* Landscape = this->LandscapeInfo->LandscapeActor.Get();
+		const bool bTargetIsHeightmap = (this->Target.TargetType == ELandscapeToolTargetType::Heightmap);
 		const bool bCombinedLayerOperation = bTargetIsHeightmap && UISettings->bCombinedLayersOperation && Landscape;
 
 		// Can't use slope if use Flatten Target because no normal is provided
@@ -938,7 +939,7 @@ public:
 		LayerDataCache.Initialize(this->LandscapeInfo, bCombinedLayerOperation);
 		LayerDataCache.Read(X1, Y1, X2, Y2, Data);
 
-		const float PaintStrength = this->GetStrength(UISettings, bTargetIsHeightmap);
+		const float PaintStrength = this->GetStrength(UISettings);
 
 		// Apply the brush
 		for (int32 Y = BrushInfo.GetBounds().Min.Y; Y < BrushInfo.GetBounds().Max.Y; Y++)
@@ -1314,7 +1315,7 @@ public:
 		CA_SUPPRESS(6326);
 		bool bUseWeightTargetValue = UISettings->bUseWeightTargetValue && ToolTarget::TargetType == ELandscapeToolTargetType::Weightmap;
 
-		const float PaintStrength = this->GetStrength(UISettings, /* bInIsHeighthMap = */ false);
+		const float PaintStrength = this->GetStrength(UISettings);
 
 		// Apply the brush
 		for (int32 Y = BrushInfo.GetBounds().Min.Y; Y < BrushInfo.GetBounds().Max.Y; Y++)
diff --git a/Engine/Source/Editor/LandscapeEditor/Private/LandscapeEdModeTools.h b/Engine/Source/Editor/LandscapeEditor/Private/LandscapeEdModeTools.h
--- a/Engine/Source/Editor/LandscapeEditor/Private/LandscapeEdModeTools.h
+++ b/Engine/Source/Editor/LandscapeEditor/Private/LandscapeEdModeTools.h
@@ -551,7 +551,7 @@ public:
 		if (!HasCachedData(X1, Y1, X2, Y2))
 		{
 			// The cache needs to be expanded, compute the new bounds : 
-			// The bounds we calculate here need to be what would be the result of calling CacheData with this region, meaning that they should include the previous bounds. This will let us pass the correct region of interest to PrepareRegionForCaching
+			// The bounds we calculate here need to be what would be the result of calling CacheData with this region, meaning that they should include the previous bounds
 			FIntRect NewCacheBounds(bIsValid ? FMath::Min(X1, CachedX1) : X1, bIsValid ? FMath::Min(Y1, CachedY1) : Y1, bIsValid ? FMath::Max(X2, CachedX2) : X2, bIsValid ? FMath::Max(Y2, CachedY2) : Y2);
 
 			// SelectiveRender only supports heightmaps at the moment.
@@ -1210,6 +1210,7 @@ public:
 	{
 		Collector.AddReferencedObject(LandscapeInfo);
 	}
+	
 	virtual FString GetReferencerName() const override
 	{
 		return TEXT("FLandscapeToolStrokeBase");
