Use standard icons for create/delete buttons for target layers in landscape editor mode

#rnx

diff --git a/Engine/Source/Editor/EditorStyle/Private/StarshipStyle.cpp b/Engine/Source/Editor/EditorStyle/Private/StarshipStyle.cpp
--- a/Engine/Source/Editor/EditorStyle/Private/StarshipStyle.cpp
+++ b/Engine/Source/Editor/EditorStyle/Private/StarshipStyle.cpp
@@ -7308,9 +7308,6 @@ void FStarshipEditorStyle::FStyle::SetupLandscapeEditorStyle()
 		Set("LandscapeEditor.Target_Visibility", new IMAGE_BRUSH_SVG("Icons/Landscape/Landscape_Target_Visibility_64x", Icon48x48));
 		Set("LandscapeEditor.Target_Unknown",    new IMAGE_BRUSH_SVG("Icons/Landscape/Landscape_Target_Unknown_64x",    Icon48x48));
 
-		Set("LandscapeEditor.Target_Create",     new IMAGE_BRUSH("Icons/PlusSymbol_12x", Icon12x12));
-		Set("LandscapeEditor.Target_Delete",     new IMAGE_BRUSH("Icons/Cross_12x", Icon12x12));
-
 		Set("LandscapeEditor.Target_DisplayOrder.Default.Ascending", new IMAGE_BRUSH_SVG("Icons/Landscape/Landscape_Target_DisplayOrder_Default_Ascending_16x", Icon16x16));
 		Set("LandscapeEditor.Target_DisplayOrder.Default.Descending", new IMAGE_BRUSH_SVG("Icons/Landscape/Landscape_Target_DisplayOrder_Default_Descending_16x", Icon16x16));
 		Set("LandscapeEditor.Target_DisplayOrder.Alphabetical.Ascending", new IMAGE_BRUSH_SVG("Icons/Landscape/Landscape_Target_DisplayOrder_Alphabetical_Ascending_16x", Icon16x16));
diff --git a/Engine/Source/Editor/LandscapeEditor/Private/LandscapeEditorDetailCustomization_ImportLayers.cpp b/Engine/Source/Editor/LandscapeEditor/Private/LandscapeEditorDetailCustomization_ImportLayers.cpp
--- a/Engine/Source/Editor/LandscapeEditor/Private/LandscapeEditorDetailCustomization_ImportLayers.cpp
+++ b/Engine/Source/Editor/LandscapeEditor/Private/LandscapeEditorDetailCustomization_ImportLayers.cpp
@@ -168,7 +168,7 @@ void FLandscapeEditorStructCustomization_FLandscapeImportLayer::CustomizeChildre
 										.OnClicked_Lambda([PropertyHandle_LayerInfo, LayerName]() { return FLandscapeEditorStructCustomization_FLandscapeImportLayer::OnImportLayerCreateClicked(PropertyHandle_LayerInfo, LayerName); })
 										[
 											SNew(SImage)
-												.Image(FAppStyle::GetBrush("LandscapeEditor.Target_Create"))
+												.Image(FAppStyle::GetBrush("Icons.Plus"))
 												.ColorAndOpacity(FSlateColor::UseForeground())
 										]
 								]
diff --git a/Engine/Source/Editor/LandscapeEditor/Private/LandscapeEditorDetailCustomization_TargetLayers.cpp b/Engine/Source/Editor/LandscapeEditor/Private/LandscapeEditorDetailCustomization_TargetLayers.cpp
--- a/Engine/Source/Editor/LandscapeEditor/Private/LandscapeEditorDetailCustomization_TargetLayers.cpp
+++ b/Engine/Source/Editor/LandscapeEditor/Private/LandscapeEditorDetailCustomization_TargetLayers.cpp
@@ -1281,6 +1281,16 @@ TSharedPtr<SWidget> FLandscapeEditorCustomNodeBuilder_TargetLayers::GenerateRow(
 	{
 		InlineTextBlocks.Add(TSharedPtr<SInlineEditableTextBlock>());
 
+		TSharedRef<SWidget> DeleteButton = PropertyCustomizationHelpers::MakeDeleteButton(
+			FSimpleDelegate::CreateStatic(&FLandscapeEditorCustomNodeBuilder_TargetLayers::OnTargetLayerDeleteClicked, Target),
+			LOCTEXT("Tooltip_Delete", "Delete Layer"));
+
+		TSharedRef<SWidget> CreateButton = PropertyCustomizationHelpers::MakeCustomButton(
+			FAppStyle::GetBrush(TEXT("Icons.Plus")),
+			FSimpleDelegate::CreateStatic(&FLandscapeEditorCustomNodeBuilder_TargetLayers::OnTargetLayerCreateClicked, Target),
+			LOCTEXT("Tooltip_Create", "Create Layer Info"),
+			TAttribute<bool>::Create(TAttribute<bool>::FGetter::CreateStatic(&FLandscapeEditorCustomNodeBuilder_TargetLayers::GetTargetLayerCreateEnabled, Target)));
+
 		RowWidget = SNew(SLandscapeEditorSelectableBorder)
 			.Padding(0.0f)
 			.VAlign(VAlign_Center)
@@ -1442,34 +1452,13 @@ TSharedPtr<SWidget> FLandscapeEditorCustomNodeBuilder_TargetLayers::GenerateRow(
 									.AutoWidth()
 									.VAlign(VAlign_Center)
 									[
-										SNew(SButton)
-											.ButtonStyle(FAppStyle::Get(), "HoverHintOnly")
-											.ContentPadding(4.0f)
-											.ForegroundColor(FSlateColor::UseForeground())
-											.IsFocusable(false)
-											.ToolTipText(LOCTEXT("Tooltip_Create", "Create Layer Info"))
-											.IsEnabled_Static(&FLandscapeEditorCustomNodeBuilder_TargetLayers::GetTargetLayerCreateEnabled, Target)
-											.OnClicked_Static(&FLandscapeEditorCustomNodeBuilder_TargetLayers::OnTargetLayerCreateClicked, Target)
-											[
-												SNew(SImage)
-													.Image(FAppStyle::GetBrush("LandscapeEditor.Target_Create"))
-											]
+										CreateButton
 									]
 									+ SHorizontalBox::Slot()
 									.AutoWidth()
 									.VAlign(VAlign_Center)
 									[
-										SNew(SButton)
-											.ButtonStyle(FAppStyle::Get(), "HoverHintOnly")
-											.ContentPadding(4.0f)
-											.ForegroundColor(FSlateColor::UseForeground())
-											.IsFocusable(false)
-											.ToolTipText(LOCTEXT("Tooltip_Delete", "Delete Layer"))
-											.OnClicked_Static(&FLandscapeEditorCustomNodeBuilder_TargetLayers::OnTargetLayerDeleteClicked, Target)
-											[
-												SNew(SImage)
-													.Image(FAppStyle::GetBrush("LandscapeEditor.Target_Delete"))
-											]
+										DeleteButton
 									]
 							]
 							+ SVerticalBox::Slot()
@@ -2062,7 +2051,7 @@ bool FLandscapeEditorCustomNodeBuilder_TargetLayers::GetTargetLayerCreateEnabled
 	return false;
 }
 
-FReply FLandscapeEditorCustomNodeBuilder_TargetLayers::OnTargetLayerCreateClicked(const TSharedRef<FLandscapeTargetListInfo> Target)
+void FLandscapeEditorCustomNodeBuilder_TargetLayers::OnTargetLayerCreateClicked(const TSharedRef<FLandscapeTargetListInfo> Target)
 {
 	check(!Target->LayerInfoObj.IsValid());
 
@@ -2087,8 +2076,6 @@ FReply FLandscapeEditorCustomNodeBuilder_TargetLayers::OnTargetLayerCreateClicke
 			CreateTargetLayerInfoAsset(Target, PackageName, FileName.ToString());
 		}
 	}
-
-	return FReply::Handled();
 }
 
 void FLandscapeEditorCustomNodeBuilder_TargetLayers::CreateTargetLayerInfoAsset(const TSharedRef<FLandscapeTargetListInfo> Target, const FString& PackageName, const FString& FileName)
@@ -2120,7 +2107,7 @@ void FLandscapeEditorCustomNodeBuilder_TargetLayers::CreateTargetLayerInfoAsset(
 	}
 }
 
-FReply FLandscapeEditorCustomNodeBuilder_TargetLayers::OnTargetLayerDeleteClicked(const TSharedRef<FLandscapeTargetListInfo> Target)
+void FLandscapeEditorCustomNodeBuilder_TargetLayers::OnTargetLayerDeleteClicked(const TSharedRef<FLandscapeTargetListInfo> Target)
 {
 	check(Target->LandscapeInfo.IsValid());
 
@@ -2138,8 +2125,6 @@ FReply FLandscapeEditorCustomNodeBuilder_TargetLayers::OnTargetLayerDeleteClicke
 			LandscapeEdMode->UpdateTargetList();
 		}
 	}
-
-	return FReply::Handled();
 }
 
 FSlateColor FLandscapeEditorCustomNodeBuilder_TargetLayers::GetLayerUsageDebugColor(const TSharedRef<FLandscapeTargetListInfo> Target)
diff --git a/Engine/Source/Editor/LandscapeEditor/Private/LandscapeEditorDetailCustomization_TargetLayers.h b/Engine/Source/Editor/LandscapeEditor/Private/LandscapeEditorDetailCustomization_TargetLayers.h
--- a/Engine/Source/Editor/LandscapeEditor/Private/LandscapeEditorDetailCustomization_TargetLayers.h
+++ b/Engine/Source/Editor/LandscapeEditor/Private/LandscapeEditorDetailCustomization_TargetLayers.h
@@ -108,9 +108,9 @@ protected:
 	static void OnTargetLayerSetObject(const FAssetData& AssetData, const TSharedRef<FLandscapeTargetListInfo> Target); 
 	static EVisibility GetTargetLayerInfoSelectorVisibility(const TSharedRef<FLandscapeTargetListInfo> Target);
 	static bool GetTargetLayerCreateEnabled(const TSharedRef<FLandscapeTargetListInfo> Target);
-	static FReply OnTargetLayerCreateClicked(const TSharedRef<FLandscapeTargetListInfo> Target);
+	static void OnTargetLayerCreateClicked(const TSharedRef<FLandscapeTargetListInfo> Target);
 	static void CreateTargetLayerInfoAsset(const TSharedRef<FLandscapeTargetListInfo> Target, const FString& PackageName, const FString& FileName);
-	static FReply OnTargetLayerDeleteClicked(const TSharedRef<FLandscapeTargetListInfo> Target);
+	static void OnTargetLayerDeleteClicked(const TSharedRef<FLandscapeTargetListInfo> Target);
 	static FSlateColor GetLayerUsageDebugColor(const TSharedRef<FLandscapeTargetListInfo> Target);
 	static EVisibility GetDebugModeLayerUsageVisibility(const TSharedRef<FLandscapeTargetListInfo> Target);
 	static EVisibility GetDebugModeLayerUsageVisibility_Invert(const TSharedRef<FLandscapeTargetListInfo> Target);
