Changed landscape overridable properties UX : Reset to Default now affects the overridden state instead of the property value. Supports multi-select

diff --git a/Engine/Source/Editor/LandscapeEditor/Private/LandscapeProxyUIDetails.cpp b/Engine/Source/Editor/LandscapeEditor/Private/LandscapeProxyUIDetails.cpp
--- a/Engine/Source/Editor/LandscapeEditor/Private/LandscapeProxyUIDetails.cpp
+++ b/Engine/Source/Editor/LandscapeEditor/Private/LandscapeProxyUIDetails.cpp
@@ -536,6 +536,26 @@ void FLandscapeProxyUIDetails::CustomizeDetails( IDetailLayoutBuilder& DetailBui
 						]
 					]
 					.IsValueEnabled(EnabledAttribute);
+
+					// The Reset To Default button corresponds to the overridden state of the overridable property : 
+					DetailRow->OverrideResetToDefault(FResetToDefaultOverride::Create(
+						FIsResetToDefaultVisible::CreateLambda([EditingStreamingProxies, PropertyName](TSharedPtr<IPropertyHandle> InPropertyHandle) -> bool
+							{
+								return Algo::AnyOf(EditingStreamingProxies, [PropertyName](TWeakObjectPtr<ALandscapeStreamingProxy> StreamingProxy) { return StreamingProxy.IsValid() && StreamingProxy->IsSharedPropertyOverridden(PropertyName); });
+							}),
+						FResetToDefaultHandler::CreateLambda([EditingStreamingProxies, PropertyName](TSharedPtr<IPropertyHandle> InPropertyHandle)
+							{
+								const FScopedTransaction Transaction(LOCTEXT("ResetSetSharedPropertyOverride", "Reset Property Override"));
+								for (const TWeakObjectPtr<ALandscapeStreamingProxy> StreamingProxy : EditingStreamingProxies)
+								{
+									if (!StreamingProxy.IsValid())
+									{
+										continue;
+									}
+
+									StreamingProxy->SetSharedPropertyOverride(PropertyName, false);
+								}
+							})));
 				}
 			}
 		}
