Added support for custom thumbnail for landscape layer info assets (if present, it will be used, regardless of the bDisplayTargetLayerThumbnails project setting).
This allows users to generate custom thumbnails even when the automatic target layer thumbnail generation is failing (for example if the landscape material outputs to RVT)
Also, fixed missing fallback icon when not displaying target layer thumbnails

diff --git a/Engine/Source/Editor/LandscapeEditor/Private/LandscapeEditorDetailCustomization_ImportLayers.cpp b/Engine/Source/Editor/LandscapeEditor/Private/LandscapeEditorDetailCustomization_ImportLayers.cpp
--- a/Engine/Source/Editor/LandscapeEditor/Private/LandscapeEditorDetailCustomization_ImportLayers.cpp
+++ b/Engine/Source/Editor/LandscapeEditor/Private/LandscapeEditorDetailCustomization_ImportLayers.cpp
@@ -27,6 +27,7 @@
 #include "Framework/Application/SlateApplication.h"
 #include "Framework/MultiBox/MultiBoxBuilder.h"
 #include "Misc/MessageDialog.h"
+#include "ObjectTools.h"
 
 #define LOCTEXT_NAMESPACE "LandscapeEditor.ImportLayers"
 
@@ -93,10 +94,19 @@ void FLandscapeEditorStructCustomization_FLandscapeImportLayer::CustomizeChildre
 		LayerNameText = NSLOCTEXT("PropertyEditor", "MultipleValues", "Multiple Values");
 	}
 
+	UObject* ThumbnailObject = nullptr;
+
 	UObject* ThumbnailMIC = nullptr;
 	Result = PropertyHandle_ThumbnailMIC->GetValue(ThumbnailMIC);
 	checkSlow(Result == FPropertyAccess::Success);
 
+	UObject* LayerInfo = nullptr;
+	Result = PropertyHandle_LayerInfo->GetValue(LayerInfo);
+	checkSlow(Result == FPropertyAccess::Success);
+
+	// If the layer info is valid and has a custom thumbnail, it's considered as an override, so use this as the thumbnail object : 
+	ThumbnailObject = ((LayerInfo != nullptr) && ThumbnailTools::AssetHasCustomThumbnail(LayerInfo->GetFullName())) ? LayerInfo : ThumbnailMIC;
+
 	auto CreateFilePathNameContentWidget = [&](bool bImport)
 		{
 			return SNew(SHorizontalBox)
@@ -119,7 +129,7 @@ void FLandscapeEditorStructCustomization_FLandscapeImportLayer::CustomizeChildre
 				.VAlign(VAlign_Center)
 				.Padding(FMargin(2))
 				[
-					SNew(SLandscapeAssetThumbnail, ThumbnailMIC, StructCustomizationUtils.GetThumbnailPool().ToSharedRef(), FName("LandscapeEditor.Target_Weightmap"))
+					SNew(SLandscapeAssetThumbnail, ThumbnailObject, StructCustomizationUtils.GetThumbnailPool().ToSharedRef(), FName("LandscapeEditor.Target_Weightmap"))
 						.ThumbnailSize(FIntPoint(48, 48))
 				]
 				+ SHorizontalBox::Slot()
diff --git a/Engine/Source/Editor/LandscapeEditor/Private/LandscapeEditorDetailCustomization_TargetLayers.cpp b/Engine/Source/Editor/LandscapeEditor/Private/LandscapeEditorDetailCustomization_TargetLayers.cpp
--- a/Engine/Source/Editor/LandscapeEditor/Private/LandscapeEditorDetailCustomization_TargetLayers.cpp
+++ b/Engine/Source/Editor/LandscapeEditor/Private/LandscapeEditorDetailCustomization_TargetLayers.cpp
@@ -1281,6 +1281,16 @@ TSharedPtr<SWidget> FLandscapeEditorCustomNodeBuilder_TargetLayers::GenerateRow(
 	{
 		InlineTextBlocks.Add(TSharedPtr<SInlineEditableTextBlock>());
 
+		ULandscapeLayerInfoObject* LayerInfo = Target->LayerInfoObj.Get();
+		UMaterialInstanceConstant* ThumbnailMIC = Target->ThumbnailMIC.Get();
+		UObject* ThumbnailObject = nullptr;
+
+		if (Target->bIsLayerReferencedByMaterial)
+		{
+			// If the layer info is valid and has a custom thumbnail, it's considered as an override, so use this as the thumbnail object : 
+			ThumbnailObject = ((LayerInfo != nullptr) && ThumbnailTools::AssetHasCustomThumbnail(LayerInfo->GetFullName())) ? static_cast<UObject*>(LayerInfo) : static_cast<UObject*>(ThumbnailMIC);
+		}
+
 		TSharedRef<SWidget> DeleteButton = PropertyCustomizationHelpers::MakeDeleteButton(
 			FSimpleDelegate::CreateStatic(&FLandscapeEditorCustomNodeBuilder_TargetLayers::OnTargetLayerDeleteClicked, Target),
 			LOCTEXT("Tooltip_Delete", "Delete Layer"));
@@ -1334,7 +1344,7 @@ TSharedPtr<SWidget> FLandscapeEditorCustomNodeBuilder_TargetLayers::GenerateRow(
 					.VAlign(VAlign_Center)
 					.Padding(FMargin(2))
 					[
-						SNew(SLandscapeAssetThumbnail, Target->bIsLayerReferencedByMaterial ? Target->ThumbnailMIC.Get() : nullptr, ThumbnailPool, Target->bIsLayerReferencedByMaterial ? FName("LandscapeEditor.Target_Weightmap") : FName("LandscapeEditor.Target_Unknown"))
+						SNew(SLandscapeAssetThumbnail, ThumbnailObject, ThumbnailPool, Target->bIsLayerReferencedByMaterial ? FName("LandscapeEditor.Target_Weightmap") : FName("LandscapeEditor.Target_Unknown"))
 						.Visibility_Static(&FLandscapeEditorCustomNodeBuilder_TargetLayers::GetDebugModeLayerUsageVisibility_Invert, Target)
 						.ThumbnailSize(FIntPoint(48, 48))
 						// Open landscape layer info asset on double-click on the thumbnail : 
diff --git a/Engine/Source/Runtime/Landscape/Private/Landscape.cpp b/Engine/Source/Runtime/Landscape/Private/Landscape.cpp
--- a/Engine/Source/Runtime/Landscape/Private/Landscape.cpp
+++ b/Engine/Source/Runtime/Landscape/Private/Landscape.cpp
@@ -7294,6 +7294,16 @@ void ALandscapeProxy::OnLayerInfoObjectDataChanged(const FOnLandscapeLayerInfoDa
 	{
 		LandscapeActor->RequestLayersContentUpdateForceAll(ELandscapeLayerUpdateMode::Update_All, InParams.bUserTriggered);
 	}
+
+	if (bUpdateAll && (InParams.PropertyChangedEvent.ChangeType != EPropertyChangeType::Interactive))
+	{
+		// There's no real hook to attach to when a custom thumbnail is generated (or cleared) for a given asset and we need to regenerate the thumbnails when that happens.
+		//  PostEditChange is called on the layer info object when this happens, though, so we can ask for regenerating thumbnails when that happens. Not the most efficient
+		//  thing but landscape layer changes are pretty rare anyway
+		ILandscapeModule& LandscapeModule = FModuleManager::GetModuleChecked<ILandscapeModule>("Landscape");
+		LandscapeModule.GetLandscapeEditorServices()->RegenerateLayerThumbnails();
+		LandscapeModule.GetLandscapeEditorServices()->RefreshDetailPanel();
+	}
 }
 
 void ALandscapeProxy::InvalidatePhysicalMaterial()
