Landscape Editor - Fixed issue with Copy/Paste Gizmo recommended Heightmap size dropdown not accessible after using copy/paste shortcuts
* Renamed legacy Layers to Target Layers

#Synced-CL 47733618

diff --git a/Engine/Source/Editor/LandscapeEditor/Private/LandscapeEditorDetailCustomization_CopyPaste.cpp b/Engine/Source/Editor/LandscapeEditor/Private/LandscapeEditorDetailCustomization_CopyPaste.cpp
--- a/Engine/Source/Editor/LandscapeEditor/Private/LandscapeEditorDetailCustomization_CopyPaste.cpp
+++ b/Engine/Source/Editor/LandscapeEditor/Private/LandscapeEditorDetailCustomization_CopyPaste.cpp
@@ -40,6 +40,16 @@ void FLandscapeEditorDetailCustomization_CopyPaste::CustomizeDetails(IDetailLayo
 		return;
 	}
 
+	// Each time Ctrl-C/Ctrl-V are used the CopyPaste customization re-runs. The GuessedDimensionComboList is 
+	// stored within this object so it resets automatically. When there is a valid file path, regenerate the dimensions each draw
+	if (FEdModeLandscape* LandscapeEdMode = GetEditorMode())
+	{
+		if (!LandscapeEdMode->UISettings->GizmoHeightmapFilenameString.IsEmpty())
+		{
+			GenerateGuessDimensionList();
+		}
+	}
+
 	IDetailCategoryBuilder& ToolsCategory = DetailBuilder.EditCategory("Tool Settings");
 
 	ToolsCategory.AddCustomRow(LOCTEXT("CopyToGizmo", "Copy Data to Gizmo"))
@@ -217,9 +227,9 @@ void FLandscapeEditorDetailCustomization_CopyPaste::CustomizeDetails(IDetailLayo
 			]
 	];
 
-	TSharedRef<IPropertyHandle> PropertyHandle_ImportLayers = DetailBuilder.GetProperty(GET_MEMBER_NAME_CHECKED(ULandscapeEditorObject, GizmoImportLayers));
-	DetailBuilder.HideProperty(PropertyHandle_ImportLayers);
-	GizmoImportExportGroup.AddPropertyRow(PropertyHandle_ImportLayers);
+	TSharedRef<IPropertyHandle> PropertyHandle_ImportTargetLayers = DetailBuilder.GetProperty(GET_MEMBER_NAME_CHECKED(ULandscapeEditorObject, GizmoImportLayers));
+	DetailBuilder.HideProperty(PropertyHandle_ImportTargetLayers);
+	GizmoImportExportGroup.AddPropertyRow(PropertyHandle_ImportTargetLayers);
 
 	GizmoImportExportGroup.AddWidgetRow()
 	.FilterString(LOCTEXT("GizmoImportExport", "ImportExport"))
@@ -485,35 +495,38 @@ FReply FLandscapeEditorDetailCustomization_CopyPaste::OnGizmoImportButtonClicked
 			TArray<TArray<uint8> > LayerDataArrays;
 			TArray<uint8*> LayerDataPtrs;
 
-			for (int32 LayerIndex = 0; LayerIndex < LandscapeEdMode->UISettings->GizmoImportLayers.Num(); LayerIndex++)
+			for (int32 TargetLayerIndex = 0; TargetLayerIndex < LandscapeEdMode->UISettings->GizmoImportLayers.Num(); TargetLayerIndex++)
 			{
-				const FGizmoImportLayer& Layer = LandscapeEdMode->UISettings->GizmoImportLayers[LayerIndex];
-				FString LayerName = Layer.LayerName.Replace(TEXT(" "), TEXT(""));
-				if (LayerName == TEXT(""))
+				const FGizmoImportLayer& TargetLayer = LandscapeEdMode->UISettings->GizmoImportLayers[TargetLayerIndex];
+				FString TargetLayerName = TargetLayer.LayerName.Replace(TEXT(" "), TEXT(""));
+				if (TargetLayerName.IsEmpty() && !TargetLayer.bNoImport)
 				{
 					FMessageDialog::Open(EAppMsgType::Ok,
-						FText::Format(NSLOCTEXT("UnrealEd", "LandscapeImport_BadLayerName", "You must enter a name for the layer being imported from {0}."), FText::FromString(Layer.LayerFilename)));
+						FText::Format(NSLOCTEXT("UnrealEd", "LandscapeImport_BadTargetLayerName", "You must enter a name for the target layer being imported from {0}."), FText::FromString(TargetLayer.LayerFilename)));
 					return FReply::Handled();
 				}
 
-				if (Layer.LayerFilename != TEXT("") && !Layer.bNoImport)
+				if (!TargetLayer.LayerFilename.IsEmpty() && !TargetLayer.bNoImport)
 				{
 					TArray<uint8>* LayerData = new(LayerDataArrays)(TArray<uint8>);
-					FFileHelper::LoadFileToArray(*LayerData, *Layer.LayerFilename);
+					FFileHelper::LoadFileToArray(*LayerData, *TargetLayer.LayerFilename);
 
 					if (LayerData->Num() != (LandscapeEdMode->UISettings->GizmoImportSize.X * LandscapeEdMode->UISettings->GizmoImportSize.Y))
 					{
-						const FText MessageFormat = NSLOCTEXT("UnrealEd", "LandscapeImport_BadLayerSize", "Layer {0} file size does not match.\nExpected {1} entries but file contains {2}.");
+						const FText MessageFormat = NSLOCTEXT("UnrealEd", "LandscapeImport_BadTargetLayerSize", "Target Layer {0} file size does not match.\nExpected {1} entries but file contains {2}.");
 						const FText Message = FText::Format(MessageFormat, 
-							FText::FromString(Layer.LayerFilename),
+							FText::FromString(TargetLayer.LayerFilename),
 							FText::AsNumber(LandscapeEdMode->UISettings->GizmoImportSize.X * LandscapeEdMode->UISettings->GizmoImportSize.Y), 
 							FText::AsNumber(LayerData->Num()));
 						FMessageDialog::Open(EAppMsgType::Ok, Message);
 						return FReply::Handled();
 					}
 
-					LayerInfos.Add(LandscapeEdMode->CurrentToolTarget.LandscapeInfo->GetLayerInfoByName(FName(*LayerName)));
-					LayerDataPtrs.Add(&(*LayerData)[0]);
+					if (ULandscapeLayerInfoObject* LayerInfoObject = LandscapeEdMode->CurrentToolTarget.LandscapeInfo->GetLayerInfoByName(FName(*TargetLayerName)))
+					{
+						LayerInfos.Add(LayerInfoObject);
+						LayerDataPtrs.Add(&(*LayerData)[0]);
+					}
 				}
 			}
 
diff --git a/Engine/Source/Editor/LandscapeEditor/Public/LandscapeEditorObject.h b/Engine/Source/Editor/LandscapeEditor/Public/LandscapeEditorObject.h
--- a/Engine/Source/Editor/LandscapeEditor/Public/LandscapeEditorObject.h
+++ b/Engine/Source/Editor/LandscapeEditor/Public/LandscapeEditorObject.h
@@ -156,10 +156,10 @@ struct FGizmoImportLayer
 {
 	GENERATED_USTRUCT_BODY()
 
-	UPROPERTY(Category="Import", EditAnywhere)
+	UPROPERTY(Category="Import", EditAnywhere, meta = (DisplayName = "Target Layer Filename"))
 	FString LayerFilename;
 
-	UPROPERTY(Category="Import", EditAnywhere)
+	UPROPERTY(Category="Import", EditAnywhere, meta = (DisplayName = "Target Layer Name", ToolTip = "Layer Name must match an existing target layer on the Landscape or the import will fail."))
 	FString LayerName;
 
 	UPROPERTY(Category="Import", EditAnywhere)
@@ -499,7 +499,7 @@ class ULandscapeEditorObject : public UObject
 	UPROPERTY(Category="Tool Settings", EditAnywhere, NonTransactional, AdvancedDisplay, meta=(DisplayName="Heightmap Size", ShowForTools="CopyPaste"))
 	FIntPoint GizmoImportSize;
 
-	UPROPERTY(Category="Tool Settings", EditAnywhere, NonTransactional, AdvancedDisplay, meta=(DisplayName="Layers", ShowForTools="CopyPaste"))
+	UPROPERTY(Category="Tool Settings", EditAnywhere, NonTransactional, AdvancedDisplay, meta=(DisplayName="Target Layers", ShowForTools="CopyPaste"))
 	TArray<FGizmoImportLayer> GizmoImportLayers;
 
 	// Mirror Tool
