Landscape mode : paint tools : new exclusive paint mode for weight-blended layers
* Added alternate modifier behavior to let the user "exclusive-paint" weight-blended layers from the same blend group, that is, paint a given (set of) target layer and "invert-paint" the other target layers from the same blend group, similar to the behavior in prior UE versions, except it's now possible to turn it on or off with a keybind, in order to let users tweak target layers individually
* Erasing paint now doesn't re-balance other layers anymore, it simply erases the currently painted layer, which is much more natural
* By default, exclusive paint mode is triggered by having the Ctrl key held but it can also be changed by the user in order to have it triggered only when the Ctrl key is *not* held (i.e. same behavior as UE5.6) or forced as active or inactive (i.e. no Ctrl key bind), which can be useful when the user uses another Landscape Editor Control Type editor preference than the default, which reserves usage of the Ctrl key for other purposes
* There's now an icon on the target layer thumbnail indicating the type of action performed by the current tool : normal or inverse. For example, shift-painting (i.e. erasing) will be considered as "inverse". This helps the user understand the exclusive paint mode behavior, since the affected layer is marked as normals while the other weight-blended layers from the same blend group are marked "inverse"
* There's now a color strip on the target layer thumbnails (similar-looking to the actual asset thumbnails, with a random color, associated with the blend method and group name, i.e. consistent from one run to another) that gives a hint about the layer's blend group (if any). This helps the user understand which of the target layers will be affected by exclusive painting
* Fixed bug with landscape component allow list addition/removal ('=' / '-' shortcut when using the paint tools) where removal was incorrectly bypassed, leading to the tool being basically useless

Misc:
* Fixed landscape target layer list not updating when changing the blend method or group
* Fixed weird-looking tooltip on target layer (because it was using the preview material MIC's "asset" tooltip)
* Fixed a bunch of local landscape settings that were incorrectly serialized in the .ini file, leading to a loss of settings
* Various fixes for landscape UX
* Lots of (unlikely) bug fixes suggested by BugHawk
* Refactor landscape tool code to allow multiple target layers to be written at once. Not implemented on the UX side but the tools now support this

#Synced-CL 48524909

diff --git a/Engine/Source/Editor/EditorStyle/Private/StarshipStyle.cpp b/Engine/Source/Editor/EditorStyle/Private/StarshipStyle.cpp
--- a/Engine/Source/Editor/EditorStyle/Private/StarshipStyle.cpp
+++ b/Engine/Source/Editor/EditorStyle/Private/StarshipStyle.cpp
@@ -3309,6 +3309,7 @@ void FStarshipEditorStyle::FStyle::SetupMenuBarStyles()
 void FStarshipEditorStyle::FStyle::SetupGeneralIcons()
 {
 	Set("Plus", new IMAGE_BRUSH("Icons/PlusSymbol_12x", Icon12x12));
+	Set("Minus", new IMAGE_BRUSH("Icons/MinusSymbol_12x", Icon12x12));
 	Set("Cross", new IMAGE_BRUSH("Icons/Cross_12x", Icon12x12));
 	Set("ArrowUp", new IMAGE_BRUSH("Icons/ArrowUp_12x", Icon12x12));
 	Set("ArrowDown", new IMAGE_BRUSH("Icons/ArrowDown_12x", Icon12x12));
@@ -7303,6 +7304,9 @@ void FStarshipEditorStyle::FStyle::SetupLandscapeEditorStyle()
 		Set("LandscapeEditor.Brush.AffectsVisibilityLayer.Enabled", new IMAGE_BRUSH("Icons/icon_Landscape_Affects_VisibilityLayer_Enabled_16x", Icon16x16));
 		Set("LandscapeEditor.Brush.AffectsVisibilityLayer.Disabled", new IMAGE_BRUSH("Icons/icon_Landscape_Affects_VisibilityLayer_Disabled_16x", Icon16x16));
 
+		Set("LandscapeEditor.Brush.Action_Regular", new IMAGE_BRUSH_SVG("Icons/Landscape/Landscape_Brush_Action_Regular_16x", Icon16x16));
+		Set("LandscapeEditor.Brush.Action_Inverse", new IMAGE_BRUSH_SVG("Icons/Landscape/Landscape_Brush_Action_Inverse_16x", Icon16x16));
+
 		Set("LandscapeEditor.InspectedObjects.ShowDetails", new CORE_IMAGE_BRUSH_SVG("Starship/Common/settings", Icon16x16));
 
 		// Target List
diff --git a/Engine/Source/Editor/LandscapeEditor/LandscapeEditor.Build.cs b/Engine/Source/Editor/LandscapeEditor/LandscapeEditor.Build.cs
--- a/Engine/Source/Editor/LandscapeEditor/LandscapeEditor.Build.cs
+++ b/Engine/Source/Editor/LandscapeEditor/LandscapeEditor.Build.cs
@@ -13,6 +13,7 @@ public class LandscapeEditor : ModuleRules
 				"CoreUObject",
 				"Json",
 				"ApplicationCore",
+				"Documentation",
 				"Slate",
 				"SlateCore",
 				"Engine",
diff --git a/Engine/Source/Editor/LandscapeEditor/Private/LandscapeEdMode.cpp b/Engine/Source/Editor/LandscapeEditor/Private/LandscapeEdMode.cpp
--- a/Engine/Source/Editor/LandscapeEditor/Private/LandscapeEdMode.cpp
+++ b/Engine/Source/Editor/LandscapeEditor/Private/LandscapeEdMode.cpp
@@ -121,6 +121,7 @@ FEdModeLandscape::FEdModeLandscape()
 	, bUpdatingLandscapeInfo(false)
 	, bNeedsUpdateTargetLayerList(false)
 	, bHasMapChanged(false)
+	, InputProcessor(MakeShared<FInputProcessor>())
 {
 	using namespace UE::Landscape::Editor;
 	
@@ -407,6 +408,16 @@ void FEdModeLandscape::SetInspectedObjects(const TArray<TWeakObjectPtr<UObject>>
 	RefreshInspectedObjectsDetailPanel();
 }
 
+bool FEdModeLandscape::IsGlobalModifierPressed() const
+{
+	return InputProcessor->bIsShiftDown;
+}
+
+bool FEdModeLandscape::IsGlobalAlternateModifierPressed() const
+{
+	return InputProcessor->bIsCtrlDown;
+}
+
 /** FEdMode: Called when the mode is entered */
 void FEdModeLandscape::Enter()
 {
@@ -644,6 +655,8 @@ void FEdModeLandscape::Enter()
 
 	// Make sure HasMapChanged is reset when entering edit mode
 	bHasMapChanged = false;
+
+	InputProcessor->StartTrackingInput();
 }
 
 /** FEdMode: Called when the mode is exited */
@@ -651,6 +664,8 @@ void FEdModeLandscape::Exit()
 {
 	using namespace UE::Landscape::Editor::Private;
 
+	InputProcessor->StopTrackingInput();
+
 	// Reset mouse tracking info : 
 	GIsGizmoDragging = false;
 	GIsAdjustingBrush  = false;
@@ -880,6 +895,24 @@ void FEdModeLandscape::Tick(FEditorViewportClient* ViewportClient, float DeltaTi
 	}
 }
 
+ELandscapeTargetLayerActionType FEdModeLandscape::GetEstimatedTargetLayerActionType(const TSharedRef<FLandscapeTargetListInfo> Target)
+{
+	const bool bIsGlobalModifierPressed = IsGlobalModifierPressed();
+	const bool bIsGlobalAlternateModifierPressed = IsGlobalAlternateModifierPressed();
+
+	// Bail out if one of the LandscapeEditorControlType's "Hold/Don't Hold Ctrl" condition is not met, so that the tools don't mistakenly let the user know that they will do anything :
+	const ELandscapeFoliageEditorControlType LandscapeEditorControlType = GetDefault<ULevelEditorViewportSettings>()->LandscapeEditorControlType;
+	if ((LandscapeEditorControlType == ELandscapeFoliageEditorControlType::RequireCtrl && !bIsGlobalAlternateModifierPressed)
+		|| (LandscapeEditorControlType == ELandscapeFoliageEditorControlType::RequireNoCtrl && bIsGlobalAlternateModifierPressed))
+	{
+		return ELandscapeTargetLayerActionType::None;
+	}
+
+	return CurrentTool 
+		? CurrentTool->GetTargetLayerActionType(Target->LayerInfoObj.Get(), bIsGlobalModifierPressed, bIsGlobalAlternateModifierPressed)
+		: ELandscapeTargetLayerActionType::None;
+}
+
 /** FEdMode: Called when the mouse enters the viewport area */
 bool FEdModeLandscape::MouseEnter(FEditorViewportClient* InViewportClient, FViewport* InViewport, int32 MouseX, int32 MouseY)
 {
@@ -1976,7 +2009,7 @@ bool FEdModeLandscape::InputKey(FEditorViewportClient* ViewportClient, FViewport
 					 (LandscapeEditorControlType == ELandscapeFoliageEditorControlType::RequireCtrl   && IsCtrlDown(Viewport)) ||
 					 (LandscapeEditorControlType == ELandscapeFoliageEditorControlType::RequireNoCtrl && !IsCtrlDown(Viewport)))))
 			{
-				if (CurrentTool && (CurrentTool->GetSupportedTargetTypes() == ELandscapeToolTargetTypeFlags::None || CurrentToolTarget.TargetType != ELandscapeToolTargetType::Invalid))
+				if (CurrentTool && (CurrentTool->GetSupportedTargetType() == ELandscapeToolTargetType::Invalid || CurrentToolTarget.TargetType != ELandscapeToolTargetType::Invalid))
 				{
 					FVector HitLocation;
 					if (LandscapeMouseTrace(ViewportClient, HitLocation))
@@ -2138,7 +2171,7 @@ void FEdModeLandscape::SetCurrentTool(FName ToolName, FName TargetLayerName)
 		FLandscapeTool* Tool = LandscapeTools[i].Get();
 		if (ToolName == Tool->GetToolName())
 		{
-			if (EnumHasAnyFlags(Tool->GetSupportedTargetTypes(), CurrentToolMode->SupportedTargetTypes))
+			if (EnumHasAnyFlags(UE::Landscape::GetLandscapeToolTargetTypeAsFlags(Tool->GetSupportedTargetType()), CurrentToolMode->SupportedTargetTypes))
 			{
 				ToolIndex = i;
 				break;
@@ -2209,20 +2242,18 @@ void FEdModeLandscape::SetCurrentTool(int32 ToolIndex, FName TargetLayerName)
 	CurrentTool = NewTool;
 
 	// Set target type appropriate for tool
-	if (CurrentTool->GetSupportedTargetTypes() == ELandscapeToolTargetTypeFlags::None)
+	if (CurrentTool->GetSupportedTargetType() == ELandscapeToolTargetType::Invalid)
 	{
 		CurrentToolTarget.TargetType = ELandscapeToolTargetType::Invalid;
 		SetCurrentTargetLayer(NAME_None, nullptr);
 	}
 	else
 	{
-		const ELandscapeToolTargetTypeFlags TargetTypeMask = CurrentTool->GetSupportedTargetTypes();
-		checkSlow(TargetTypeMask != ELandscapeToolTargetTypeFlags::None);
-		// GetLandscapeToolTargetTypeAsFlags does not support ELandscapeToolTargetType::Invalid so explicitly check this first 
-		if (CurrentToolTarget.TargetType == ELandscapeToolTargetType::Invalid || !EnumHasAnyFlags(UE::Landscape::GetLandscapeToolTargetTypeAsFlags(CurrentToolTarget.TargetType), TargetTypeMask))
+		if (CurrentTool->GetSupportedTargetType() != CurrentToolTarget.TargetType)
 		{
-			auto filter = [TargetTypeMask, TargetLayerName](const TSharedRef<FLandscapeTargetListInfo>& Target) { return (EnumHasAnyFlags(UE::Landscape::GetLandscapeToolTargetTypeAsFlags(Target->TargetType), TargetTypeMask)) && (TargetLayerName == NAME_None || TargetLayerName == Target->LayerName); };
-			const TSharedRef<FLandscapeTargetListInfo>* Target = LandscapeTargetList.FindByPredicate(filter);
+			const ELandscapeToolTargetType TargetType = CurrentTool->GetSupportedTargetType();
+			auto Filter = [TargetType, TargetLayerName](const TSharedRef<FLandscapeTargetListInfo>& Target) { return (Target->TargetType == TargetType) && (TargetLayerName == NAME_None || TargetLayerName == Target->LayerName); };
+			const TSharedRef<FLandscapeTargetListInfo>* Target = LandscapeTargetList.FindByPredicate(Filter);
 			if (Target != nullptr)
 			{
 				check(CurrentToolTarget.LandscapeInfo == (*Target)->LandscapeInfo);
@@ -4373,4 +4404,30 @@ FName FLandscapeTargetListInfo::GetLayerName() const
 	return LayerInfoObj.IsValid() ? LayerInfoObj->GetLayerName() : LayerName;
 }
 
-#undef LOCTEXT_NAMESPACE
+void FEdModeLandscape::FInputProcessor::StartTrackingInput()
+{
+	FSlateApplication::Get().RegisterInputPreProcessor(AsShared());
+}
+
+void FEdModeLandscape::FInputProcessor::StopTrackingInput()
+{
+	FSlateApplication::Get().UnregisterInputPreProcessor(AsShared());
+}
+
+bool FEdModeLandscape::FInputProcessor::HandleKeyDownEvent(FSlateApplication& SlateApp, const FKeyEvent& InKeyEvent)
+{
+	bIsShiftDown = SlateApp.GetModifierKeys().IsShiftDown();
+	bIsCtrlDown = SlateApp.GetModifierKeys().IsControlDown();
+	// return false, as we don't want to disturb the input handling, we just want to track whether the modifier keys are pressed at the application level
+	return false;
+}
+
+bool FEdModeLandscape::FInputProcessor::HandleKeyUpEvent(FSlateApplication& SlateApp, const FKeyEvent& InKeyEvent)
+{
+	bIsShiftDown = SlateApp.GetModifierKeys().IsShiftDown();
+	bIsCtrlDown = SlateApp.GetModifierKeys().IsControlDown();
+	// return false, as we don't want to disturb the input handling, we just want to track whether the modifier keys are pressed at the application level
+	return false;
+}
+
+#undef LOCTEXT_NAMESPACE
\ No newline at end of file
diff --git a/Engine/Source/Editor/LandscapeEditor/Private/LandscapeEdMode.h b/Engine/Source/Editor/LandscapeEditor/Private/LandscapeEdMode.h
--- a/Engine/Source/Editor/LandscapeEditor/Private/LandscapeEdMode.h
+++ b/Engine/Source/Editor/LandscapeEditor/Private/LandscapeEdMode.h
@@ -16,6 +16,7 @@
 #include "LandscapeEditTypes.h"
 #include "Containers/Set.h"
 #include "LandscapeImportHelper.h"
+#include "Framework/Application/IInputProcessor.h"
 
 class ALandscape;
 class FCanvas;
@@ -53,12 +54,12 @@ template<typename TargetType> class FLandscapeToolCopyPaste;
 
 struct FLandscapeToolMode
 {
-	const FName				ToolModeName;
-	ELandscapeToolTargetTypeFlags	SupportedTargetTypes;
+	const FName ToolModeName;
+	ELandscapeToolTargetTypeFlags SupportedTargetTypes;
 
-	TArray<FName>			ValidTools;
-	FName					CurrentToolName;
-	FName					CurrentTargetLayerName;
+	TArray<FName> ValidTools;
+	FName CurrentToolName;
+	FName CurrentTargetLayerName;
 
 	FLandscapeToolMode(FName InToolModeName, ELandscapeToolTargetTypeFlags InSupportedTargetTypes)
 		: ToolModeName(InToolModeName)
@@ -260,7 +261,9 @@ enum class ELandscapeEditingState : uint8
 /**
  * Landscape editor mode
  */
-class FEdModeLandscape : public FEdMode, public ILandscapeEdModeInterface
+class FEdModeLandscape 
+	: public FEdMode
+	, public ILandscapeEdModeInterface
 {
 public:
 
@@ -361,6 +364,7 @@ public:
 	virtual ELandscapeToolTargetType GetLandscapeToolTargetType() const override;
 	virtual const ULandscapeEditLayerBase* GetLandscapeSelectedLayer() const override;
 	virtual ULandscapeLayerInfoObject* GetSelectedLandscapeLayerInfo() const override;
+
 	// Deprecated
 	virtual void OnCanHaveLayersContentChanged() final {};
 
@@ -407,7 +411,7 @@ public:
 	/** FEdMode: Allow us to disable mouse delta tracking during painting */
 	virtual bool DisallowMouseDeltaTracking() const override;
 
-	/** FEdMode: Called once per frame */
+	/** FEdMode: Called once per frame per viewport */
 	virtual void Tick(FEditorViewportClient* ViewportClient, float DeltaTime) override;
 
 	/** FEdMode: Called when clicking on a hit proxy */
@@ -637,6 +641,14 @@ public:
 	}
 	void SetInspectedObjects(const TArray<TWeakObjectPtr<UObject>>& InObjects);
 
+	bool IsGlobalModifierPressed() const;
+	bool IsGlobalAlternateModifierPressed() const;
+
+	/**	See GetTargetLayerActionType, except it's a global function meant to be used in the UX, where there's no viewport, and the action type for a given layer depends on 
+	  the pressed modifier buttons (only available per-viewport), so this will give an "estimate" of the result, based on the fact that the modifier button was pressed on 
+	  at least one of the active viewports. Not pretty, I know... */
+	ELandscapeTargetLayerActionType GetEstimatedTargetLayerActionType(const TSharedRef<FLandscapeTargetListInfo> Target);
+
 private:
 	TArray<TSharedRef<FLandscapeTargetListInfo>> LandscapeTargetList;
 	TArray<FLandscapeListInfo> LandscapeList;
@@ -679,4 +691,25 @@ private:
 	// Stores last active tool info. Used to toggle between current and last editor modes (Eg. Toggle Splines Layer)
 	FName LastActiveToolModeName = FName("ToolMode_Sculpt");
 	FName LastActiveToolName = FName("Sculpt");
+
+	// Implement a IInputProcessor to be able to track the shift/ctrl modifier key's states globally (i.e. not per-viewport), in order to access it in the editor mode UX code
+	struct FInputProcessor : public IInputProcessor, public TSharedFromThis<FInputProcessor>
+	{
+		void StartTrackingInput();
+		void StopTrackingInput();
+
+		//~ Begin IInputProcessor interface
+		virtual void Tick(const float DeltaTime, FSlateApplication& SlateApp, TSharedRef<ICursor> Cursor) override {};
+		virtual bool HandleKeyDownEvent(FSlateApplication& SlateApp, const FKeyEvent& InKeyEvent) override;
+		virtual bool HandleKeyUpEvent(FSlateApplication& SlateApp, const FKeyEvent& InKeyEvent) override;
+		virtual const TCHAR* GetDebugName() const 
+		{ 
+			return TEXT("FEdModeLandscape::FInputProcessor"); 
+		}
+		//~ End IInputProcessor interface
+
+		bool bIsShiftDown = false;
+		bool bIsCtrlDown = false;
+	};
+	TSharedRef<FInputProcessor> InputProcessor;
 };
\ No newline at end of file
diff --git a/Engine/Source/Editor/LandscapeEditor/Private/LandscapeEdModeBPCustomTools.cpp b/Engine/Source/Editor/LandscapeEditor/Private/LandscapeEdModeBPCustomTools.cpp
--- a/Engine/Source/Editor/LandscapeEditor/Private/LandscapeEdModeBPCustomTools.cpp
+++ b/Engine/Source/Editor/LandscapeEditor/Private/LandscapeEdModeBPCustomTools.cpp
@@ -51,7 +51,7 @@ public:
 	virtual FText GetDisplayMessage() const override { return FText(); };
 
 	virtual void SetEditRenderType() override { GLandscapeEditRenderMode = ELandscapeEditRenderMode::None | (GLandscapeEditRenderMode & ELandscapeEditRenderMode::BitMaskForMask); }
-	virtual bool SupportsMask() override { return false; }
+	virtual bool SupportsMask() const override { return false; }
 
 	virtual void EnterTool() override
 	{
diff --git a/Engine/Source/Editor/LandscapeEditor/Private/LandscapeEdModeBrushes.cpp b/Engine/Source/Editor/LandscapeEditor/Private/LandscapeEdModeBrushes.cpp
--- a/Engine/Source/Editor/LandscapeEditor/Private/LandscapeEdModeBrushes.cpp
+++ b/Engine/Source/Editor/LandscapeEditor/Private/LandscapeEdModeBrushes.cpp
@@ -60,6 +60,7 @@ protected:
 	TObjectPtr<UMaterialInterface> BrushMaterial;
 	TMap<TWeakObjectPtr<ULandscapeComponent>, TWeakObjectPtr<UMaterialInstanceDynamic>, FDefaultSetAllocator, TWeakObjectPtrMapKeyFuncs<TWeakObjectPtr<ULandscapeComponent>, TWeakObjectPtr<UMaterialInstanceDynamic>>> BrushMaterialInstanceMap;
 
+	// TODO [jonathan.bard] : This should move one level up (in FLandscapeBrush), so that FLandscapeBrushComponent can also disable painting on the edges and with the allow list
 	bool bCanPaint;
 
 	virtual float CalculateFalloff(float Distance, float Radius, float Falloff) = 0;
@@ -293,37 +294,15 @@ public:
 		LastMousePosition = FVector2f(LandscapeX, LandscapeY);
 	}
 
-	virtual FLandscapeBrushData ApplyBrush(const TArray<FLandscapeToolInteractorPosition>& InInteractorPositions) override
+	virtual FLandscapeBrush::FOverlapInfo GetOverlapInfo(const TArray<FLandscapeToolInteractorPosition>& InInteractorPositions) const override
 	{
-		// bCanPaint is evaluated at an old mouse position, and also does not take interactor positions into account,
-		// so cannot be relied upon to be correct vis-a-vis overlapping unloaded components.  We check those again below.
-		if (!bCanPaint)
-		{
-			return FLandscapeBrushData();
-		}
-
 		ULandscapeInfo* LandscapeInfo = EdMode->CurrentToolTarget.LandscapeInfo.Get();
 		const float ScaleXY = static_cast<float>(FMath::Abs(LandscapeInfo->DrawScale.X));
 		const float TotalRadius = EdMode->UISettings->GetCurrentToolBrushRadius() / ScaleXY;
-		const float Radius = (1.0f - EdMode->UISettings->GetCurrentToolBrushFalloff()) * TotalRadius;
-		const float Falloff = EdMode->UISettings->GetCurrentToolBrushFalloff() * TotalRadius;
-
-		// Cap number of mouse positions to a sensible number
-		TArray<FLandscapeToolInteractorPosition> InteractorPositions;
-		if (InInteractorPositions.Num() > 10)
-		{
-			for (int32 i = 0; i < 10; ++i)
-			{
-				// Scale so we include the first and last of the input positions
-				InteractorPositions.Add(InInteractorPositions[(i * (InInteractorPositions.Num() - 1)) / 9]);
-			}
-		}
-		else
-		{
-			InteractorPositions = InInteractorPositions;
-		}
 
+		TArray<FLandscapeToolInteractorPosition> InteractorPositions = CapInteractorPositions(InInteractorPositions);
 		FIntRect ExclusiveBounds;
+		TSet<ULandscapeComponent*> OverlappedComponents;
 		for (const FLandscapeToolInteractorPosition& InteractorPosition : InteractorPositions)
 		{
 			FIntRect SpotExclusiveBounds;
@@ -340,6 +319,11 @@ public:
 			{
 				ExclusiveBounds.Union(SpotExclusiveBounds);
 			}
+
+			// Shrink bounds by 1,1 to avoid GetComponentsInRegion picking up extra components on all sides due to the overlap between components
+			TSet<ULandscapeComponent*> SpotOverlappedComponents;
+			LandscapeInfo->GetComponentsInRegion(SpotExclusiveBounds.Min.X + 1, SpotExclusiveBounds.Min.Y + 1, SpotExclusiveBounds.Max.X - 1, SpotExclusiveBounds.Max.Y - 1, SpotOverlappedComponents);
+			OverlappedComponents.Append(SpotOverlappedComponents);
 		}
 
 		// Clamp to landscape bounds.  GetLandscapeExtent provides inclusive bounds.
@@ -347,25 +331,53 @@ public:
 		if (!ensure(LandscapeInfo->GetLandscapeExtent(MinX, MinY, MaxX, MaxY)))
 		{
 			// Landscape has no components somehow
-			return FLandscapeBrushData();
+			return FLandscapeBrush::FOverlapInfo();
 		}
 		// Clip using exclusive bounds so we can detect empty result.
 		ExclusiveBounds.Clip(FIntRect(MinX, MinY, MaxX + 1, MaxY + 1));
 		if (ExclusiveBounds.Area() <= 0)
 		{
 			// Brush is completely outside the landscape extent
-			return FLandscapeBrushData();
+			return FLandscapeBrush::FOverlapInfo();
 		}
 
+		FLandscapeBrush::FOverlapInfo Result;
+		Result.ExclusiveBounds = ExclusiveBounds;
+		Result.OverlappedLoadedComponents = OverlappedComponents;
+
 		// check if region to be modified overlaps with an unloaded component
 		// TODO [chris.tchou] : The exclusive bounds here (InclusiveMax+1) shouldn't be necessary since the code looks like it expects inclusive values.  Without it, edges can be modified by the brush. Possibly a bug in the brush apply code.
-		bool bHasUnloadedComponents = LandscapeInfo->HasUnloadedComponentsInRegion(ExclusiveBounds.Min.X, ExclusiveBounds.Min.Y, ExclusiveBounds.Max.X, ExclusiveBounds.Max.Y);
-		if (bHasUnloadedComponents)
+		Result.bOverlapsUnloadedComponents = LandscapeInfo->HasUnloadedComponentsInRegion(ExclusiveBounds.Min.X, ExclusiveBounds.Min.Y, ExclusiveBounds.Max.X, ExclusiveBounds.Max.Y);
+
+		return Result;
+	}
+
+	virtual FLandscapeBrushData ApplyBrush(const TArray<FLandscapeToolInteractorPosition>& InInteractorPositions) override
+	{
+		// bCanPaint is evaluated at an old mouse position, and also does not take interactor positions into account,
+		// so cannot be relied upon to be correct vis-a-vis overlapping unloaded components.  We check those again below.
+		if (!bCanPaint)
 		{
 			return FLandscapeBrushData();
 		}
 
-		FLandscapeBrushData BrushData(ExclusiveBounds);
+		ULandscapeInfo* LandscapeInfo = EdMode->CurrentToolTarget.LandscapeInfo.Get();
+		const float ScaleXY = static_cast<float>(FMath::Abs(LandscapeInfo->DrawScale.X));
+		const float TotalRadius = EdMode->UISettings->GetCurrentToolBrushRadius() / ScaleXY;
+		const float Radius = (1.0f - EdMode->UISettings->GetCurrentToolBrushFalloff()) * TotalRadius;
+		const float Falloff = EdMode->UISettings->GetCurrentToolBrushFalloff() * TotalRadius;
+
+		FLandscapeBrush::FOverlapInfo OverlapInfo = GetOverlapInfo(InInteractorPositions);
+		// Disable brush if we don't overlap with anything or if it overlaps with some unloaded components (to prevent data from differing on the edges)
+		if (OverlapInfo.ExclusiveBounds.IsEmpty() || OverlapInfo.bOverlapsUnloadedComponents)
+		{
+			return FLandscapeBrushData();
+		}
+
+		// Cap number of mouse positions to a sensible number
+		TArray<FLandscapeToolInteractorPosition> InteractorPositions = CapInteractorPositions(InInteractorPositions);
+
+		FLandscapeBrushData BrushData(OverlapInfo.ExclusiveBounds);
 
 		for (const FLandscapeToolInteractorPosition& InteractorPosition : InteractorPositions)
 		{
@@ -374,7 +386,7 @@ public:
 			SpotExclusiveBounds.Min.Y = FMath::FloorToInt32(InteractorPosition.Position.Y - TotalRadius);
 			SpotExclusiveBounds.Max.X = FMath::CeilToInt32(InteractorPosition.Position.X + TotalRadius) + 1;
 			SpotExclusiveBounds.Max.Y = FMath::CeilToInt32(InteractorPosition.Position.Y + TotalRadius) + 1;
-			SpotExclusiveBounds.Clip(ExclusiveBounds);
+			SpotExclusiveBounds.Clip(OverlapInfo.ExclusiveBounds);
 
 			for (int32 Y = SpotExclusiveBounds.Min.Y; Y < SpotExclusiveBounds.Max.Y; Y++)
 			{
@@ -416,6 +428,28 @@ public:
 
 		return BrushData;
 	}
+
+private:
+	static TArray<FLandscapeToolInteractorPosition> CapInteractorPositions(const TArray<FLandscapeToolInteractorPosition>& InInteractorPositions)
+	{
+		// Cap number of mouse positions to a sensible number
+		static constexpr int32 MaxInteractorPositions = 10;
+		TArray<FLandscapeToolInteractorPosition> OutInteractorPositions;
+		if (InInteractorPositions.Num() > MaxInteractorPositions)
+		{
+			OutInteractorPositions.Reserve(MaxInteractorPositions);
+			for (int32 i = 0; i < MaxInteractorPositions; ++i)
+			{
+				// Scale so we include the first and last of the input positions
+				OutInteractorPositions.Add(InInteractorPositions[(i * (InInteractorPositions.Num() - 1)) / (MaxInteractorPositions - 1)]);
+			}
+		}
+		else
+		{
+			OutInteractorPositions = InInteractorPositions;
+		}
+		return OutInteractorPositions;
+	}
 };
 
 // 
diff --git a/Engine/Source/Editor/LandscapeEditor/Private/LandscapeEdModeComponentTools.cpp b/Engine/Source/Editor/LandscapeEditor/Private/LandscapeEdModeComponentTools.cpp
--- a/Engine/Source/Editor/LandscapeEditor/Private/LandscapeEdModeComponentTools.cpp
+++ b/Engine/Source/Editor/LandscapeEditor/Private/LandscapeEdModeComponentTools.cpp
@@ -155,7 +155,7 @@ public:
 	virtual FText GetDisplayMessage() const override { return NSLOCTEXT("UnrealEd", "LandscapeMode_Selection_Message", "Paint a mask on the Landscape to protect areas from editing."); };
 
 	virtual void SetEditRenderType() override { GLandscapeEditRenderMode = ELandscapeEditRenderMode::SelectComponent | (GLandscapeEditRenderMode & ELandscapeEditRenderMode::BitMaskForMask); }
-	virtual bool SupportsMask() override { return false; }
+	virtual bool SupportsMask() const override { return false; }
 };
 
 //
@@ -168,7 +168,7 @@ class FLandscapeToolStrokeMask : public FLandscapeToolStrokeBase
 public:
 	FLandscapeToolStrokeMask(FEdModeLandscape* InEdMode, FEditorViewportClient* InViewportClient, const FLandscapeToolTarget& InTarget)
 		: Super(InEdMode, InViewportClient, InTarget)
-		, Cache(InTarget)
+		, Cache(InTarget.LandscapeInfo.Get())
 	{
 	}
 
@@ -266,9 +266,9 @@ public:
 	virtual FText GetDisplayMessage() const override { return NSLOCTEXT("UnrealEd", "LandscapeMode_Mask_Message", "Region Selection"); };
 
 	virtual void SetEditRenderType() override { GLandscapeEditRenderMode = ELandscapeEditRenderMode::SelectRegion | (GLandscapeEditRenderMode & ELandscapeEditRenderMode::BitMaskForMask); }
-	virtual bool SupportsMask() override { return true; }
+	virtual bool SupportsMask() const override { return true; }
 
-	virtual ELandscapeToolType GetToolType() override { return ELandscapeToolType::Mask; }
+	virtual ELandscapeToolType GetToolType() const override { return ELandscapeToolType::Mask; }
 };
 
 //
@@ -281,7 +281,7 @@ class FLandscapeToolStrokeVisibility : public FLandscapeToolStrokeBase
 public:
 	FLandscapeToolStrokeVisibility(FEdModeLandscape* InEdMode, FEditorViewportClient* InViewportClient, const FLandscapeToolTarget& InTarget)
 		: Super(InEdMode, InViewportClient, InTarget)
-		, Cache(InTarget)
+		, Cache(InTarget.LandscapeInfo.Get())
 	{
 	}
 
@@ -303,7 +303,7 @@ public:
 			BrushInfo.GetInclusiveBounds(X1, Y1, X2, Y2);
 
 			// Invert when holding Shift
-			bool bInvert = InteractorPositions[InteractorPositions.Num() - 1].bModifierPressed;
+			bool bInvert = InteractorPositions.Last().bModifierPressed;
 
 			// Tablet pressure
 			float Pressure = ViewportClient->Viewport->IsPenActive() ? ViewportClient->Viewport->GetTabletPressure() : 1.0f;
@@ -361,11 +361,22 @@ public:
 	virtual FText GetDisplayMessage() const override { return NSLOCTEXT("UnrealEd", "LandscapeMode_Visibility_Message", "This tool will allow you to mask out the visibility and collision of areas of your Landscape when used in conjunction with the Landscape Hole Material."); };
 
 	virtual void SetEditRenderType() override { GLandscapeEditRenderMode = ELandscapeEditRenderMode::None | (GLandscapeEditRenderMode & ELandscapeEditRenderMode::BitMaskForMask); }
-	virtual bool SupportsMask() override { return false; }
+	virtual bool SupportsMask() const override { return false; }
 
-	virtual ELandscapeToolTargetTypeFlags GetSupportedTargetTypes() override
+	virtual ELandscapeToolTargetType GetSupportedTargetType() const override
 	{
-		return ELandscapeToolTargetTypeFlags::Visibility;
+		return ELandscapeToolTargetType::Visibility;
+	}
+
+	virtual ELandscapeTargetLayerActionType GetTargetLayerActionType(const ULandscapeLayerInfoObject* InLayerInfo, bool bInIsModifierPressed, bool bInIsAlternateModifierPressed) const override
+	{
+		if (InLayerInfo != ALandscapeProxy::VisibilityLayer)
+		{
+			return ELandscapeTargetLayerActionType::None;
+		}
+
+		const bool bInvert = bInIsModifierPressed;
+		return bInvert ? ELandscapeTargetLayerActionType::Inverted : ELandscapeTargetLayerActionType::Normal;
 	}
 };
 
@@ -529,7 +540,7 @@ public:
 	virtual FText GetDisplayMessage() const override { return NSLOCTEXT("UnrealEd", "LandscapeMode_MoveToLevel_Message", "Move the selected components, via using the Selection tool, to the current streaming level.  This makes it possible to move sections of a Landscape into a streaming level so that they will be streamed in and out with that level, optimizing the performance of the Landscape."); };
 
 	virtual void SetEditRenderType() override { GLandscapeEditRenderMode = ELandscapeEditRenderMode::SelectComponent | (GLandscapeEditRenderMode & ELandscapeEditRenderMode::BitMaskForMask); }
-	virtual bool SupportsMask() override { return false; }
+	virtual bool SupportsMask() const override { return false; }
 };
 
 //
@@ -723,10 +734,10 @@ public:
 	virtual FText GetDisplayMessage() const override { return NSLOCTEXT("UnrealEd", "LandscapeMode_AddComponent_Message", "Create new components for the current Landscape, one at a time.  The cursor shows a green wireframe where new components can be added."); };
 
 	virtual void SetEditRenderType() override { GLandscapeEditRenderMode = ELandscapeEditRenderMode::None | (GLandscapeEditRenderMode & ELandscapeEditRenderMode::BitMaskForMask); }
-	virtual bool SupportsMask() override { return false; }
+	virtual bool SupportsMask() const override { return false; }
 
 	// Sphere traces can result in components being added at a good distance from any neighboring components, because it intersects against a virtual plane
-	virtual bool UseSphereTrace() override { return false; }
+	virtual bool UseSphereTrace() const override { return false; }
 
 	virtual bool CanToolBeActivated() const override
 	{ 
@@ -942,7 +953,7 @@ public:
 	virtual FText GetDisplayMessage() const override { return NSLOCTEXT("UnrealEd", "LandscapeMode_DeleteComponent_Message", "Delete selected components . If no components are currently selected, deletes the component highlighted under the mouse cursor. "); };
 
 	virtual void SetEditRenderType() override { GLandscapeEditRenderMode = ELandscapeEditRenderMode::SelectComponent | (GLandscapeEditRenderMode & ELandscapeEditRenderMode::BitMaskForMask); }
-	virtual bool SupportsMask() override { return false; }
+	virtual bool SupportsMask() const override { return false; }
 };
 
 //
@@ -953,12 +964,16 @@ class FLandscapeToolStrokeCopy : public FLandscapeToolStrokeBase
 {
 	using Super = FLandscapeToolStrokeBase;
 
+protected:
+	using CacheType = ToolTarget::CacheClass;
+	using ValueType = CacheType::DataType;
+
 public:
 	FLandscapeToolStrokeCopy(FEdModeLandscape* InEdMode, FEditorViewportClient* InViewportClient, const FLandscapeToolTarget& InTarget)
 		: Super(InEdMode, InViewportClient, InTarget)
-		, Cache(InTarget)
-		, HeightCache(InTarget)
-		, WeightCache(InTarget)
+		, Cache(InTarget.LandscapeInfo.Get(), nullptr) // TODO [jonathan.bard] : this won't work for weightmaps
+		, HeightCache(InTarget.LandscapeInfo.Get())
+		, WeightCache(InTarget.LandscapeInfo.Get())
 	{
 	}
 
@@ -1001,7 +1016,7 @@ public:
 
 			TArray<uint16> HeightData;
 			TArray<uint8> WeightDatas; // Weight*Layers...
-			TArray<typename ToolTarget::CacheClass::DataType> Data;
+			TArray<ValueType> Data;
 
 			TSet<ULandscapeLayerInfoObject*> LayerInfoSet;
 
@@ -1088,7 +1103,7 @@ public:
 									}
 									else
 									{
-										typename ToolTarget::CacheClass::DataType OriginalValue = Data[index];
+										ValueType OriginalValue = Data[index];
 										if (EdMode->CurrentToolTarget.TargetType == ELandscapeToolTargetType::Heightmap)
 										{
 											GizmoPreData[LocalX + LocalY * 2].Data = Gizmo->GetNormalizedHeight(OriginalValue);
@@ -1252,7 +1267,7 @@ public:
 	}
 
 protected:
-	typename ToolTarget::CacheClass Cache;
+	CacheType Cache;
 	FLandscapeHeightCache HeightCache;
 	FLandscapeFullWeightCache WeightCache;
 };
@@ -1280,9 +1295,9 @@ public:
 		GLandscapeEditRenderMode |= (this->EdMode && this->EdMode->CurrentToolTarget.LandscapeInfo.IsValid() && this->EdMode->CurrentToolTarget.LandscapeInfo->SelectedRegion.Num()) ? ELandscapeEditRenderMode::SelectRegion : ELandscapeEditRenderMode::SelectComponent;
 	}
 
-	virtual ELandscapeToolTargetTypeFlags GetSupportedTargetTypes() override
+	virtual ELandscapeToolTargetType GetSupportedTargetType() const override
 	{
-		return UE::Landscape::GetLandscapeToolTargetTypeAsFlags(ToolTarget::TargetType);
+		return ToolTarget::TargetType;
 	}
 
 	virtual bool BeginTool(FEditorViewportClient* ViewportClient, const FLandscapeToolTarget& InTarget, const FVector& InHitLocation) override
@@ -1316,13 +1331,18 @@ class FLandscapeToolStrokePaste : public FLandscapeToolStrokeBase
 {
 	using Super = FLandscapeToolStrokeBase;
 
+protected:
+	using CacheType = ToolTarget::CacheClass;
+	using ValueType = CacheType::DataType;
+
 public:
 	FLandscapeToolStrokePaste(FEdModeLandscape* InEdMode, FEditorViewportClient* InViewportClient, const FLandscapeToolTarget& InTarget)
 		: Super(InEdMode, InViewportClient, InTarget)
-		, Cache(InTarget)
-		, HeightCache(InTarget)
-		, WeightCache(InTarget)
+		, Cache(InTarget.LandscapeInfo.Get(), nullptr) // this won't work for weightmaps (we need to pass the targeted target layers), but this tool is currently not implemented for that case
+		, HeightCache(InTarget.LandscapeInfo.Get())
+		, WeightCache(InTarget.LandscapeInfo.Get())
 	{
+		static_assert(ToolTarget::TargetType == ELandscapeToolTargetType::Heightmap, "This tool is currently only implemented for FHeightmapToolTarget as of now");
 	}
 
 	void Apply(FEditorViewportClient* ViewportClient, FLandscapeBrush* Brush, const ULandscapeEditorObject* UISettings, const TArray<FLandscapeToolInteractorPosition>& InteractorPositions)
@@ -1389,7 +1409,7 @@ public:
 
 			TArray<uint16> HeightData;
 			TArray<uint8> WeightDatas; // Weight*Layers...
-			TArray<typename ToolTarget::CacheClass::DataType> Data;
+			TArray<ValueType> Data;
 
 			if (bApplyToAll)
 			{
@@ -1537,7 +1557,7 @@ public:
 								}
 								else
 								{
-									Data[index] = static_cast<typename ToolTarget::CacheClass::DataType>(FMath::Lerp(OriginalValue, DestValue, PaintAmount));
+									Data[index] = static_cast<ValueType>(FMath::Lerp(OriginalValue, DestValue, PaintAmount));
 								}
 							}
 						}
@@ -1578,7 +1598,7 @@ public:
 	}
 
 protected:
-	typename ToolTarget::CacheClass Cache;
+	CacheType Cache;
 	FLandscapeHeightCache HeightCache;
 	FLandscapeFullWeightCache WeightCache;
 };
@@ -1611,9 +1631,9 @@ public:
 		bUseGizmoRegion = InbUseGizmoRegion;
 	}
 
-	virtual ELandscapeToolTargetTypeFlags GetSupportedTargetTypes() override
+	virtual ELandscapeToolTargetType GetSupportedTargetType() const override
 	{
-		return UE::Landscape::GetLandscapeToolTargetTypeAsFlags(ToolTarget::TargetType);
+		return ToolTarget::TargetType;
 	}
 
 	virtual ELandscapeLayerUpdateMode GetBeginToolContentUpdateFlag() const override { return ELandscapeLayerUpdateMode::Update_All_Editing; }
@@ -1658,6 +1678,35 @@ public:
 		return FLandscapeToolBase<FLandscapeToolStrokePaste<ToolTarget>>::MouseMove(ViewportClient, Viewport, x, y);
 	}
 
+	virtual ELandscapeTargetLayerActionType GetTargetLayerActionType(const ULandscapeLayerInfoObject* InLayerInfo, bool bInIsModifierPressed, bool bInIsAlternateModifierPressed) const override
+	{
+		ALandscapeGizmoActiveActor* Gizmo = this->EdMode->UISettings->CurrentGizmoActor.Get();
+		// Cache and copy in Gizmo's region...
+		if (this->EdMode->CurrentToolTarget.LandscapeInfo.IsValid() && Gizmo && Gizmo->GetRootComponent())
+		{
+			if (Gizmo->SelectedData.Num() == 0)
+			{
+				const bool bApplyToAll = this->EdMode->UISettings->bApplyToAllTargets;
+				if (this->EdMode->CurrentToolTarget.TargetType == ELandscapeToolTargetType::Heightmap)
+				{
+					return ELandscapeTargetLayerActionType::Normal;
+				}
+				else if (bApplyToAll)
+				{
+					// This tool will only change values on a layer if it is part of the copied data :
+					return Gizmo->LayerInfos.Contains(InLayerInfo) ? ELandscapeTargetLayerActionType::Normal : ELandscapeTargetLayerActionType::None;
+				}
+				else
+				{
+					// This is not handled ATM (see comment in ctor) and will need to be addressed if the copy paste tool becomes available in paint mode 
+					checkNoEntry();
+				}
+			}
+		}
+
+		return ELandscapeTargetLayerActionType::None;
+	}
+
 protected:
 	bool bUseGizmoRegion;
 	FLandscapeBrush* BackupCurrentBrush;
@@ -1829,7 +1878,7 @@ public:
 	virtual FText GetDisplayMessage() const override { return NSLOCTEXT("UnrealEd", "LandscapeMode_NewLandscape_Message", "Create or import a new heightmap.  Assign a material and configure the components.  When you are ready to create your new Landscape, press the Create button in the lower-right corner of this panel. "); };
 
 	virtual void SetEditRenderType() override { GLandscapeEditRenderMode = ELandscapeEditRenderMode::None | (GLandscapeEditRenderMode & ELandscapeEditRenderMode::BitMaskForMask); }
-	virtual bool SupportsMask() override { return false; }
+	virtual bool SupportsMask() const override { return false; }
 
 	virtual void EnterTool() override
 	{
@@ -2287,7 +2336,7 @@ public:
 	virtual FText GetDisplayMessage() const override { return LOCTEXT("LandscapeMode_ResizeLandscape_Message", "Change Landscape Component Size"); };
 
 	virtual void SetEditRenderType() override { GLandscapeEditRenderMode = ELandscapeEditRenderMode::None | (GLandscapeEditRenderMode & ELandscapeEditRenderMode::BitMaskForMask); }
-	virtual bool SupportsMask() override { return false; }
+	virtual bool SupportsMask() const override { return false; }
 
 	virtual void EnterTool() override
 	{
@@ -2354,7 +2403,7 @@ public:
 
 	virtual bool UsesTransformWidget() const { return true; }
 	virtual void SetEditRenderType() override { GLandscapeEditRenderMode = ELandscapeEditRenderMode::SelectComponent | (GLandscapeEditRenderMode & ELandscapeEditRenderMode::BitMaskForMask); }
-	virtual bool SupportsMask() override { return false; }
+	virtual bool SupportsMask() const override { return false; }
 
 	virtual FVector GetWidgetLocation() const override
 	{
diff --git a/Engine/Source/Editor/LandscapeEditor/Private/LandscapeEdModeErosionTools.cpp b/Engine/Source/Editor/LandscapeEditor/Private/LandscapeEdModeErosionTools.cpp
--- a/Engine/Source/Editor/LandscapeEditor/Private/LandscapeEdModeErosionTools.cpp
+++ b/Engine/Source/Editor/LandscapeEditor/Private/LandscapeEdModeErosionTools.cpp
@@ -22,16 +22,16 @@ class FLandscapeToolStrokeErosionBase : public FLandscapeToolStrokeBase
 public:
 	FLandscapeToolStrokeErosionBase(FEdModeLandscape* InEdMode, FEditorViewportClient* InViewportClient, const FLandscapeToolTarget& InTarget)
 		: Super(InEdMode, InViewportClient, InTarget)
-		, HeightCache(InTarget)
-		, LayerHeightDataCache(InTarget, this->HeightCache)
-		, WeightCache(InTarget)
+		, HeightCache(InTarget.LandscapeInfo.Get())
+		, EditLayerHeightDataCache(InTarget.LandscapeInfo.Get(), this->HeightCache)
+		, WeightCache(InTarget.LandscapeInfo.Get())
 		, bWeightApplied(InTarget.TargetType != ELandscapeToolTargetType::Heightmap)
 	{
 	}
 
 protected:
 	FLandscapeHeightCache HeightCache;
-	FLandscapeLayerDataCache<FHeightmapToolTarget> LayerHeightDataCache;
+	FLandscapeEditLayerStackDataCache<FHeightmapToolTarget> EditLayerHeightDataCache;
 	FLandscapeFullWeightCache WeightCache;
 	bool bWeightApplied;
 };
@@ -47,9 +47,15 @@ public:
 	{
 	}
 
-	virtual ELandscapeToolTargetTypeFlags GetSupportedTargetTypes() override
+	virtual ELandscapeToolTargetType GetSupportedTargetType() const override
 	{
-		return ELandscapeToolTargetTypeFlags::Heightmap;
+		return ELandscapeToolTargetType::Heightmap;
+	}
+
+	virtual ELandscapeTargetLayerActionType GetTargetLayerActionType(const ULandscapeLayerInfoObject* InLayerInfo, bool bInIsModifierPressed, bool bInIsAlternateModifierPressed) const override
+	{
+		// Only heightmaps are supported ATM
+		return (this->EdMode->CurrentToolTarget.TargetType == ELandscapeToolTargetType::Heightmap) ? ELandscapeTargetLayerActionType::Normal : ELandscapeTargetLayerActionType::None;
 	}
 };
 
@@ -69,7 +75,7 @@ public:
 
 	virtual void SetEditLayer(const FGuid& EditLayerGUID) override
 	{
-		LayerHeightDataCache.SetCacheEditingLayer(EditLayerGUID);
+		EditLayerHeightDataCache.SetCacheEditingLayer(EditLayerGUID);
 		WeightCache.DataAccess.SetEditLayer(EditLayerGUID);
 	}
 
@@ -110,8 +116,8 @@ public:
 		const int32 LayerNum = this->LandscapeInfo->Layers.Num();
 
 		TArray<uint16> HeightData;
-		LayerHeightDataCache.Initialize(this->LandscapeInfo, bCombinedLayerOperation);
-		LayerHeightDataCache.Read(X1, Y1, X2, Y2, HeightData);
+		EditLayerHeightDataCache.Initialize(bCombinedLayerOperation);
+		EditLayerHeightDataCache.Read(X1, Y1, X2, Y2, HeightData);
 
 		TArray<uint8> WeightDatas; // Weight*Layers...
 		if (UISettings->bErosionUseLayerHardness || bWeightApplied)
@@ -219,7 +225,7 @@ public:
 												// Need to normalize weight...
 												for (int32 LayerIdx = 0; LayerIdx < LayerNum; LayerIdx++)
 												{
-													WeightDatas[Neighbor[Idx] * LayerNum + LayerIdx] = (uint8)(255.0f * NeighborWeight[Idx*LayerNum + LayerIdx] / TotalWeight);
+													WeightDatas[Neighbor[Idx] * LayerNum + LayerIdx] = (TotalWeight > 0.0f) ? (uint8)(255.0f * NeighborWeight[Idx*LayerNum + LayerIdx] / TotalWeight) : 0;
 												}
 											}
 										}
@@ -242,7 +248,7 @@ public:
 									// Need to normalize weight...
 									for (int32 LayerIdx = 0; LayerIdx < LayerNum; LayerIdx++)
 									{
-										WeightDatas[Center*LayerNum + LayerIdx] = (uint8)(255.0f * CenterWeights[LayerIdx] / TotalWeight);
+										WeightDatas[Center*LayerNum + LayerIdx] = (TotalWeight > 0.0f) ? (uint8)(255.0f * CenterWeights[LayerIdx] / TotalWeight) : 0;
 									}
 								}
 
@@ -282,7 +288,7 @@ public:
 			}
 		}
 
-		LayerHeightDataCache.Write(X1, Y1, X2, Y2, HeightData);
+		EditLayerHeightDataCache.Write(X1, Y1, X2, Y2, HeightData);
 		if (bWeightApplied)
 		{
 			WeightCache.SetCachedData(X1, Y1, X2, Y2, WeightDatas, LayerNum, ELandscapeLayerPaintingRestriction::None);
@@ -326,7 +332,7 @@ public:
 
 	virtual void SetEditLayer(const FGuid& EditLayerGUID) override
 	{
-		LayerHeightDataCache.SetCacheEditingLayer(EditLayerGUID);
+		EditLayerHeightDataCache.SetCacheEditingLayer(EditLayerGUID);
 		WeightCache.DataAccess.SetEditLayer(EditLayerGUID);
 	}
 
@@ -368,8 +374,8 @@ public:
 		const float SedimentCapacity = 0.10f * UISettings->SedimentCapacity; //DissolvingRatio; //0.01;
 
 		TArray<uint16> HeightData;
-		LayerHeightDataCache.Initialize(this->LandscapeInfo, bCombinedLayerOperation);
-		LayerHeightDataCache.Read(X1, Y1, X2, Y2, HeightData);
+		EditLayerHeightDataCache.Initialize(bCombinedLayerOperation);
+		EditLayerHeightDataCache.Read(X1, Y1, X2, Y2, HeightData);
 
 		// Apply the brush
 		TArray<uint16> WaterData;
@@ -482,7 +488,7 @@ public:
 									uint32 WaterDiff = (uint32)(WaterTransfer * (float)AltitudeDiff[Idx] / TotalAltitudeDiff);
 									WaterData[Neighbor[Idx]] += static_cast<uint16>(WaterDiff);
 									TotalWaterDiff += WaterDiff;
-									uint32 SedimentDiff = (uint32)(SedimentData[Center] * (float)WaterDiff / WaterData[Center]);
+									uint32 SedimentDiff = (WaterData[Center] > 0) ? (uint32)(SedimentData[Center] * (float)WaterDiff / WaterData[Center]) : 0;
 									SedimentData[Neighbor[Idx]] += static_cast<uint16>(SedimentDiff);
 									TotalSedimentDiff += SedimentDiff;
 								}
@@ -521,7 +527,7 @@ public:
 			LowPassFilter<uint16>(X1, Y1, X2, Y2, BrushInfo, HeightData, UISettings->HErosionDetailScale, 1.0f);
 		}
 
-		LayerHeightDataCache.Write(X1, Y1, X2, Y2, HeightData);
+		EditLayerHeightDataCache.Write(X1, Y1, X2, Y2, HeightData);
 
 		// Dirty any runtime virtual textures that our landscape components write to.
 		LandscapeInfo->DirtyRuntimeVirtualTextureForLandscapeArea(X1, Y1, X2, Y2);
diff --git a/Engine/Source/Editor/LandscapeEditor/Private/LandscapeEdModeMirrorTool.cpp b/Engine/Source/Editor/LandscapeEditor/Private/LandscapeEdModeMirrorTool.cpp
--- a/Engine/Source/Editor/LandscapeEditor/Private/LandscapeEdModeMirrorTool.cpp
+++ b/Engine/Source/Editor/LandscapeEditor/Private/LandscapeEdModeMirrorTool.cpp
@@ -50,11 +50,16 @@ public:
 	virtual FText GetDisplayMessage() const override { return NSLOCTEXT("UnrealEd", "LandscapeTool_Mirror_Message", "Copy one side of a landscape to the other side so that you can easily mirror or rotate the landscape geometry along the X or Y axis."); };
 
 	virtual void SetEditRenderType() override { GLandscapeEditRenderMode = ELandscapeEditRenderMode::None | (GLandscapeEditRenderMode & ELandscapeEditRenderMode::BitMaskForMask); }
-	virtual bool SupportsMask() override { return false; }
+	virtual bool SupportsMask() const override { return false; }
 
-	virtual ELandscapeToolTargetTypeFlags GetSupportedTargetTypes() override
+	virtual ELandscapeToolTargetType GetSupportedTargetType() const override
 	{
-		return ELandscapeToolTargetTypeFlags::Heightmap;
+		return ELandscapeToolTargetType::Heightmap;
+	}
+
+	virtual ELandscapeTargetLayerActionType GetTargetLayerActionType(const ULandscapeLayerInfoObject* InLayerInfo, bool bInIsModifierPressed, bool bInIsAlternateModifierPressed) const override
+	{
+		return ELandscapeTargetLayerActionType::Normal;
 	}
 
 	virtual void EnterTool() override
@@ -116,6 +121,11 @@ public:
 	{
 		// The editor can try to render the tool before the UpdateLandscapeEditorData command runs and the landscape editor realizes that the landscape has been hidden/deleted
 		const ULandscapeInfo* const LandscapeInfo = EdMode->CurrentToolTarget.LandscapeInfo.Get();
+		if (LandscapeInfo == nullptr)
+		{
+			return;
+		}
+
 		const ALandscapeProxy* const LandscapeProxy = LandscapeInfo->GetLandscapeProxy();
 		if (LandscapeProxy)
 		{
@@ -194,6 +204,11 @@ public:
 	{
 		// The editor can try to render the transform widget before the landscape editor ticks and realises that the landscape has been hidden/deleted
 		const ULandscapeInfo* const LandscapeInfo = EdMode->CurrentToolTarget.LandscapeInfo.Get();
+		if (LandscapeInfo == nullptr)
+		{
+			return false;
+		}
+
 		const ALandscapeProxy* const LandscapeProxy = LandscapeInfo->GetLandscapeProxy();
 		if (LandscapeProxy)
 		{
@@ -231,6 +246,11 @@ public:
 	virtual FVector GetWidgetLocation() const override
 	{
 		const ULandscapeInfo* const LandscapeInfo = EdMode->CurrentToolTarget.LandscapeInfo.Get();
+		if (LandscapeInfo == nullptr)
+		{
+			return FVector::ZeroVector;
+		}
+
 		const ALandscapeProxy* const LandscapeProxy = LandscapeInfo->GetLandscapeProxy();
 		if (LandscapeProxy)
 		{
@@ -508,7 +528,17 @@ public:
 		FScopedSetLandscapeEditingLayer Scope(EdMode->GetLandscape(), EdMode->GetCurrentLayerGuid(), [&] { EdMode->RequestLayersContentUpdateForceAll(); });
 
 		const ULandscapeInfo* const LandscapeInfo = EdMode->CurrentToolTarget.LandscapeInfo.Get();
+		if (LandscapeInfo == nullptr)
+		{
+			return;
+		}
+
 		const ALandscapeProxy* const LandscapeProxy = LandscapeInfo->GetLandscapeProxy();
+		if (LandscapeProxy == nullptr)
+		{
+			return;
+		}
+
 		const FTransform LandscapeToWorld = LandscapeProxy->LandscapeActorToWorld();
 		const FVector2D MirrorPoint = EdMode->UISettings->MirrorPoint;
 		int32 BlendWidth = FMath::Clamp(EdMode->UISettings->MirrorSmoothingWidth, 0, 32768);
diff --git a/Engine/Source/Editor/LandscapeEditor/Private/LandscapeEdModePaintTools.cpp b/Engine/Source/Editor/LandscapeEditor/Private/LandscapeEdModePaintTools.cpp
--- a/Engine/Source/Editor/LandscapeEditor/Private/LandscapeEdModePaintTools.cpp
+++ b/Engine/Source/Editor/LandscapeEditor/Private/LandscapeEdModePaintTools.cpp
@@ -1,6 +1,7 @@
 // Copyright Epic Games, Inc. All Rights Reserved.
 
 #include "CoreMinimal.h"
+#include "Algo/Find.h"
 #include "Misc/App.h"
 #include "InputCoreTypes.h"
 #include "Engine/EngineTypes.h"
@@ -56,9 +57,15 @@ public:
 	{
 	}
 
-	virtual ELandscapeToolTargetTypeFlags GetSupportedTargetTypes() override
+	virtual ELandscapeToolTargetType GetSupportedTargetType() const override
 	{
-		return UE::Landscape::GetLandscapeToolTargetTypeAsFlags(TToolTarget::TargetType);
+		return TToolTarget::TargetType;
+	}
+
+	virtual ELandscapeTargetLayerActionType GetTargetLayerActionType(const ULandscapeLayerInfoObject* InLayerInfo, bool bInIsModifierPressed, bool bInIsAlternateModifierPressed) const override
+	{
+		// By default, consider that all tools deriving from FLandscapeToolPaintBase will affect this layer if it's in the list of targeted layers :
+		return this->EdMode->CurrentToolTarget.GetLayerInfos().Contains(InLayerInfo) ? ELandscapeTargetLayerActionType::Normal : ELandscapeTargetLayerActionType::None;
 	}
 };
 
@@ -68,26 +75,45 @@ class FLandscapeToolStrokePaintBase : public FLandscapeToolStrokeBase
 {
 	using Super = FLandscapeToolStrokeBase;
 
+protected:
+	using ValueType = typename ToolTarget::CacheClass::DataType;
+	using CacheType = typename ToolTarget::CacheClass;
+
 public:
 	FLandscapeToolStrokePaintBase(FEdModeLandscape* InEdMode, FEditorViewportClient* InViewportClient, const FLandscapeToolTarget& InTarget)
 		: Super(InEdMode, InViewportClient, InTarget)
-		, Cache(InTarget)
 	{
 		const ULandscapeSettings* Settings = GetDefault<ULandscapeSettings>();
 		PaintStrengthGamma = Settings->PaintStrengthGamma;
 		bDisablePaintingStartupSlowdown = Settings->bDisablePaintingStartupSlowdown;
+
+		// Instantiate a cache object for every affected layer (note : null layer info object pointer is accepted for heightmaps!)
+		TArray<ULandscapeLayerInfoObject*> AffectedLayerInfos = InTarget.GetLayerInfos();
+		PerTargetLayerCache.Reserve(AffectedLayerInfos.Num());
+		for (ULandscapeLayerInfoObject* LayerInfo : AffectedLayerInfos)
+		{ 
+			PerTargetLayerCache.Add(LayerInfo, CacheType(InTarget.LandscapeInfo.Get(), LayerInfo));
+		}
 	}
 
 	float GetStrength(const ULandscapeEditorObject* UISettings) const
 	{
 		// Apply paint strength gamma for paint strokes only : 
-		return (Target.TargetType == ELandscapeToolTargetType::Heightmap)
+		return (ToolTarget::TargetType == ELandscapeToolTargetType::Heightmap) 
 			? UISettings->GetCurrentToolStrength() 
 			: FMath::Pow(UISettings->GetCurrentToolStrength(), PaintStrengthGamma);
 	}
 
+	CacheType& GetCacheForTargetLayer(ULandscapeLayerInfoObject* InLayerInfo)
+	{
+		checkf((ToolTarget::TargetType == ELandscapeToolTargetType::Heightmap) ^ (InLayerInfo != nullptr), TEXT("Heightmap type requires no target layer, weightmap type requires one"));
+		CacheType* Cache = PerTargetLayerCache.Find(InLayerInfo);
+		check(Cache != nullptr);
+		return *Cache;
+	}
+
 protected:
-	typename ToolTarget::CacheClass Cache;
+	TMap<ULandscapeLayerInfoObject*, CacheType> PerTargetLayerCache;
 	float PaintStrengthGamma = 1.0f;
 	bool bDisablePaintingStartupSlowdown = false;
 };
@@ -98,66 +124,172 @@ protected:
 class FLandscapeToolStrokePaint : public FLandscapeToolStrokePaintBase<FWeightmapToolTarget>
 {
 	using Super = FLandscapeToolStrokePaintBase<FWeightmapToolTarget>;
-	using ToolTarget = FWeightmapToolTarget;
-	using ValueType = ToolTarget::CacheClass::DataType;
 
-	TMap<FIntPoint, float> TotalInfluenceMap; // amount of time and weight the brush has spent on each vertex.
-
-	bool bIsAllowListMode;
-	bool bAddToAllowList;
 public:
 
+	/** Struct that holds the state of a given target layer wrt the paint stroke */
+	struct FAffectedTargetLayer
+	{
+	public:
+		FAffectedTargetLayer(ULandscapeLayerInfoObject* InLayer, bool bInInverted)
+			: Layer(InLayer)
+			, bInverted(bInInverted)
+		{}
+
+		ULandscapeLayerInfoObject* GetLayer() const
+		{
+			return Layer;
+		}
+
+		bool IsInverted() const
+		{
+			return bInverted;
+		}
+
+		ELandscapeTargetLayerActionType GetTargetLayerActionType(bool bInvertRequested, bool bInExclusivePaintingRequested) const
+		{
+			// Only let inverted layers be painted when allowed
+			if (!bInExclusivePaintingRequested && bInverted)
+			{
+				return ELandscapeTargetLayerActionType::None;
+			}
+
+			// Don't affect inverted layers when invert-painting : simply erase the painted layers in that case :
+			if (bInvertRequested && bInverted)
+			{
+				return ELandscapeTargetLayerActionType::None;
+			}
+
+			// Apply invert on non-inverted layers when we're not invert-painting already :
+			const bool bLayerInvert = bInvertRequested || bInverted;
+			return bLayerInvert ? ELandscapeTargetLayerActionType::Inverted : ELandscapeTargetLayerActionType::Normal;
+		}
+
+	private:
+		/** The target layer being affected */
+		ULandscapeLayerInfoObject* Layer = nullptr;
+		/** Indicates that the effect of this paint stroke on this layer should be inverted (when exclusive-painting) */
+		bool bInverted = false;
+	};
+
 	FLandscapeToolStrokePaint(FEdModeLandscape* InEdMode, FEditorViewportClient* InViewportClient, const FLandscapeToolTarget& InTarget)
 		: Super(InEdMode, InViewportClient, InTarget)
 		, bIsAllowListMode(EdMode->UISettings->PaintingRestriction == ELandscapeLayerPaintingRestriction::UseComponentAllowList &&
 		                   (InViewportClient->Viewport->KeyState(EKeys::Equals) || InViewportClient->Viewport->KeyState(EKeys::Hyphen)))
 		, bAddToAllowList(bIsAllowListMode && InViewportClient->Viewport->KeyState(EKeys::Equals))
-	{}
+		, AllAffectedTargetLayers(GetAffectedTargetLayersForTarget(InTarget))
+	{
+
+		// Add a cache entry in the alternate cache for each affected target layer that is inverted (and verify that others have an entry in the normal cache already)
+		for (const FAffectedTargetLayer& AffectedTargetLayer : AllAffectedTargetLayers)
+		{
+			if (AffectedTargetLayer.IsInverted())
+			{
+				check(!this->PerTargetLayerCache.Contains(AffectedTargetLayer.GetLayer())); // if it's inverted, it should not be in PerTargetLayerCache
+				AlternatePerTargetLayerCache.Add(AffectedTargetLayer.GetLayer(), CacheType(InTarget.LandscapeInfo.Get(), AffectedTargetLayer.GetLayer()));
+			}
+			else
+			{
+				check(this->PerTargetLayerCache.Contains(AffectedTargetLayer.GetLayer())); // if it's not inverted, it should be in PerTargetLayerCache already
+			}
+		}
+	}
+
+	static TArray<FAffectedTargetLayer> GetAffectedTargetLayersForTarget(const FLandscapeToolTarget& InTarget)
+	{
+		TArray<FAffectedTargetLayer> AffectedTargetLayers;
+		TSet<ULandscapeLayerInfoObject*> TargetedLayerInfos(InTarget.GetLayerInfos());
+		Algo::Transform(TargetedLayerInfos, AffectedTargetLayers, [](ULandscapeLayerInfoObject* InLayer) { return FAffectedTargetLayer(InLayer, /*bInInverted = */false); });
+
+		TArray<FTargetLayerBlendGroup> BlendGroups = InTarget.LandscapeInfo->GetTargetLayerBlendGroups();
+
+		for (ULandscapeLayerInfoObject* TargetLayerInfo : TargetedLayerInfos)
+		{
+			TOptional<FTargetLayerBlendGroup> BlendGroup = InTarget.LandscapeInfo->GetBlendGroupForTargetLayer(TargetLayerInfo);
+			if (BlendGroup.IsSet())
+			{
+				// Now create a cache entry for each of the target layers that are in the blend group and not in the list of affected layers. These will be inversely affected
+				for (TObjectPtr<ULandscapeLayerInfoObject> LayerInfo : BlendGroup->GetTargetLayers())
+				{
+					if (!TargetedLayerInfos.Contains(LayerInfo))
+					{
+						check(!AffectedTargetLayers.ContainsByPredicate([LayerInfo](const FAffectedTargetLayer& InLayer){ return InLayer.GetLayer() == LayerInfo; }));
+						AffectedTargetLayers.Emplace(LayerInfo, /*bInInverted = */true);
+					}
+				}
+			}
+		}
+
+		return AffectedTargetLayers;
+	}
+
+	static bool IsExclusivePaintingRequested(const ULandscapeEditorObject* UISettings, bool bInAlternateModifierPressed)
+	{
+		switch (UISettings->WeightBlendedTargetLayerPaintMode)
+		{ 
+		case ELandscapeWeightBlendPaintMode::ExclusiveRequiresCtrl:
+		{
+			return bInAlternateModifierPressed;
+		}
+		case ELandscapeWeightBlendPaintMode::ExclusiveRequiresNoCtrl:
+		{
+			return !bInAlternateModifierPressed;
+		}
+		case ELandscapeWeightBlendPaintMode::ForceExclusive:
+		{
+			return true;
+		}
+		case ELandscapeWeightBlendPaintMode::ForceNonExclusive:
+		{
+			return false;
+		}
+		default:
+			checkNoEntry();
+		}
+
+		return false;
+	}
 
 	void Apply(FEditorViewportClient* ViewportClient, FLandscapeBrush* Brush, const ULandscapeEditorObject* UISettings, const TArray<FLandscapeToolInteractorPosition>& InteractorPositions)
 	{
 		TRACE_CPUPROFILER_EVENT_SCOPE(FLandscapeToolStrokePaint_Apply);
 
 		// Invert when holding Shift
-		bool bInvert = InteractorPositions.Last().bModifierPressed;
-		UE_LOG(LogLandscapeTools, VeryVerbose, TEXT("bInvert = %d"), bInvert);
+		const bool bInvert = InteractorPositions.Last().bModifierPressed;
+		const bool bExclusivePaintingRequested = IsExclusivePaintingRequested(UISettings, InteractorPositions.Last().bAlternateModifierPressed);
+		UE_LOG(LogLandscapeTools, VeryVerbose, TEXT("bInvert = %d, bExclusivePaintingRequested = %d"), bInvert, bExclusivePaintingRequested);
 
+		// TODO [jonathan.bard] : this should move one level up, as this should also be doable when using FLandscapeBrushComponent : 
 		if (bIsAllowListMode)
 		{
-			// Get list of components to delete from brush
-			// TODO - only retrieve bounds as we don't need the vert data
-			FLandscapeBrushData BrushInfo = Brush->ApplyBrush(InteractorPositions);
-			if (!BrushInfo)
+			// Get list of components overlapped by the brush : 
+			FLandscapeBrush::FOverlapInfo OverlapInfo = Brush->GetOverlapInfo(InteractorPositions);
+			if (!OverlapInfo)
 			{
 				return;
 			}
 
-			int32 X1, Y1, X2, Y2;
-			BrushInfo.GetInclusiveBounds(X1, Y1, X2, Y2);
+			TSet<ULandscapeLayerInfoObject*> SelectedLayers;
+			PerTargetLayerCache.GetKeys(SelectedLayers);
 
-			// Shrink bounds by 1,1 to avoid GetComponentsInRegion picking up extra components on all sides due to the overlap between components
-			TSet<ULandscapeComponent*> SelectedComponents;
-			LandscapeInfo->GetComponentsInRegion(X1 + 1, Y1 + 1, X2 - 1, Y2 - 1, SelectedComponents);
+			TSet<ULandscapeComponent*> SelectedComponents = OverlapInfo.OverlappedLoadedComponents;
 
 			for (ULandscapeComponent* Component : SelectedComponents)
 			{
-				Component->Modify();
-			}
-
-			if (bAddToAllowList)
-			{
-				for (ULandscapeComponent* Component : SelectedComponents)
-				{
-					Component->LayerAllowList.AddUnique(Target.LayerInfo.Get());
-				}
-			}
-			else
-			{
-				FLandscapeEditDataInterface LandscapeEdit(LandscapeInfo);
-				for (ULandscapeComponent* Component : SelectedComponents)
+				for (ULandscapeLayerInfoObject* LayerInfo : SelectedLayers)
 				{
-					Component->LayerAllowList.RemoveSingle(Target.LayerInfo.Get());
-					Component->DeleteLayer(Target.LayerInfo.Get(), LandscapeEdit);
+					if (bAddToAllowList && !Component->LayerAllowList.Contains(LayerInfo))
+					{
+						Component->Modify();
+						Component->LayerAllowList.Add(LayerInfo);
+					}
+					else if (!bAddToAllowList && Component->LayerAllowList.Contains(LayerInfo))
+					{
+						Component->Modify();
+						FLandscapeEditDataInterface LandscapeEdit(LandscapeInfo);
+						Component->LayerAllowList.RemoveSingle(LayerInfo);
+						Component->DeleteLayer(LayerInfo, LandscapeEdit);
+					}
 				}
 			}
 
@@ -180,111 +312,144 @@ public:
 		const bool bUseWeightTargetValue = UISettings->bUseWeightTargetValue;
 		const bool bCacheOriginalData = !bUseWeightTargetValue;
 
-		this->Cache.CacheData(X1, Y1, X2, Y2, bCacheOriginalData);
+		for (const FAffectedTargetLayer& AffectedLayer : AllAffectedTargetLayers)
+		{
+			const ELandscapeTargetLayerActionType LayerActionType = AffectedLayer.GetTargetLayerActionType(bInvert, bExclusivePaintingRequested);
+			if (LayerActionType == ELandscapeTargetLayerActionType::None)
+			{
+				continue;
+			}
 
-		// The data we'll be writing to
-		TArray<ValueType> Data;
-		this->Cache.GetCachedData(X1, Y1, X2, Y2, Data);
+			CacheType* Cache = AffectedLayer.IsInverted() ? AlternatePerTargetLayerCache.Find(AffectedLayer.GetLayer()) : this->PerTargetLayerCache.Find(AffectedLayer.GetLayer());
+			check(Cache != nullptr);
 
-		// The source data we use for editing. 
-		TArray<ValueType>* SourceDataArrayPtr = &Data;
-		TArray<ValueType> OriginalData;
+			TMap<FIntPoint, float>& TotalInfluenceMap = PerTargetTotalInfluenceMap.FindOrAdd(AffectedLayer.GetLayer());
 
-		if (!bUseWeightTargetValue)
-		{
-			// When painting weights (and not using target value mode), we use a source value that tends more
-			// to the current value as we paint over the same region multiple times.
-			// TODO: Make this frame-rate independent
-			this->Cache.GetOriginalData(X1, Y1, X2, Y2, OriginalData);
-			SourceDataArrayPtr = &OriginalData;
+			Cache->CacheData(X1, Y1, X2, Y2, bCacheOriginalData);
 
-			for (int32 Y = Y1; Y <= Y2; Y++)
+			// The data we'll be writing to
+			TArray<ValueType> Data;
+			Cache->GetCachedData(X1, Y1, X2, Y2, Data);
+
+			// The source data we use for editing. 
+			TArray<ValueType>* SourceDataArrayPtr = &Data;
+			TArray<ValueType> OriginalData;
+
+			if (!bUseWeightTargetValue)
 			{
-				ValueType* DataScanline = Data.GetData() + (Y - Y1) * (X2 - X1 + 1) + (0 - X1);
-				ValueType* OriginalDataScanline = OriginalData.GetData() + (Y - Y1) * (X2 - X1 + 1) + (0 - X1);
-				for (int32 X = X1; X <= X2; X++)
+				// When painting weights (and not using target value mode), we use a source value that tends more
+				// to the current value as we paint over the same region multiple times.
+				// TODO: Make this frame-rate independent
+				Cache->GetOriginalData(X1, Y1, X2, Y2, OriginalData);
+				SourceDataArrayPtr = &OriginalData;
+
+				for (int32 Y = Y1; Y <= Y2; Y++)
 				{
-					ValueType& SourceValue = OriginalDataScanline[X];
-					ValueType& CurrentValue = DataScanline[X];
-					if (bDisablePaintingStartupSlowdown)
+					ValueType* DataScanline = Data.GetData() + (Y - Y1) * (X2 - X1 + 1) + (0 - X1);
+					ValueType* OriginalDataScanline = OriginalData.GetData() + (Y - Y1) * (X2 - X1 + 1) + (0 - X1);
+					for (int32 X = X1; X <= X2; X++)
 					{
-						SourceValue = CurrentValue;
-					}
-					else
-					{
-						float VertexInfluence = TotalInfluenceMap.FindRef(FIntPoint(X, Y));
-						SourceValue = FMath::Lerp(SourceValue, CurrentValue, FMath::Min<float>(VertexInfluence * 0.05f, 1.0f));
+						ValueType& SourceValue = OriginalDataScanline[X];
+						ValueType& CurrentValue = DataScanline[X];
+						if (bDisablePaintingStartupSlowdown)
+						{
+							SourceValue = CurrentValue;
+						}
+						else
+						{
+							float VertexInfluence = TotalInfluenceMap.FindRef(FIntPoint(X, Y));
+							SourceValue = FMath::Lerp(SourceValue, CurrentValue, FMath::Min<float>(VertexInfluence * 0.05f, 1.0f));
+						}
 					}
-					
 				}
 			}
-		}
 
-		// Adjust strength based on brush size and drawscale, so strength 1 = one hemisphere
-		const float AdjustedStrength = ToolTarget::StrengthMultiplier(this->LandscapeInfo, UISettings->GetCurrentToolBrushRadius());
-		float PaintStrength = this->GetStrength(UISettings) * Pressure * AdjustedStrength;
+			// Adjust strength based on brush size and drawscale, so strength 1 = one hemisphere
+			const float AdjustedStrength = ToolTarget::StrengthMultiplier(this->LandscapeInfo, UISettings->GetCurrentToolBrushRadius());
+			float PaintStrength = this->GetStrength(UISettings) * Pressure * AdjustedStrength;
 		
-		ValueType DestValue = FWeightmapToolTarget::CacheClass::ClampValue(static_cast<int32>(255.0f * UISettings->WeightTargetValue));
+			ValueType DestValue = FWeightmapToolTarget::CacheClass::ClampValue(static_cast<int32>(255.0f * UISettings->WeightTargetValue));
 
-		// TODO: make paint tool framerate independent like the sculpt tool
-		// const float DeltaTime = FMath::Min<float>(FApp::GetDeltaTime(), 0.1f); // Under 10 fps slow down paint speed
-		// SculptStrength *= DeltaTime * 3.0f; // * 3.0f to partially compensate for impact of DeltaTime on slowing the tools down compared to the old framerate-dependent version
+			// TODO: make paint tool framerate independent like the sculpt tool
+			// const float DeltaTime = FMath::Min<float>(FApp::GetDeltaTime(), 0.1f); // Under 10 fps slow down paint speed
+			// SculptStrength *= DeltaTime * 3.0f; // * 3.0f to partially compensate for impact of DeltaTime on slowing the tools down compared to the old framerate-dependent version
 
-		if (PaintStrength <= 0.0f)
-		{
-			return;
-		}
-
-		if (!bUseWeightTargetValue)
-		{
-			PaintStrength = FMath::Max(PaintStrength, 1.0f);
-		}
+			if (PaintStrength <= 0.0f)
+			{
+				continue;
+			}
 
-		// Apply the brush
-		for (int32 Y = BrushInfo.GetBounds().Min.Y; Y < BrushInfo.GetBounds().Max.Y; Y++)
-		{
-			const float* BrushScanline = BrushInfo.GetDataPtr(FIntPoint(0, Y));
-			ValueType* DataScanline = Data.GetData() + (Y - Y1) * (X2 - X1 + 1) + (0 - X1);
-			ValueType* SourceDataScanline = SourceDataArrayPtr->GetData() + (Y - Y1) * (X2 - X1 + 1) + (0 - X1);
+			// Erasing cancels target value mode, it feels weird otherwise : we cannot "invent" a new target value and we also cannot *not* respect the inverted action contract, 
+			//  otherwise, it feels like the tool is ineffective :
+			const bool bUseWeightTargetValueForLayer = (LayerActionType == ELandscapeTargetLayerActionType::Inverted) ? false : bUseWeightTargetValue;
+			if (!bUseWeightTargetValueForLayer)
+			{
+				PaintStrength = FMath::Max(PaintStrength, 1.0f);
+			}
 
-			for (int32 X = BrushInfo.GetBounds().Min.X; X < BrushInfo.GetBounds().Max.X; X++)
+			// Apply the brush
+			for (int32 Y = BrushInfo.GetBounds().Min.Y; Y < BrushInfo.GetBounds().Max.Y; Y++)
 			{
-				const FIntPoint Key = FIntPoint(X, Y);
-				const float BrushValue = BrushScanline[X];
+				const float* BrushScanline = BrushInfo.GetDataPtr(FIntPoint(0, Y));
+				ValueType* DataScanline = Data.GetData() + (Y - Y1) * (X2 - X1 + 1) + (0 - X1);
+				ValueType* SourceDataScanline = SourceDataArrayPtr->GetData() + (Y - Y1) * (X2 - X1 + 1) + (0 - X1);
 
-				// Update influence map
-				float VertexInfluence = TotalInfluenceMap.FindRef(Key);
-				TotalInfluenceMap.Add(Key, VertexInfluence + BrushValue);
+				for (int32 X = BrushInfo.GetBounds().Min.X; X < BrushInfo.GetBounds().Max.X; X++)
+				{
+					const FIntPoint Key = FIntPoint(X, Y);
+					const float BrushValue = BrushScanline[X];
 
-				float PaintAmount = BrushValue * PaintStrength;
-				ValueType& CurrentValue = DataScanline[X];
-				const ValueType& SourceValue = SourceDataScanline[X];
+					// Update influence map
+					float VertexInfluence = TotalInfluenceMap.FindRef(Key);
+					TotalInfluenceMap.Add(Key, VertexInfluence + BrushValue);
 
-				if (bUseWeightTargetValue)
-				{
-					CurrentValue = FMath::Lerp(CurrentValue, DestValue, PaintAmount / AdjustedStrength);
-				}
-				else
-				{
-					const int32 IntPaintAmount = FMath::RoundToInt(PaintAmount);
-					if (bInvert)
+					float PaintAmount = BrushValue * PaintStrength;
+					ValueType& CurrentValue = DataScanline[X];
+					const ValueType& SourceValue = SourceDataScanline[X];
+
+					if (bUseWeightTargetValueForLayer)
 					{
-						CurrentValue = ToolTarget::CacheClass::ClampValue(FMath::Min<int32>(SourceValue - IntPaintAmount, CurrentValue));
+						CurrentValue = FMath::Lerp(CurrentValue, DestValue, PaintAmount / AdjustedStrength);
 					}
 					else
 					{
-						CurrentValue = ToolTarget::CacheClass::ClampValue(FMath::Max<int32>(SourceValue + IntPaintAmount, CurrentValue));
+						const int32 IntPaintAmount = FMath::RoundToInt(PaintAmount);
+						if (LayerActionType == ELandscapeTargetLayerActionType::Inverted)
+						{
+							CurrentValue = CacheType::ClampValue(FMath::Min<int32>(SourceValue - IntPaintAmount, CurrentValue));
+						}
+						else
+						{
+							CurrentValue = CacheType::ClampValue(FMath::Max<int32>(SourceValue + IntPaintAmount, CurrentValue));
+						}
 					}
 				}
 			}
-		}
 
-		this->Cache.SetCachedData(X1, Y1, X2, Y2, Data, UISettings->PaintingRestriction);
-		this->Cache.Flush();
+			Cache->SetCachedData(X1, Y1, X2, Y2, Data, UISettings->PaintingRestriction);
+			Cache->Flush();
+		}
 		
 		// Dirty any runtime virtual textures that our landscape components write to.
 		this->LandscapeInfo->DirtyRuntimeVirtualTextureForLandscapeArea(X1, Y1, X2, Y2);
 	}
+
+protected:
+	using ToolTarget = FWeightmapToolTarget;
+	using ValueType = Super::ValueType;
+	using CacheType = Super::CacheType;
+
+	// TODO [jonathan.bard] : remove this, it's only used for bDisablePaintingStartupSlowdown == false, which likely no one uses
+	TMap<ULandscapeLayerInfoObject*, TMap<FIntPoint, float>> PerTargetTotalInfluenceMap; // amount of time and weight the brush has spent on each vertex.
+
+	bool bIsAllowListMode;
+	bool bAddToAllowList;
+
+	/** Contains info about all the layers being affected by this paint stroke */
+	TArray<FAffectedTargetLayer> AllAffectedTargetLayers;
+
+	/** Cache for each FAffectedTargetLayer that were not cached by the base class's PerTargetLayerCache (inverted layers) */
+	TMap<ULandscapeLayerInfoObject*, CacheType> AlternatePerTargetLayerCache;
 };
 
 class FLandscapeToolPaint : public FLandscapeToolPaintBase<FWeightmapToolTarget, FLandscapeToolStrokePaint>
@@ -310,20 +475,37 @@ public:
 
 		Super::EnterTool();
 	}
+
+	virtual ELandscapeTargetLayerActionType GetTargetLayerActionType(const ULandscapeLayerInfoObject* InLayerInfo, bool bInIsModifierPressed, bool bInIsAlternateModifierPressed) const override
+	{
+		const bool bInvertRequested = bInIsModifierPressed;
+		const bool bExclusivePaintingRequested = FLandscapeToolStrokePaint::IsExclusivePaintingRequested(this->EdMode->UISettings, bInIsAlternateModifierPressed);
+
+		const TArray<FLandscapeToolStrokePaint::FAffectedTargetLayer> AffectedTargetLayers = FLandscapeToolStrokePaint::GetAffectedTargetLayersForTarget(this->EdMode->CurrentToolTarget);
+		const FLandscapeToolStrokePaint::FAffectedTargetLayer* AffectedLayer = AffectedTargetLayers.FindByPredicate([InLayerInfo](const FLandscapeToolStrokePaint::FAffectedTargetLayer& InAffectedLayer)
+			{
+				return InAffectedLayer.GetLayer() == InLayerInfo;
+			});
+		return AffectedLayer
+			? AffectedLayer->GetTargetLayerActionType(bInvertRequested, bExclusivePaintingRequested)
+			: ELandscapeTargetLayerActionType::None;
+		return ELandscapeTargetLayerActionType::Normal;
+	}
 };
 
 class FLandscapeToolStrokeErase : public FLandscapeToolStrokePaintBase<FHeightmapToolTarget>
 {
 	using Super = FLandscapeToolStrokePaintBase<FHeightmapToolTarget>;
+
+protected:
 	using ToolTarget = FHeightmapToolTarget;
-	using ValueType = typename ToolTarget::CacheClass::DataType;
-	const ValueType FlattenHeight;
+	using ValueType = Super::ValueType;
+	using CacheType = Super::CacheType;
 
 public:
 
 	FLandscapeToolStrokeErase(FEdModeLandscape* InEdMode, FEditorViewportClient* InViewportClient, const FLandscapeToolTarget& InTarget)
 		: Super(InEdMode, InViewportClient, InTarget)
-		, FlattenHeight(LandscapeDataAccess::GetTexHeight(0.f))
 	{}
 
 	void Apply(FEditorViewportClient* ViewportClient, FLandscapeBrush* Brush, const ULandscapeEditorObject* UISettings, const TArray<FLandscapeToolInteractorPosition>& InteractorPositions)
@@ -351,10 +533,14 @@ public:
 		X2 += 1;
 		Y2 += 1;
 		
-		this->Cache.CacheData(X1, Y1, X2, Y2);
+		// Erase is currently heightmap-only so we can assume a single target layer : 
+		CacheType& Cache = GetCacheForTargetLayer(/*InLayerInfo = */nullptr);
+		Cache.CacheData(X1, Y1, X2, Y2);
 
 		TArray<ValueType> Data;
-		this->Cache.GetCachedData(X1, Y1, X2, Y2, Data);
+		Cache.GetCachedData(X1, Y1, X2, Y2, Data);
+
+		static const ValueType FlattenHeight = LandscapeDataAccess::GetTexHeight(0.f);
 
 		// Apply the brush
 		for (int32 Y = BrushInfo.GetBounds().Min.Y; Y < BrushInfo.GetBounds().Max.Y; Y++)
@@ -383,8 +569,8 @@ public:
 			}
 		}
 
-		this->Cache.SetCachedData(X1, Y1, X2, Y2, Data, UISettings->PaintingRestriction);
-		this->Cache.Flush();
+		Cache.SetCachedData(X1, Y1, X2, Y2, Data, UISettings->PaintingRestriction);
+		Cache.Flush();
 
 		// Dirty any runtime virtual textures that our landscape components write to.
 		this->LandscapeInfo->DirtyRuntimeVirtualTextureForLandscapeArea(X1, Y1, X2, Y2);
@@ -395,8 +581,11 @@ public:
 class FLandscapeToolStrokeSculpt : public FLandscapeToolStrokePaintBase<FHeightmapToolTarget>
 {
 	using Super = FLandscapeToolStrokePaintBase<FHeightmapToolTarget>;
+
+protected:
 	using ToolTarget = FHeightmapToolTarget;
-	using ValueType = ToolTarget::CacheClass::DataType;
+	using CacheType = Super::CacheType;
+	using ValueType = Super::ValueType;
 
 public:
 
@@ -406,7 +595,9 @@ public:
 
 	virtual void SetEditLayer(const FGuid& EditLayerGUID) override
 	{
-		this->Cache.DataAccess.SetEditLayer(EditLayerGUID);
+		// Sculpt is currently heightmap-only so we can assume a single target layer : 
+		CacheType& Cache = GetCacheForTargetLayer(/*InLayerInfo = */nullptr);
+		Cache.DataAccess.SetEditLayer(EditLayerGUID);
 	}
 
 	void Apply(FEditorViewportClient* ViewportClient, FLandscapeBrush* Brush, const ULandscapeEditorObject* UISettings, const TArray<FLandscapeToolInteractorPosition>& InteractorPositions)
@@ -414,7 +605,7 @@ public:
 		TRACE_CPUPROFILER_EVENT_SCOPE(FLandscapeToolStrokeSculpt_Apply);
 
 		// Invert when holding Shift
-		bool bInvert = InteractorPositions.Last().bModifierPressed;
+		const bool bInvert = InteractorPositions.Last().bModifierPressed;
 		UE_LOG(LogLandscapeTools, VeryVerbose, TEXT("bInvert = %d"), bInvert);
 
 		// Get list of verts to update
@@ -436,13 +627,15 @@ public:
 		X2 += 1;
 		Y2 += 1;
 
-		this->Cache.CacheData(X1, Y1, X2, Y2);
+		// Sculpt is heightmap-only so we can assume a single target layer : 
+		CacheType& Cache = GetCacheForTargetLayer(/*InLayerInfo = */nullptr);
+		Cache.CacheData(X1, Y1, X2, Y2);
 
 		bool bUseClayBrush = UISettings->bUseClayBrush;
 
 		// The data we'll be writing to
 		TArray<ValueType> Data;
-		this->Cache.GetCachedData(X1, Y1, X2, Y2, Data);
+		Cache.GetCachedData(X1, Y1, X2, Y2, Data);
 
 		// The source data we use for editing. 
 		TArray<ValueType>* SourceDataArrayPtr = &Data;
@@ -611,8 +804,8 @@ public:
 			}
 		}
 
-		this->Cache.SetCachedData(X1, Y1, X2, Y2, Data);
-		this->Cache.Flush();
+		Cache.SetCachedData(X1, Y1, X2, Y2, Data);
+		Cache.Flush();
 
 		// Dirty any runtime virtual textures that our landscape components write to.
 		this->LandscapeInfo->DirtyRuntimeVirtualTextureForLandscapeArea(X1, Y1, X2, Y2);
@@ -632,6 +825,12 @@ public:
 	virtual const TCHAR* GetToolName() const override { return TEXT("Sculpt"); }
 	virtual FText GetDisplayName() const override { return NSLOCTEXT("UnrealEd", "LandscapeMode_Sculpt", "Sculpt"); };
 	virtual FText GetDisplayMessage() const override { return NSLOCTEXT("UnrealEd", "LandscapeMode_Sculpt_Message", "Raise or lower the Landscape using the selected brush shape and falloff."); };
+
+	virtual ELandscapeTargetLayerActionType GetTargetLayerActionType(const ULandscapeLayerInfoObject* InLayerInfo, bool bInIsModifierPressed, bool bInIsAlternateModifierPressed) const
+	{
+		const bool bInvert = bInIsModifierPressed;
+		return bInvert ? ELandscapeTargetLayerActionType::Inverted : ELandscapeTargetLayerActionType::Normal;
+	}
 };
 
 class FLandscapeToolErase : public FLandscapeToolPaintBase<FHeightmapToolTarget, FLandscapeToolStrokeErase>
@@ -647,6 +846,11 @@ public:
 	virtual const TCHAR* GetToolName() const override { return TEXT("Erase"); }
 	virtual FText GetDisplayName() const override { return NSLOCTEXT("UnrealEd", "LandscapeMode_Erase", "Erase"); };
 	virtual FText GetDisplayMessage() const override { return NSLOCTEXT("UnrealEd", "LandscapeMode_Erase_Message", "Erase the Landscape using the selected brush shape and falloff."); };
+
+	virtual ELandscapeTargetLayerActionType GetTargetLayerActionType(const ULandscapeLayerInfoObject* InLayerInfo, bool bInIsModifierPressed, bool bInIsAlternateModifierPressed) const override
+	{
+		return ELandscapeTargetLayerActionType::Normal;
+	}
 };
 
 // 
@@ -657,18 +861,32 @@ template<class ToolTarget>
 class FLandscapeToolStrokeSmooth : public FLandscapeToolStrokePaintBase<ToolTarget>
 {
 	using Super = FLandscapeToolStrokePaintBase<ToolTarget>;
-	using ValueType = typename ToolTarget::CacheClass::DataType;
-	FLandscapeLayerDataCache<ToolTarget> LayerDataCache;
+
+protected:
+	using CacheType = Super::CacheType;
+	using ValueType = Super::ValueType;
+
+	TMap<ULandscapeLayerInfoObject*, FLandscapeEditLayerStackDataCache<ToolTarget>> PerTargetLayerEditLayerDataCache;
+
 public:
 	FLandscapeToolStrokeSmooth(FEdModeLandscape* InEdMode, FEditorViewportClient* InViewportClient, const FLandscapeToolTarget& InTarget)
 		: Super(InEdMode, InViewportClient, InTarget)
-		, LayerDataCache(InTarget, this->Cache)
 	{
+
+		for (auto& It : this->PerTargetLayerCache)
+		{
+			ULandscapeLayerInfoObject* LayerInfo = It.Key;
+			CacheType& Cache = It.Value;
+			PerTargetLayerEditLayerDataCache.Add(LayerInfo, FLandscapeEditLayerStackDataCache<ToolTarget>(InTarget.LandscapeInfo.Get(), Cache));
+		}
 	}
 
 	virtual void SetEditLayer(const FGuid& EditLayerGUID) override
 	{
-		LayerDataCache.SetCacheEditingLayer(EditLayerGUID);
+		for (auto& It : PerTargetLayerEditLayerDataCache)
+		{
+			It.Value.SetCacheEditingLayer(EditLayerGUID);;
+		}
 	}
 
 	void Apply(FEditorViewportClient* ViewportClient, FLandscapeBrush* Brush, const ULandscapeEditorObject* UISettings, const TArray<FLandscapeToolInteractorPosition>& InteractorPositions)
@@ -678,7 +896,7 @@ public:
 		if (!this->LandscapeInfo) return;
 
 		ALandscape* Landscape = this->LandscapeInfo->LandscapeActor.Get();
-		const bool bTargetIsHeightmap = (this->Target.TargetType == ELandscapeToolTargetType::Heightmap);
+		const bool bTargetIsHeightmap = (ToolTarget::TargetType == ELandscapeToolTargetType::Heightmap);
 		const bool bCombinedLayerOperation = bTargetIsHeightmap && UISettings->bCombinedLayersOperation && Landscape;
 
 		// Get list of verts to update
@@ -705,76 +923,83 @@ public:
 			Y2 += 1;
 		}
 
-		TArray<ValueType> Data;
-		LayerDataCache.Initialize(this->LandscapeInfo, bCombinedLayerOperation);
-		LayerDataCache.Read(X1, Y1, X2, Y2, Data);
-		const TArray<ValueType> ReadData { Data };
-		
-		const float ToolStrength = FMath::Clamp<float>(this->GetStrength(UISettings) * Pressure, 0.0f, 1.0f);
-
-		// Apply the brush
-		if (UISettings->bDetailSmooth)
+		for (auto& It : this->PerTargetLayerEditLayerDataCache)
 		{
-			LowPassFilter<ValueType>(X1, Y1, X2, Y2, BrushInfo, Data, UISettings->DetailScale, ToolStrength);
-		}
-		else
-		{
-			const int32 FilterRadius = UISettings->SmoothFilterKernelSize;
+			TArray<ValueType> Data;
+			FLandscapeEditLayerStackDataCache<ToolTarget>& EditLayerDataCache = It.Value;
+			EditLayerDataCache.Initialize(bCombinedLayerOperation);
+			EditLayerDataCache.Read(X1, Y1, X2, Y2, Data);
+			const TArray<ValueType> ReadData { Data };
+		
+			const float ToolStrength = FMath::Clamp<float>(this->GetStrength(UISettings) * Pressure, 0.0f, 1.0f);
 
-			for (int32 Y = BrushInfo.GetBounds().Min.Y; Y < BrushInfo.GetBounds().Max.Y; Y++)
+			// Apply the brush
+			if (UISettings->bDetailSmooth)
 			{
-				const float* BrushScanline = BrushInfo.GetDataPtr(FIntPoint(0, Y));
-				ValueType* DataScanline = Data.GetData() + (Y - Y1) * (X2 - X1 + 1) + (0 - X1);
+				LowPassFilter<ValueType>(X1, Y1, X2, Y2, BrushInfo, Data, UISettings->DetailScale, ToolStrength);
+			}
+			else
+			{
+				const int32 FilterRadius = UISettings->SmoothFilterKernelSize;
 
-				for (int32 X = BrushInfo.GetBounds().Min.X; X < BrushInfo.GetBounds().Max.X; X++)
+				for (int32 Y = BrushInfo.GetBounds().Min.Y; Y < BrushInfo.GetBounds().Max.Y; Y++)
 				{
-					const float BrushValue =  BrushScanline[X] * ToolStrength;
+					const float* BrushScanline = BrushInfo.GetDataPtr(FIntPoint(0, Y));
+					ValueType* DataScanline = Data.GetData() + (Y - Y1) * (X2 - X1 + 1) + (0 - X1);
 
-					if (BrushValue > 0.0f)
+					for (int32 X = BrushInfo.GetBounds().Min.X; X < BrushInfo.GetBounds().Max.X; X++)
 					{
-						// needs to be ~12 bits larger than ToolTarget::CacheClass::DataType (for max FilterRadius (31))
-						// the editor is 64-bit native so just go the whole hog :)
-						int64 FilterValue = 0;
-						int32 FilterSamplingNumber = 0;
-
-						const int32 XRadius = FMath::Min3<int32>(FilterRadius, X - BrushInfo.GetBounds().Min.X, BrushInfo.GetBounds().Max.X - X - 1);
-						const int32 YRadius = FMath::Min3<int32>(FilterRadius, Y - BrushInfo.GetBounds().Min.Y, BrushInfo.GetBounds().Max.Y - Y - 1);
-
-						const int32 SampleX1 = X - XRadius; checkSlow(SampleX1 >= BrushInfo.GetBounds().Min.X);
-						const int32 SampleY1 = Y - YRadius; checkSlow(SampleY1 >= BrushInfo.GetBounds().Min.Y);
-						const int32 SampleX2 = X + XRadius; checkSlow(SampleX2 < BrushInfo.GetBounds().Max.X);
-						const int32 SampleY2 = Y + YRadius; checkSlow(SampleY2 < BrushInfo.GetBounds().Max.Y);
-						for (int32 SampleY = SampleY1; SampleY <= SampleY2; SampleY++)
-						{
-							const float* SampleBrushScanline = BrushInfo.GetDataPtr(FIntPoint(0, SampleY));
-							const float* SampleBrushScanline2 = BrushInfo.GetDataPtr(FIntPoint(0, Y + (Y - SampleY)));
-							const ValueType* SampleDataScanline = ReadData.GetData() + (SampleY - Y1) * (X2 - X1 + 1) + (0 - X1);
+						const float BrushValue =  BrushScanline[X] * ToolStrength;
 
-							for (int32 SampleX = SampleX1; SampleX <= SampleX2; SampleX++)
+						if (BrushValue > 0.0f)
+						{
+							// needs to be ~12 bits larger than ToolTarget::CacheClass::DataType (for max FilterRadius (31))
+							// the editor is 64-bit native so just go the whole hog :)
+							int64 FilterValue = 0;
+							int32 FilterSamplingNumber = 0;
+
+							const int32 XRadius = FMath::Min3<int32>(FilterRadius, X - BrushInfo.GetBounds().Min.X, BrushInfo.GetBounds().Max.X - X - 1);
+							const int32 YRadius = FMath::Min3<int32>(FilterRadius, Y - BrushInfo.GetBounds().Min.Y, BrushInfo.GetBounds().Max.Y - Y - 1);
+
+							const int32 SampleX1 = X - XRadius; checkSlow(SampleX1 >= BrushInfo.GetBounds().Min.X);
+							const int32 SampleY1 = Y - YRadius; checkSlow(SampleY1 >= BrushInfo.GetBounds().Min.Y);
+							const int32 SampleX2 = X + XRadius; checkSlow(SampleX2 < BrushInfo.GetBounds().Max.X);
+							const int32 SampleY2 = Y + YRadius; checkSlow(SampleY2 < BrushInfo.GetBounds().Max.Y);
+							for (int32 SampleY = SampleY1; SampleY <= SampleY2; SampleY++)
 							{
-								// constrain sample to within the brush, symmetrically to prevent flattening bug
-								const float SampleBrushValue =
-									FMath::Min(
-										FMath::Min<float>(SampleBrushScanline[SampleX], SampleBrushScanline[X + (X - SampleX)]),
-										FMath::Min<float>(SampleBrushScanline2[SampleX], SampleBrushScanline2[X + (X - SampleX)])
-									);
-								if (SampleBrushValue > 0.0f)
+								const float* SampleBrushScanline = BrushInfo.GetDataPtr(FIntPoint(0, SampleY));
+								const float* SampleBrushScanline2 = BrushInfo.GetDataPtr(FIntPoint(0, Y + (Y - SampleY)));
+								const ValueType* SampleDataScanline = ReadData.GetData() + (SampleY - Y1) * (X2 - X1 + 1) + (0 - X1);
+
+								for (int32 SampleX = SampleX1; SampleX <= SampleX2; SampleX++)
 								{
-									FilterValue += SampleDataScanline[SampleX];
-									FilterSamplingNumber++;
+									// constrain sample to within the brush, symmetrically to prevent flattening bug
+									const float SampleBrushValue =
+										FMath::Min(
+											FMath::Min<float>(SampleBrushScanline[SampleX], SampleBrushScanline[X + (X - SampleX)]),
+											FMath::Min<float>(SampleBrushScanline2[SampleX], SampleBrushScanline2[X + (X - SampleX)])
+										);
+									if (SampleBrushValue > 0.0f)
+									{
+										FilterValue += SampleDataScanline[SampleX];
+										FilterSamplingNumber++;
+									}
 								}
 							}
-						}
 
-						FilterValue /= FilterSamplingNumber;
+							if (FilterSamplingNumber > 0)
+							{
+								FilterValue /= FilterSamplingNumber;
+							}
 
-						DataScanline[X] = FMath::Lerp(DataScanline[X], static_cast<ValueType>(FilterValue), BrushValue);
+							DataScanline[X] = FMath::Lerp(DataScanline[X], static_cast<ValueType>(FilterValue), BrushValue);
+						}
 					}
 				}
 			}
-		}
 
-		LayerDataCache.Write(X1, Y1, X2, Y2, Data, UISettings->PaintingRestriction);
+			EditLayerDataCache.Write(X1, Y1, X2, Y2, Data, UISettings->PaintingRestriction);
+		}
 
 		// Dirty any runtime virtual textures that our landscape components write to.
 		this->LandscapeInfo->DirtyRuntimeVirtualTextureForLandscapeArea(X1, Y1, X2, Y2);
@@ -795,7 +1020,6 @@ public:
 	virtual const TCHAR* GetToolName() const override { return TEXT("Smooth"); }
 	virtual FText GetDisplayName() const override { return NSLOCTEXT("UnrealEd", "LandscapeMode_Smooth", "Smooth"); };
 	virtual FText GetDisplayMessage() const override { return NSLOCTEXT("UnrealEd", "LandscapeMode_Smooth_Message", "Smooth the Landscape within the brushes influence by averaging the Z position of the Landscape vertices."); };
-
 };
 
 //
@@ -805,32 +1029,47 @@ template<class ToolTarget>
 class FLandscapeToolStrokeFlatten : public FLandscapeToolStrokePaintBase<ToolTarget>
 {
 	using Super = FLandscapeToolStrokePaintBase<ToolTarget>;
-	using ValueType = typename ToolTarget::CacheClass::DataType;
-	ValueType FlattenValue;
 
-	FVector FlattenNormal;
-	float FlattenPlaneDist;
-	bool bInitializedFlattenValue;
-	FLandscapeLayerDataCache<ToolTarget> LayerDataCache;
+protected:
+	using CacheType = Super::CacheType;
+	using ValueType = Super::ValueType;
+
+	TMap<ULandscapeLayerInfoObject*, FLandscapeEditLayerStackDataCache<ToolTarget>> PerTargetLayerEditLayerDataCache;
+
+	struct FFlattenValues
+	{
+		ValueType FlattenValue = static_cast<ValueType>(0);
+		FVector FlattenNormal = FVector::Zero();
+		float FlattenPlaneDist = 0.0f;
+	};
+	TMap<ULandscapeLayerInfoObject*, FFlattenValues> PerTargetLayerFlattenValues;
 
 public:
 	FLandscapeToolStrokeFlatten(FEdModeLandscape* InEdMode, FEditorViewportClient* InViewportClient, const FLandscapeToolTarget& InTarget)
 		: FLandscapeToolStrokePaintBase<ToolTarget>(InEdMode, InViewportClient, InTarget)
-		, bInitializedFlattenValue(false)
-		, LayerDataCache(InTarget, this->Cache)
 	{
-		if (InEdMode->UISettings->bUseFlattenTarget && (InTarget.TargetType == ELandscapeToolTargetType::Heightmap))
+		if (InEdMode->UISettings->bUseFlattenTarget && (ToolTarget::TargetType == ELandscapeToolTargetType::Heightmap))
 		{
 			FTransform LocalToWorld = InTarget.LandscapeInfo->GetLandscapeProxy()->ActorToWorld();
 			float Height = static_cast<float>((InEdMode->UISettings->GetFlattenTarget(/*bInReturnPreviewValueIfActive = */false) - LocalToWorld.GetTranslation().Z) / LocalToWorld.GetScale3D().Z);
-			FlattenValue = static_cast<ValueType>(LandscapeDataAccess::GetTexHeight(Height));
-			bInitializedFlattenValue = true;
+			FFlattenValues& FlattenValues = PerTargetLayerFlattenValues.Add(/*(ULandscapeLayerInfoObject)*/nullptr);
+			FlattenValues.FlattenValue = static_cast<ValueType>(LandscapeDataAccess::GetTexHeight(Height));
+		}
+
+		for (auto& It : this->PerTargetLayerCache)
+		{
+			ULandscapeLayerInfoObject* LayerInfo = It.Key;
+			CacheType& Cache = It.Value;
+			PerTargetLayerEditLayerDataCache.Add(LayerInfo, FLandscapeEditLayerStackDataCache<ToolTarget>(InTarget.LandscapeInfo.Get(), Cache));
 		}
 	}
 
 	virtual void SetEditLayer(const FGuid& EditLayerGUID) override
 	{
-		LayerDataCache.SetCacheEditingLayer(EditLayerGUID);
+		for (auto& It : PerTargetLayerEditLayerDataCache)
+		{
+			It.Value.SetCacheEditingLayer(EditLayerGUID);;
+		}
 	}
 
 	void Apply(FEditorViewportClient* ViewportClient, FLandscapeBrush* Brush, const ULandscapeEditorObject* UISettings, const TArray<FLandscapeToolInteractorPosition>& InteractorPositions)
@@ -838,80 +1077,12 @@ public:
 		if (!this->LandscapeInfo) return;
 
 		ALandscape* Landscape = this->LandscapeInfo->LandscapeActor.Get();
-		const bool bTargetIsHeightmap = (this->Target.TargetType == ELandscapeToolTargetType::Heightmap);
+		const bool bTargetIsHeightmap = (ToolTarget::TargetType == ELandscapeToolTargetType::Heightmap);
 		const bool bCombinedLayerOperation = bTargetIsHeightmap && UISettings->bCombinedLayersOperation && Landscape;
 
 		// Can't use slope if use Flatten Target because no normal is provided
 		bool bUseSlopeFlatten = UISettings->bUseSlopeFlatten && !UISettings->bUseFlattenTarget;
 
-		if (!bInitializedFlattenValue || (UISettings->bPickValuePerApply && bTargetIsHeightmap))
-		{
-			bInitializedFlattenValue = false;
-			float FlattenX = static_cast<float>(InteractorPositions[0].Position.X);
-			float FlattenY = static_cast<float>(InteractorPositions[0].Position.Y);
-			int32 FlattenHeightX = FMath::FloorToInt(FlattenX);
-			int32 FlattenHeightY = FMath::FloorToInt(FlattenY);
-
-			float InterpolatedValue = 0.f;
-			float V00 = 0.f;
-			float V10 = 0.f;
-			float V01 = 0.f;
-			float V11 = 0.f;
-			if (bCombinedLayerOperation)
-			{
-				// create a new accessor and point it read from the final runtime data
-				typename ToolTarget::CacheClass::AccessorClass RuntimeDataAccessor(this->Target);
-				RuntimeDataAccessor.SetEditLayer(FGuid());
-
-				ValueType P00, P10, P01, P11;
-				// GetData will handle interpolation/extrapolation as needed if this samples outside valid landscape coords.
-				FIntPoint SampleCoords(FlattenHeightX, FlattenHeightY);
-				RuntimeDataAccessor.GetData(SampleCoords.X, SampleCoords.Y, SampleCoords.X, SampleCoords.Y, &P00);
-				SampleCoords = { FlattenHeightX + 1, FlattenHeightY };
-				RuntimeDataAccessor.GetData(SampleCoords.X, SampleCoords.Y, SampleCoords.X, SampleCoords.Y, &P10);
-				SampleCoords = { FlattenHeightX, FlattenHeightY + 1 };
-				RuntimeDataAccessor.GetData(SampleCoords.X, SampleCoords.Y, SampleCoords.X, SampleCoords.Y, &P01);
-				SampleCoords = { FlattenHeightX + 1, FlattenHeightY + 1 };
-				RuntimeDataAccessor.GetData(SampleCoords.X, SampleCoords.Y, SampleCoords.X, SampleCoords.Y, &P11);
-				// Release Texture Mips that will be Locked by the next SynchronousUpdateComponentVisibilityForHeight (inside the LayerDataCache.Read call)
-				RuntimeDataAccessor.Flush();
-				V00 = P00;
-				V10 = P10;
-				V01 = P01;
-				V11 = P11;
-				InterpolatedValue = FMath::Lerp(FMath::Lerp(V00, V10, FlattenX - FlattenHeightX), FMath::Lerp(V01, V11, FlattenX - FlattenHeightX), FlattenY - FlattenHeightY);
-			}
-			else
-			{
-				this->Cache.CacheData(FlattenHeightX, FlattenHeightY, FlattenHeightX + 1, FlattenHeightY + 1);
-				InterpolatedValue = this->Cache.GetValue(FlattenX, FlattenY);
-			}
-			FlattenValue = static_cast<ValueType>(InterpolatedValue);
-
-			if (bUseSlopeFlatten && bTargetIsHeightmap)
-			{
-				if (bCombinedLayerOperation)
-				{
-					// Can't rely on cache in this mode
-					FVector Vert00 = FVector(0.0f, 0.0f, V00);
-					FVector Vert01 = FVector(0.0f, 1.0f, V01);
-					FVector Vert10 = FVector(1.0f, 0.0f, V10);
-					FVector Vert11 = FVector(1.0f, 1.0f, V11);
-					FVector FaceNormal1 = ((Vert00 - Vert10) ^ (Vert10 - Vert11)).GetSafeNormal();
-					FVector FaceNormal2 = ((Vert11 - Vert01) ^ (Vert01 - Vert00)).GetSafeNormal();
-					FlattenNormal = (FaceNormal1 + FaceNormal2).GetSafeNormal();
-				}
-				else
-				{
-					FlattenNormal = this->Cache.GetNormal(FlattenHeightX, FlattenHeightY);
-				}
-				FlattenPlaneDist = static_cast<float>(-(FlattenNormal | FVector(FlattenX, FlattenY, InterpolatedValue)));
-			}
-
-			bInitializedFlattenValue = true;
-		}
-
-
 		// Get list of verts to update
 		FLandscapeBrushData BrushInfo = Brush->ApplyBrush(InteractorPositions);
 		if (!BrushInfo)
@@ -935,148 +1106,223 @@ public:
 			Y2 += 1;
 		}
 
-		TArray<ValueType> Data;
-		LayerDataCache.Initialize(this->LandscapeInfo, bCombinedLayerOperation);
-		LayerDataCache.Read(X1, Y1, X2, Y2, Data);
+		for (auto& It : this->PerTargetLayerEditLayerDataCache)
+		{
+			ULandscapeLayerInfoObject* LayerInfo = It.Key;
+			FLandscapeEditLayerStackDataCache<ToolTarget>& EditLayerDataCache = It.Value;
+			CacheType& Cache = this->GetCacheForTargetLayer(LayerInfo);
 
-		const float PaintStrength = this->GetStrength(UISettings);
+			FFlattenValues* FlattenValues = PerTargetLayerFlattenValues.Find(LayerInfo);
+			const bool bInitializedFlattenValue = (FlattenValues != nullptr);
+			if (!bInitializedFlattenValue || (UISettings->bPickValuePerApply && bTargetIsHeightmap))
+			{
+				FlattenValues = &PerTargetLayerFlattenValues.FindOrAdd(LayerInfo);
+
+				float FlattenX = static_cast<float>(InteractorPositions[0].Position.X);
+				float FlattenY = static_cast<float>(InteractorPositions[0].Position.Y);
+				int32 FlattenHeightX = FMath::FloorToInt(FlattenX);
+				int32 FlattenHeightY = FMath::FloorToInt(FlattenY);
+
+				float InterpolatedValue = 0.f;
+				float V00 = 0.f;
+				float V10 = 0.f;
+				float V01 = 0.f;
+				float V11 = 0.f;
+				if (bCombinedLayerOperation)
+				{
+					// create a new accessor and point it read from the final runtime data
+					typename ToolTarget::CacheClass::AccessorClass RuntimeDataAccessor(this->LandscapeInfo, LayerInfo);
+					RuntimeDataAccessor.SetEditLayer(FGuid());
+
+					ValueType P00, P10, P01, P11;
+					// GetData will handle interpolation/extrapolation as needed if this samples outside valid landscape coords.
+					FIntPoint SampleCoords(FlattenHeightX, FlattenHeightY);
+					RuntimeDataAccessor.GetData(SampleCoords.X, SampleCoords.Y, SampleCoords.X, SampleCoords.Y, &P00);
+					SampleCoords = { FlattenHeightX + 1, FlattenHeightY };
+					RuntimeDataAccessor.GetData(SampleCoords.X, SampleCoords.Y, SampleCoords.X, SampleCoords.Y, &P10);
+					SampleCoords = { FlattenHeightX, FlattenHeightY + 1 };
+					RuntimeDataAccessor.GetData(SampleCoords.X, SampleCoords.Y, SampleCoords.X, SampleCoords.Y, &P01);
+					SampleCoords = { FlattenHeightX + 1, FlattenHeightY + 1 };
+					RuntimeDataAccessor.GetData(SampleCoords.X, SampleCoords.Y, SampleCoords.X, SampleCoords.Y, &P11);
+					// Release Texture Mips that will be Locked by the next SynchronousUpdateComponentVisibilityForHeight (inside the EditLayerDataCache.Read call)
+					RuntimeDataAccessor.Flush();
+					V00 = P00;
+					V10 = P10;
+					V01 = P01;
+					V11 = P11;
+					InterpolatedValue = FMath::Lerp(FMath::Lerp(V00, V10, FlattenX - FlattenHeightX), FMath::Lerp(V01, V11, FlattenX - FlattenHeightX), FlattenY - FlattenHeightY);
+				}
+				else
+				{
+					Cache.CacheData(FlattenHeightX, FlattenHeightY, FlattenHeightX + 1, FlattenHeightY + 1);
+					InterpolatedValue = Cache.GetValue(FlattenX, FlattenY);
+				}
+				FlattenValues->FlattenValue = static_cast<ValueType>(InterpolatedValue);
 
-		// Apply the brush
-		for (int32 Y = BrushInfo.GetBounds().Min.Y; Y < BrushInfo.GetBounds().Max.Y; Y++)
-		{
-			const float* BrushScanline = BrushInfo.GetDataPtr(FIntPoint(0, Y));
-			ValueType* DataScanline = Data.GetData() + (Y - Y1) * (X2 - X1 + 1) + (0 - X1);
+				if (bUseSlopeFlatten && bTargetIsHeightmap)
+				{
+					if (bCombinedLayerOperation)
+					{
+						// Can't rely on cache in this mode
+						FVector Vert00 = FVector(0.0f, 0.0f, V00);
+						FVector Vert01 = FVector(0.0f, 1.0f, V01);
+						FVector Vert10 = FVector(1.0f, 0.0f, V10);
+						FVector Vert11 = FVector(1.0f, 1.0f, V11);
+						FVector FaceNormal1 = ((Vert00 - Vert10) ^ (Vert10 - Vert11)).GetSafeNormal();
+						FVector FaceNormal2 = ((Vert11 - Vert01) ^ (Vert01 - Vert00)).GetSafeNormal();
+						FlattenValues->FlattenNormal = (FaceNormal1 + FaceNormal2).GetSafeNormal();
+					}
+					else
+					{
+						FlattenValues->FlattenNormal = Cache.GetNormal(FlattenHeightX, FlattenHeightY);
+					}
+					FlattenValues->FlattenPlaneDist = static_cast<float>(-(FlattenValues->FlattenNormal | FVector(FlattenX, FlattenY, InterpolatedValue)));
+				}
+			}
 
-			for (int32 X = BrushInfo.GetBounds().Min.X; X < BrushInfo.GetBounds().Max.X; X++)
+			TArray<ValueType> Data;
+			EditLayerDataCache.Initialize(bCombinedLayerOperation);
+			EditLayerDataCache.Read(X1, Y1, X2, Y2, Data);
+
+			const float PaintStrength = this->GetStrength(UISettings);
+
+			// Apply the brush
+			for (int32 Y = BrushInfo.GetBounds().Min.Y; Y < BrushInfo.GetBounds().Max.Y; Y++)
 			{
-				const float BrushValue = BrushScanline[X];
+				const float* BrushScanline = BrushInfo.GetDataPtr(FIntPoint(0, Y));
+				ValueType* DataScanline = Data.GetData() + (Y - Y1) * (X2 - X1 + 1) + (0 - X1);
 
-				if (BrushValue > 0.0f)
+				for (int32 X = BrushInfo.GetBounds().Min.X; X < BrushInfo.GetBounds().Max.X; X++)
 				{
-					float Strength = FMath::Clamp<float>(BrushValue * PaintStrength * Pressure, 0.0f, 1.0f);
+					const float BrushValue = BrushScanline[X];
 
-					if (!(bUseSlopeFlatten && bTargetIsHeightmap))
+					if (BrushValue > 0.0f)
 					{
-						int32 Delta = DataScanline[X] - FlattenValue;
-						switch (UISettings->FlattenMode)
+						float Strength = FMath::Clamp<float>(BrushValue * PaintStrength * Pressure, 0.0f, 1.0f);
+
+						if (!(bUseSlopeFlatten && bTargetIsHeightmap))
 						{
-						case ELandscapeToolFlattenMode::Terrace:
-							if (bTargetIsHeightmap)
-							{
-								const FTransform& LocalToWorld = this->Target.LandscapeInfo->GetLandscapeProxy()->ActorToWorld();
-								float ScaleZ = static_cast<float>(LocalToWorld.GetScale3D().Z);
-								float TranslateZ = static_cast<float>(LocalToWorld.GetTranslation().Z);
-								float TerraceInterval = FMath::Clamp(UISettings->TerraceInterval, 1.0f, 32768.f);
-								float Smoothness = UISettings->TerraceSmooth;
-								float WorldHeight = LandscapeDataAccess::GetLocalHeight(DataScanline[X]);
-
-								//move into world space
-								WorldHeight = (WorldHeight * ScaleZ) + TranslateZ;
-								float CurrentHeight = WorldHeight;
-
-								//smoothing part
-								float CurrentLevel = WorldHeight / TerraceInterval;
-								Smoothness = 1.0f / FMath::Max(Smoothness, 0.0001f);
-								float CurrentPhase = FMath::Frac(CurrentLevel);
-								float Halfmask = FMath::Clamp(FMath::CeilToFloat(CurrentPhase - 0.5f), 0.0f, 1.0f);
-								CurrentLevel = FMath::FloorToFloat(WorldHeight / TerraceInterval);
-								float SCurve = FMath::Lerp(CurrentPhase, (1.0f - CurrentPhase), Halfmask) * 2.0f;
-								SCurve = FMath::Pow(SCurve, Smoothness) * 0.5f;
-								SCurve = FMath::Lerp(SCurve, 1.0f - SCurve, Halfmask) * TerraceInterval;
-								WorldHeight = (CurrentLevel * TerraceInterval) + SCurve;
-								//end of smoothing part
-
-								float FinalHeight = FMath::Lerp(CurrentHeight, WorldHeight, Strength);
-								FinalHeight = (FinalHeight - TranslateZ) / ScaleZ;
-								DataScanline[X] = static_cast<ValueType>(LandscapeDataAccess::GetTexHeight(FinalHeight));
-							}
-							break;
-						case ELandscapeToolFlattenMode::Interval:
-							if (bTargetIsHeightmap)
+							int32 Delta = DataScanline[X] - FlattenValues->FlattenValue;
+							switch (UISettings->FlattenMode)
 							{
-								const FTransform& LocalToWorld = this->Target.LandscapeInfo->GetLandscapeProxy()->ActorToWorld();
-								float ScaleZ = static_cast<float>(LocalToWorld.GetScale3D().Z);
-								float TranslateZ = static_cast<float>(LocalToWorld.GetTranslation().Z);
-								float TerraceInterval = UISettings->TerraceInterval;
-								float TargetHeight = LandscapeDataAccess::GetLocalHeight(FlattenValue);
-								float CurrentHeight = LandscapeDataAccess::GetLocalHeight(DataScanline[X]);
+							case ELandscapeToolFlattenMode::Terrace:
+								if (bTargetIsHeightmap)
+								{
+									const FTransform& LocalToWorld = this->LandscapeInfo->GetLandscapeProxy()->ActorToWorld();
+									float ScaleZ = static_cast<float>(LocalToWorld.GetScale3D().Z);
+									float TranslateZ = static_cast<float>(LocalToWorld.GetTranslation().Z);
+									float TerraceInterval = FMath::Clamp(UISettings->TerraceInterval, 1.0f, 32768.f);
+									float Smoothness = UISettings->TerraceSmooth;
+									float WorldHeight = LandscapeDataAccess::GetLocalHeight(DataScanline[X]);
+
+									//move into world space
+									WorldHeight = (WorldHeight * ScaleZ) + TranslateZ;
+									float CurrentHeight = WorldHeight;
+
+									//smoothing part
+									float CurrentLevel = WorldHeight / TerraceInterval;
+									Smoothness = 1.0f / FMath::Max(Smoothness, 0.0001f);
+									float CurrentPhase = FMath::Frac(CurrentLevel);
+									float Halfmask = FMath::Clamp(FMath::CeilToFloat(CurrentPhase - 0.5f), 0.0f, 1.0f);
+									CurrentLevel = FMath::FloorToFloat(WorldHeight / TerraceInterval);
+									float SCurve = FMath::Lerp(CurrentPhase, (1.0f - CurrentPhase), Halfmask) * 2.0f;
+									SCurve = FMath::Pow(SCurve, Smoothness) * 0.5f;
+									SCurve = FMath::Lerp(SCurve, 1.0f - SCurve, Halfmask) * TerraceInterval;
+									WorldHeight = (CurrentLevel * TerraceInterval) + SCurve;
+									//end of smoothing part
+
+									float FinalHeight = FMath::Lerp(CurrentHeight, WorldHeight, Strength);
+									FinalHeight = (FinalHeight - TranslateZ) / ScaleZ;
+									DataScanline[X] = static_cast<ValueType>(LandscapeDataAccess::GetTexHeight(FinalHeight));
+								}
+								break;
+							case ELandscapeToolFlattenMode::Interval:
+								if (bTargetIsHeightmap)
+								{
+									const FTransform& LocalToWorld = this->LandscapeInfo->GetLandscapeProxy()->ActorToWorld();
+									float ScaleZ = static_cast<float>(LocalToWorld.GetScale3D().Z);
+									float TranslateZ = static_cast<float>(LocalToWorld.GetTranslation().Z);
+									float TerraceInterval = UISettings->TerraceInterval;
+									float TargetHeight = LandscapeDataAccess::GetLocalHeight(FlattenValues->FlattenValue);
+									float CurrentHeight = LandscapeDataAccess::GetLocalHeight(DataScanline[X]);
 														
-								//move into world space
-								TargetHeight = (TargetHeight * ScaleZ) + TranslateZ;
-								CurrentHeight = (CurrentHeight * ScaleZ) + TranslateZ;
+									//move into world space
+									TargetHeight = (TargetHeight * ScaleZ) + TranslateZ;
+									CurrentHeight = (CurrentHeight * ScaleZ) + TranslateZ;
 
-								TargetHeight = (FMath::RoundToFloat(TargetHeight / TerraceInterval)) * TerraceInterval;
-								TargetHeight = FMath::Lerp(CurrentHeight, TargetHeight, BrushValue);
+									TargetHeight = (FMath::RoundToFloat(TargetHeight / TerraceInterval)) * TerraceInterval;
+									TargetHeight = FMath::Lerp(CurrentHeight, TargetHeight, BrushValue);
 							
-								//back to local space of landscape object
-								TargetHeight = (TargetHeight - TranslateZ) / ScaleZ;
-								DataScanline[X] = static_cast<ValueType>(LandscapeDataAccess::GetTexHeight(TargetHeight));
+									//back to local space of landscape object
+									TargetHeight = (TargetHeight - TranslateZ) / ScaleZ;
+									DataScanline[X] = static_cast<ValueType>(LandscapeDataAccess::GetTexHeight(TargetHeight));
 
+								}
+								break;
+							case ELandscapeToolFlattenMode::Raise:
+								if (Delta < 0)
+								{
+									DataScanline[X] = static_cast<ValueType>(FMath::CeilToInt(FMath::Lerp((float)DataScanline[X], (float)FlattenValues->FlattenValue, Strength)));
+								}
+								break;
+							case ELandscapeToolFlattenMode::Lower:
+								if (Delta > 0)
+								{
+									DataScanline[X] = static_cast<ValueType>(FMath::FloorToInt(FMath::Lerp((float)DataScanline[X], (float)FlattenValues->FlattenValue, Strength)));
+								}
+								break;
+							default:
+							case ELandscapeToolFlattenMode::Both:
+								if (Delta > 0)
+								{
+									DataScanline[X] = static_cast<ValueType>(FMath::FloorToInt(FMath::Lerp((float)DataScanline[X], (float)FlattenValues->FlattenValue, Strength)));
+								}
+								else
+								{
+									DataScanline[X] = static_cast<ValueType>(FMath::CeilToInt(FMath::Lerp((float)DataScanline[X], (float)FlattenValues->FlattenValue, Strength)));
+								}
+								break;
 							}
-							break;
-						case ELandscapeToolFlattenMode::Raise:
-							if (Delta < 0)
-							{
-								DataScanline[X] = static_cast<ValueType>(FMath::CeilToInt(FMath::Lerp((float)DataScanline[X], (float)FlattenValue, Strength)));
-							}
-							break;
-						case ELandscapeToolFlattenMode::Lower:
-							if (Delta > 0)
-							{
-								DataScanline[X] = static_cast<ValueType>(FMath::FloorToInt(FMath::Lerp((float)DataScanline[X], (float)FlattenValue, Strength)));
-							}
-							break;
-						default:
-						case ELandscapeToolFlattenMode::Both:
-							if (Delta > 0)
-							{
-								DataScanline[X] = static_cast<ValueType>(FMath::FloorToInt(FMath::Lerp((float)DataScanline[X], (float)FlattenValue, Strength)));
-							}
-							else
-							{
-								DataScanline[X] = static_cast<ValueType>(FMath::CeilToInt(FMath::Lerp((float)DataScanline[X], (float)FlattenValue, Strength)));
-							}
-							break;
 						}
-					}
-					else
-					{
-						ValueType DestValue = static_cast<ValueType>(-(FlattenNormal.X * X + FlattenNormal.Y * Y + FlattenPlaneDist) / FlattenNormal.Z);
-						//float PlaneDist = FlattenNormal | FVector(X, Y, HeightData(HeightDataIndex)) + FlattenPlaneDist;
-						float PlaneDist = static_cast<float>(DataScanline[X] - DestValue);
-						DestValue = static_cast<ValueType>(DataScanline[X] - PlaneDist * Strength);
-						switch (UISettings->FlattenMode)
+						else
 						{
-						case ELandscapeToolFlattenMode::Raise:
-							if (PlaneDist < 0)
-							{
-								DataScanline[X] = static_cast<ValueType>(FMath::CeilToInt(FMath::Lerp((float)DataScanline[X], (float)DestValue, Strength)));
-							}
-							break;
-						case ELandscapeToolFlattenMode::Lower:
-							if (PlaneDist > 0)
-							{
-								DataScanline[X] = static_cast<ValueType>(FMath::FloorToInt(FMath::Lerp((float)DataScanline[X], (float)DestValue, Strength)));
-							}
-							break;
-						default:
-						case ELandscapeToolFlattenMode::Both:
-							if (PlaneDist > 0)
+							ValueType DestValue = static_cast<ValueType>(-(FlattenValues->FlattenNormal.X * X + FlattenValues->FlattenNormal.Y * Y + FlattenValues->FlattenPlaneDist) / FlattenValues->FlattenNormal.Z);
+							//float PlaneDist = FlattenValues->FlattenNormal | FVector(X, Y, HeightData(HeightDataIndex)) + FlattenValues->FlattenPlaneDist;
+							float PlaneDist = static_cast<float>(DataScanline[X] - DestValue);
+							DestValue = static_cast<ValueType>(DataScanline[X] - PlaneDist * Strength);
+							switch (UISettings->FlattenMode)
 							{
-								DataScanline[X] = static_cast<ValueType>(FMath::FloorToInt(FMath::Lerp((float)DataScanline[X], (float)DestValue, Strength)));
-							}
-							else
-							{
-								DataScanline[X] = static_cast<ValueType>(FMath::CeilToInt(FMath::Lerp((float)DataScanline[X], (float)DestValue, Strength)));
+							case ELandscapeToolFlattenMode::Raise:
+								if (PlaneDist < 0)
+								{
+									DataScanline[X] = static_cast<ValueType>(FMath::CeilToInt(FMath::Lerp((float)DataScanline[X], (float)DestValue, Strength)));
+								}
+								break;
+							case ELandscapeToolFlattenMode::Lower:
+								if (PlaneDist > 0)
+								{
+									DataScanline[X] = static_cast<ValueType>(FMath::FloorToInt(FMath::Lerp((float)DataScanline[X], (float)DestValue, Strength)));
+								}
+								break;
+							default:
+							case ELandscapeToolFlattenMode::Both:
+								if (PlaneDist > 0)
+								{
+									DataScanline[X] = static_cast<ValueType>(FMath::FloorToInt(FMath::Lerp((float)DataScanline[X], (float)DestValue, Strength)));
+								}
+								else
+								{
+									DataScanline[X] = static_cast<ValueType>(FMath::CeilToInt(FMath::Lerp((float)DataScanline[X], (float)DestValue, Strength)));
+								}
+								break;
 							}
-							break;
 						}
 					}
 				}
 			}
-		}
 
-		LayerDataCache.Write(X1, Y1, X2, Y2, Data, UISettings->PaintingRestriction);
+			EditLayerDataCache.Write(X1, Y1, X2, Y2, Data, UISettings->PaintingRestriction);
+		}
 
 		// Dirty any runtime virtual textures that our landscape components write to.
 		this->LandscapeInfo->DirtyRuntimeVirtualTextureForLandscapeArea(X1, Y1, X2, Y2);
@@ -1264,7 +1510,10 @@ template<class ToolTarget>
 class FLandscapeToolStrokeNoise : public FLandscapeToolStrokePaintBase<ToolTarget>
 {
 	using Super = FLandscapeToolStrokePaintBase<ToolTarget>;
-	using ValueType = typename ToolTarget::CacheClass::DataType;
+
+protected:
+	using CacheType = Super::CacheType;
+	using ValueType = Super::ValueType;
 
 public:
 	FLandscapeToolStrokeNoise(FEdModeLandscape* InEdMode, FEditorViewportClient* InViewportClient, const FLandscapeToolTarget& InTarget)
@@ -1301,70 +1550,74 @@ public:
 			Y2 += 1;
 		}
 
-		this->Cache.CacheData(X1, Y1, X2, Y2);
-		TArray<ValueType> Data;
-		this->Cache.GetCachedData(X1, Y1, X2, Y2, Data);
-
-		float BrushSizeAdjust = 1.0f;
-		CA_SUPPRESS(6326);
-		if (ToolTarget::TargetType != ELandscapeToolTargetType::Weightmap && UISettings->GetCurrentToolBrushRadius() < UISettings->MaximumValueRadius)
+		for (auto& It : this->PerTargetLayerCache)
 		{
-			BrushSizeAdjust = UISettings->GetCurrentToolBrushRadius() / UISettings->MaximumValueRadius;
-		}
+			CacheType& Cache = It.Value;
+			Cache.CacheData(X1, Y1, X2, Y2);
+			TArray<ValueType> Data;
+			Cache.GetCachedData(X1, Y1, X2, Y2, Data);
+
+			float BrushSizeAdjust = 1.0f;
+			CA_SUPPRESS(6326);
+			if (ToolTarget::TargetType != ELandscapeToolTargetType::Weightmap && UISettings->GetCurrentToolBrushRadius() < UISettings->MaximumValueRadius)
+			{
+				BrushSizeAdjust = UISettings->GetCurrentToolBrushRadius() / UISettings->MaximumValueRadius;
+			}
 
-		CA_SUPPRESS(6326);
-		bool bUseWeightTargetValue = UISettings->bUseWeightTargetValue && ToolTarget::TargetType == ELandscapeToolTargetType::Weightmap;
+			CA_SUPPRESS(6326);
+			bool bUseWeightTargetValue = UISettings->bUseWeightTargetValue && ToolTarget::TargetType == ELandscapeToolTargetType::Weightmap;
 
-		const float PaintStrength = this->GetStrength(UISettings);
+			const float PaintStrength = this->GetStrength(UISettings);
 
-		// Apply the brush
-		for (int32 Y = BrushInfo.GetBounds().Min.Y; Y < BrushInfo.GetBounds().Max.Y; Y++)
-		{
-			const float* BrushScanline = BrushInfo.GetDataPtr(FIntPoint(0, Y));
-			ValueType* DataScanline = Data.GetData() + (Y - Y1) * (X2 - X1 + 1) + (0 - X1);
-
-			for (int32 X = BrushInfo.GetBounds().Min.X; X < BrushInfo.GetBounds().Max.X; X++)
+			// Apply the brush
+			for (int32 Y = BrushInfo.GetBounds().Min.Y; Y < BrushInfo.GetBounds().Max.Y; Y++)
 			{
-				const float BrushValue = BrushScanline[X];
+				const float* BrushScanline = BrushInfo.GetDataPtr(FIntPoint(0, Y));
+				ValueType* DataScanline = Data.GetData() + (Y - Y1) * (X2 - X1 + 1) + (0 - X1);
 
-				if (BrushValue > 0.0f)
+				for (int32 X = BrushInfo.GetBounds().Min.X; X < BrushInfo.GetBounds().Max.X; X++)
 				{
-					float OriginalValue = DataScanline[X];
-					if (bUseWeightTargetValue)
+					const float BrushValue = BrushScanline[X];
+
+					if (BrushValue > 0.0f)
 					{
-						FNoiseParameter NoiseParam(0, UISettings->NoiseScale, 255.0f / 2.0f);
-						float DestValue = NoiseModeConversion(ELandscapeToolNoiseMode::Add, NoiseParam.NoiseAmount, NoiseParam.Sample(X, Y)) * UISettings->WeightTargetValue;
-						switch (UISettings->NoiseMode)
+						float OriginalValue = DataScanline[X];
+						if (bUseWeightTargetValue)
 						{
-						case ELandscapeToolNoiseMode::Add:
-							if (OriginalValue >= DestValue)
+							FNoiseParameter NoiseParam(0, UISettings->NoiseScale, 255.0f / 2.0f);
+							float DestValue = NoiseModeConversion(ELandscapeToolNoiseMode::Add, NoiseParam.NoiseAmount, NoiseParam.Sample(X, Y)) * UISettings->WeightTargetValue;
+							switch (UISettings->NoiseMode)
 							{
-								continue;
-							}
-							break;
-						case ELandscapeToolNoiseMode::Sub:
-							DestValue += (1.0f - UISettings->WeightTargetValue) * NoiseParam.NoiseAmount;
-							if (OriginalValue <= DestValue)
-							{
-								continue;
+							case ELandscapeToolNoiseMode::Add:
+								if (OriginalValue >= DestValue)
+								{
+									continue;
+								}
+								break;
+							case ELandscapeToolNoiseMode::Sub:
+								DestValue += (1.0f - UISettings->WeightTargetValue) * NoiseParam.NoiseAmount;
+								if (OriginalValue <= DestValue)
+								{
+									continue;
+								}
+								break;
 							}
-							break;
+							DataScanline[X] = ToolTarget::CacheClass::ClampValue(FMath::RoundToInt(FMath::Lerp(OriginalValue, DestValue, BrushValue * UISettings->GetCurrentToolStrength() * Pressure)));
+						}
+						else
+						{
+							float TotalStrength = BrushValue * PaintStrength * Pressure * ToolTarget::StrengthMultiplier(this->LandscapeInfo, UISettings->GetCurrentToolBrushRadius());
+							FNoiseParameter NoiseParam(0, UISettings->NoiseScale, TotalStrength * BrushSizeAdjust);
+							float PaintAmount = NoiseModeConversion(UISettings->NoiseMode, NoiseParam.NoiseAmount, NoiseParam.Sample(X, Y));
+							DataScanline[X] = static_cast<ValueType>(ToolTarget::CacheClass::ClampValue(static_cast<int32>(OriginalValue + PaintAmount)));
 						}
-						DataScanline[X] = ToolTarget::CacheClass::ClampValue(FMath::RoundToInt(FMath::Lerp(OriginalValue, DestValue, BrushValue * UISettings->GetCurrentToolStrength() * Pressure)));
-					}
-					else
-					{
-						float TotalStrength = BrushValue * PaintStrength * Pressure * ToolTarget::StrengthMultiplier(this->LandscapeInfo, UISettings->GetCurrentToolBrushRadius());
-						FNoiseParameter NoiseParam(0, UISettings->NoiseScale, TotalStrength * BrushSizeAdjust);
-						float PaintAmount = NoiseModeConversion(UISettings->NoiseMode, NoiseParam.NoiseAmount, NoiseParam.Sample(X, Y));
-						DataScanline[X] = static_cast<ValueType>(ToolTarget::CacheClass::ClampValue(static_cast<int32>(OriginalValue + PaintAmount)));
 					}
 				}
 			}
-		}
 
-		this->Cache.SetCachedData(X1, Y1, X2, Y2, Data, UISettings->PaintingRestriction);
-		this->Cache.Flush();
+			Cache.SetCachedData(X1, Y1, X2, Y2, Data, UISettings->PaintingRestriction);
+			Cache.Flush();
+		}
 
 		// Dirty any runtime virtual textures that our landscape components write to.
 		this->LandscapeInfo->DirtyRuntimeVirtualTextureForLandscapeArea(X1, Y1, X2, Y2);
diff --git a/Engine/Source/Editor/LandscapeEditor/Private/LandscapeEdModeRampTool.cpp b/Engine/Source/Editor/LandscapeEditor/Private/LandscapeEdModeRampTool.cpp
--- a/Engine/Source/Editor/LandscapeEditor/Private/LandscapeEdModeRampTool.cpp
+++ b/Engine/Source/Editor/LandscapeEditor/Private/LandscapeEdModeRampTool.cpp
@@ -124,11 +124,16 @@ public:
 	virtual FText GetDisplayMessage() const override { return NSLOCTEXT("UnrealEd", "LandscapeMode_Ramp_Message", "Create a ramp between two specified points, adding falloffs according to the settings. Note that this tool cannot use any brushes."); };
 
 	virtual void SetEditRenderType() override { GLandscapeEditRenderMode = ELandscapeEditRenderMode::None | (GLandscapeEditRenderMode & ELandscapeEditRenderMode::BitMaskForMask); }
-	virtual bool SupportsMask() override { return false; }
+	virtual bool SupportsMask() const override { return false; }
 
-	virtual ELandscapeToolTargetTypeFlags GetSupportedTargetTypes() override
+	virtual ELandscapeToolTargetType GetSupportedTargetType() const override
 	{
-		return ELandscapeToolTargetTypeFlags::Heightmap;
+		return ELandscapeToolTargetType::Heightmap;
+	}
+
+	virtual ELandscapeTargetLayerActionType GetTargetLayerActionType(const ULandscapeLayerInfoObject* InLayerInfo, bool bInIsModifierPressed, bool bInIsAlternateModifierPressed) const override
+	{
+		return ELandscapeTargetLayerActionType::Normal;
 	}
 
 	virtual void EnterTool() override
@@ -487,6 +492,12 @@ public:
 			return;
 		}
 
+		if (!CanApplyRamp())
+		{
+			FMessageDialog::Open(EAppMsgType::Ok, LOCTEXT("RampApplyError", "Cannot apply ramp, make sure 2 control points are used."));
+			return;
+		}
+
 		FScopedTransaction Transaction(LOCTEXT("Ramp_Apply", "Landscape Editing: Add ramp"));
 		ALandscape* Landscape = EdMode->GetLandscape();
 		FGuid EditLayerGUID = EdMode->GetCurrentLayerGuid();
@@ -541,12 +552,12 @@ public:
 
 		// construct the caches, and set them to work in the EditLayer
 		FLandscapeEditDataInterface LandscapeEdit(EdMode->CurrentToolTarget.LandscapeInfo.Get(), EditLayerGUID);
-		FLandscapeHeightCache HeightCache(EdMode->CurrentToolTarget);
-		FLandscapeLayerDataCache<FHeightmapToolTarget> LayerHeightDataCache(EdMode->CurrentToolTarget, HeightCache);
-		LayerHeightDataCache.SetCacheEditingLayer(EditLayerGUID);
+		FLandscapeHeightCache HeightCache(EdMode->CurrentToolTarget.LandscapeInfo.Get());
+		FLandscapeEditLayerStackDataCache<FHeightmapToolTarget> EditLayerHeightDataCache(EdMode->CurrentToolTarget.LandscapeInfo.Get(), HeightCache);
+		EditLayerHeightDataCache.SetCacheEditingLayer(EditLayerGUID);
 
 		const bool bCombinedLayerOperation = EdMode->UISettings->bCombinedLayersOperation && Landscape;
-		LayerHeightDataCache.Initialize(EdMode->CurrentToolTarget.LandscapeInfo.Get(), bCombinedLayerOperation);
+		EditLayerHeightDataCache.Initialize(bCombinedLayerOperation);
 
 		// Heights raster
 		bool bRaiseTerrain = true; //EdMode->UISettings->Ramp_bRaiseTerrain;
@@ -560,7 +571,7 @@ public:
 			int32 ValidMinY = MinY;
 			int32 ValidMaxX = MaxX;
 			int32 ValidMaxY = MaxY;
-			LayerHeightDataCache.Read(ValidMinX, ValidMinY, ValidMaxX, ValidMaxY, Data);
+			EditLayerHeightDataCache.Read(ValidMinX, ValidMinY, ValidMaxX, ValidMaxY, Data);
 
 			if (ValidMinX > ValidMaxX || ValidMinY > ValidMaxY)
 			{
@@ -590,7 +601,7 @@ public:
 			Rasterizer.DrawTriangle(FVector2D(1, Heights[0]), FVector2D(0, Heights[0]), FVector2D(1, Heights[1]), InnerVerts[0][1], OuterVerts[0][1], InnerVerts[1][1], false);
 			Rasterizer.DrawTriangle(FVector2D(0, Heights[0]), FVector2D(1, Heights[1]), FVector2D(0, Heights[1]), OuterVerts[0][1], InnerVerts[1][1], OuterVerts[1][1], false);
 
-			LayerHeightDataCache.Write(MinX, MinY, MaxX, MaxY, Data);
+			EditLayerHeightDataCache.Write(MinX, MinY, MaxX, MaxY, Data);
 
 			// Dirty any runtime virtual textures that our landscape components write to.
 			LandscapeInfo->DirtyRuntimeVirtualTextureForLandscapeArea(MinX, MinY, MaxX, MaxY);
diff --git a/Engine/Source/Editor/LandscapeEditor/Private/LandscapeEdModeSplineTools.cpp b/Engine/Source/Editor/LandscapeEditor/Private/LandscapeEdModeSplineTools.cpp
--- a/Engine/Source/Editor/LandscapeEditor/Private/LandscapeEdModeSplineTools.cpp
+++ b/Engine/Source/Editor/LandscapeEditor/Private/LandscapeEdModeSplineTools.cpp
@@ -99,7 +99,7 @@ public:
 	virtual FText GetDisplayMessage() const override { return NSLOCTEXT("UnrealEd", "LandscapeMode_Splines_Message", "Create a Landscape Spline to carve your landscape, modify blendmasks and deform meshes into roads and other linear features.  Spline mesh settings can be found in the details panel when you have  segments selected."); };
 
 	virtual void SetEditRenderType() override { GLandscapeEditRenderMode = ELandscapeEditRenderMode::None | (GLandscapeEditRenderMode & ELandscapeEditRenderMode::BitMaskForMask); }
-	virtual bool SupportsMask() override { return false; }
+	virtual bool SupportsMask() const override { return false; }
 
 	void AddSegment(ULandscapeSplineControlPoint* Start, ULandscapeSplineControlPoint* End, bool bAutoRotateStart, bool bAutoRotateEnd)
 	{
diff --git a/Engine/Source/Editor/LandscapeEditor/Private/LandscapeEdModeTools.h b/Engine/Source/Editor/LandscapeEditor/Private/LandscapeEdModeTools.h
--- a/Engine/Source/Editor/LandscapeEditor/Private/LandscapeEdModeTools.h
+++ b/Engine/Source/Editor/LandscapeEditor/Private/LandscapeEdModeTools.h
@@ -290,10 +290,12 @@ public:
 	typedef Accessor AccessorClass;
 
 	Accessor DataAccess;
+	ULandscapeLayerInfoObject* LayerInfo = nullptr;
 
-	TLandscapeEditCache(const FLandscapeToolTarget& InTarget)
-		: DataAccess(InTarget)
-		, LandscapeInfo(InTarget.LandscapeInfo)
+	TLandscapeEditCache(ULandscapeInfo* InLandscapeInfo, ULandscapeLayerInfoObject* InLayerInfo)
+		: DataAccess(InLandscapeInfo, InLayerInfo)
+		, LayerInfo(InLayerInfo)
+		, LandscapeInfo(InLandscapeInfo)
 	{
 		check(LandscapeInfo != nullptr);
 	}
@@ -751,19 +753,23 @@ private:
 template<bool bInUseInterp>
 struct FHeightmapAccessorTool : public FHeightmapAccessor<bInUseInterp>
 {
-	FHeightmapAccessorTool(const FLandscapeToolTarget& InTarget)
-	:	FHeightmapAccessor<bInUseInterp>(InTarget.LandscapeInfo.Get())
+	FHeightmapAccessorTool(ULandscapeInfo* InLandscapeInfo, ULandscapeLayerInfoObject* InLayerInfo = nullptr)
+	:	FHeightmapAccessor<bInUseInterp>(InLandscapeInfo, InLayerInfo)
 	{
+		// We don't care about InLayerInfo since heightmaps are not a target layer per se
+		check(InLayerInfo == nullptr);
 	}
 };
 
-struct FLandscapeHeightCache : public TLandscapeEditCache<FHeightmapAccessorTool<true>, uint16>
+struct FLandscapeHeightCache : public TLandscapeEditCache<FHeightmapAccessorTool</*bInUseInterp = */true>, uint16>
 {
 	static uint16 ClampValue(int32 Value) { return static_cast<uint16>(FMath::Clamp(Value, 0, LandscapeDataAccess::MaxValue)); }
 
-	FLandscapeHeightCache(const FLandscapeToolTarget& InTarget)
-		: TLandscapeEditCache<FHeightmapAccessorTool<true>, uint16>(InTarget)
+	FLandscapeHeightCache(ULandscapeInfo* InLandscapeInfo, ULandscapeLayerInfoObject* InLayerInfo = nullptr)
+		: TLandscapeEditCache<FHeightmapAccessorTool</*bInUseInterp = */true>, uint16>(InLandscapeInfo, /*InLayerInfo = */nullptr)
 	{
+		// We don't care about InLayerInfo since heightmaps are not a target layer per se
+		check(InLayerInfo == nullptr);
 	}
 };
 
@@ -773,19 +779,14 @@ struct TAlphamapAccessorTool : public TAlphamapAccessor<bInUseInterp>
 	TAlphamapAccessorTool(ULandscapeInfo* InLandscapeInfo, ULandscapeLayerInfoObject* InLayerInfo)
 		: TAlphamapAccessor<bInUseInterp>(InLandscapeInfo, InLayerInfo)
 	{}
-
-	TAlphamapAccessorTool(const FLandscapeToolTarget& InTarget)
-		: TAlphamapAccessor<bInUseInterp>(InTarget.LandscapeInfo.Get(), InTarget.LayerInfo.Get())
-	{
-	}
 };
 
 struct FLandscapeAlphaCache : public TLandscapeEditCache<TAlphamapAccessorTool</*bInUseInterp = */true>, uint8>
 {
 	static uint8 ClampValue(int32 Value) { return static_cast<uint8>(FMath::Clamp(Value, 0, 255)); }
 
-	FLandscapeAlphaCache(const FLandscapeToolTarget& InTarget)
-		: TLandscapeEditCache<TAlphamapAccessorTool</*bInUseInterp = */true>, uint8>(InTarget)
+	FLandscapeAlphaCache(ULandscapeInfo* InLandscapeInfo, ULandscapeLayerInfoObject* InLayerInfo)
+		: TLandscapeEditCache<TAlphamapAccessorTool</*bInUseInterp = */true>, uint8>(InLandscapeInfo, InLayerInfo)
 	{
 	}
 };
@@ -794,27 +795,28 @@ struct FLandscapeVisCache : public TLandscapeEditCache<TAlphamapAccessorTool</*b
 {
 	static uint8 ClampValue(int32 Value) { return static_cast<uint8>(FMath::Clamp(Value, 0, 255)); }
 
-	FLandscapeVisCache(const FLandscapeToolTarget& InTarget)
-		: TLandscapeEditCache<TAlphamapAccessorTool</*bInUseInterp = */false>, uint8>(InTarget)
+	FLandscapeVisCache(ULandscapeInfo* InLandscapeInfo)
+		: TLandscapeEditCache<TAlphamapAccessorTool</*bInUseInterp = */false>, uint8>(InLandscapeInfo, ALandscapeProxy::VisibilityLayer)
 	{
 	}
 };
 
-
+// Allows to cache data from the combined edit layer stack (from the bottom of the stack up to the specified edit layer if bCombinedLayerOperation, 
+//  normal cache operation otherwise
 template<class ToolTarget>
-class FLandscapeLayerDataCache
+class FLandscapeEditLayerStackDataCache
 {
 public:
-	FLandscapeLayerDataCache(const FLandscapeToolTarget& InTarget, typename ToolTarget::CacheClass& Cache)
-		: LandscapeInfo(nullptr)
+	FLandscapeEditLayerStackDataCache(ULandscapeInfo* InLandscapeInfo, ToolTarget::CacheClass& Cache)
+		: LandscapeInfo(InLandscapeInfo)
 		, Landscape(nullptr)
 		, EditingLayerGuid()
 		, EditingLayerIndex(MAX_uint8)
 		, bIsInitialized(false)
 		, bCombinedLayerOperation(false)
-		, bTargetIsHeightmap(InTarget.TargetType == ELandscapeToolTargetType::Heightmap)
+		, bTargetIsHeightmap(ToolTarget::TargetType == ELandscapeToolTargetType::Heightmap)
 		, CacheUpToEditingLayer(Cache)
-		, CacheBottomLayers(InTarget)
+		, CacheBottomLayers(InLandscapeInfo, Cache.LayerInfo)
 	{
 	}
 
@@ -825,14 +827,13 @@ public:
 		EditingLayerGuid = InEditLayerGUID;
 	}
 
-	void Initialize(ULandscapeInfo* InLandscapeInfo, bool InCombinedLayerOperation)
+	void Initialize(bool bInCombinedLayerOperation)
 	{
 		check(EditingLayerGuid.IsSet());	// you must call SetCacheEditingLayer before Initialize
 		if (!bIsInitialized)
 		{
-			LandscapeInfo = InLandscapeInfo;
 			Landscape = LandscapeInfo ? LandscapeInfo->LandscapeActor.Get() : nullptr;
-			bCombinedLayerOperation = Landscape && InCombinedLayerOperation && bTargetIsHeightmap;
+			bCombinedLayerOperation = Landscape && bInCombinedLayerOperation && bTargetIsHeightmap;
 			if (bCombinedLayerOperation)
 			{
 				int32 I = 0;
@@ -944,15 +945,11 @@ template<bool bInUseInterp>
 struct FFullWeightmapAccessor
 {
 	enum { bUseInterp = bInUseInterp };
-	FFullWeightmapAccessor(ULandscapeInfo* InLandscapeInfo)
+	FFullWeightmapAccessor(ULandscapeInfo* InLandscapeInfo, ULandscapeLayerInfoObject* InLayerInfo)
 		: LandscapeInfo(InLandscapeInfo)
 		, LandscapeEdit(InLandscapeInfo)
 	{
-	}
-
-	FFullWeightmapAccessor(const FLandscapeToolTarget& InTarget)
-		: FFullWeightmapAccessor(InTarget.LandscapeInfo.Get())
-	{
+		// We don't care about InLayerInfo, actually, since we access all target layers here
 	}
 
 	void SetEditLayer(const FGuid& InEditLayerGUID)
@@ -1002,8 +999,8 @@ private:
 
 struct FLandscapeFullWeightCache : public TLandscapeEditCache<FFullWeightmapAccessor<false>, TArray<uint8>>
 {
-	FLandscapeFullWeightCache(const FLandscapeToolTarget& InTarget)
-		: TLandscapeEditCache<FFullWeightmapAccessor<false>, TArray<uint8>>(InTarget)
+	FLandscapeFullWeightCache(ULandscapeInfo* InLandscapeInfo)
+		: TLandscapeEditCache<FFullWeightmapAccessor</*bInUseInterp = */false>, TArray<uint8>>(InLandscapeInfo, /*InLayerInfo = */ nullptr)
 	{
 	}
 
@@ -1068,9 +1065,9 @@ struct FLandscapeFullWeightCache : public TLandscapeEditCache<FFullWeightmapAcce
 		DataAccess.SetData(X1, Y1, X2, Y2, Data.GetData(), PaintingRestriction);
 	}
 
-	void AddDirtyLayer(ULandscapeLayerInfoObject* LayerInfo)
+	void AddDirtyLayer(ULandscapeLayerInfoObject* InLayerInfo)
 	{
-		DataAccess.DirtyLayerInfos.Add(LayerInfo);
+		DataAccess.DirtyLayerInfos.Add(InLayerInfo);
 	}
 };
 
@@ -1081,14 +1078,11 @@ template<bool bInUseInterp>
 struct FDatamapAccessor
 {
 	enum { bUseInterp = bInUseInterp };
-	FDatamapAccessor(ULandscapeInfo* InLandscapeInfo)
+	FDatamapAccessor(ULandscapeInfo* InLandscapeInfo, ULandscapeLayerInfoObject* InLayerInfo = nullptr)
 		: LandscapeEdit(InLandscapeInfo)
 	{
-	}
-
-	FDatamapAccessor(const FLandscapeToolTarget& InTarget)
-		: FDatamapAccessor(InTarget.LandscapeInfo.Get())
-	{
+		// We don't care about InLayerInfo since data map is not a target layer per se
+		check(InLayerInfo == nullptr);
 	}
 
 	void GetData(int32& X1, int32& Y1, int32& X2, int32& Y2, TMap<FIntPoint, uint8>& Data)
@@ -1118,12 +1112,12 @@ private:
 	FLandscapeEditDataInterface LandscapeEdit;
 };
 
-struct FLandscapeDataCache : public TLandscapeEditCache<FDatamapAccessor<false>, uint8>
+struct FLandscapeDataCache : public TLandscapeEditCache<FDatamapAccessor</*bInUseInterp = */false>, uint8>
 {
 	static uint8 ClampValue(int32 Value) { return static_cast<uint8>(FMath::Clamp(Value, 0, 255)); }
 
-	FLandscapeDataCache(const FLandscapeToolTarget& InTarget)
-		: TLandscapeEditCache<FDatamapAccessor<false>, uint8>(InTarget)
+	FLandscapeDataCache(ULandscapeInfo* InLandscapeInfo)
+		: TLandscapeEditCache<FDatamapAccessor</*bInUseInterp = */false>, uint8>(InLandscapeInfo, /*InLayerInfo = */ nullptr)
 	{
 	}
 };
@@ -1184,16 +1178,14 @@ struct FWeightmapToolTarget
 class FLandscapeToolStrokeBase : protected FGCObject
 {
 public:
-	// Whether to call Apply() every frame even if the mouse hasn't moved.  May be overwritten by child classes.
-	bool bUseContinuousApply = false;
-
 	// This is also the expected signature of derived class constructor used by FLandscapeToolBase
 	FLandscapeToolStrokeBase(FEdModeLandscape* InEdMode, FEditorViewportClient* InViewportClient, const FLandscapeToolTarget& InTarget)
 		: EdMode(InEdMode)
-		, Target(InTarget)
 		, LandscapeInfo(InTarget.LandscapeInfo.Get())
 	{
-		bUseContinuousApply = (InTarget.TargetType == ELandscapeToolTargetType::Weightmap) ? InEdMode->UISettings->bApplyWithoutMovingPaint : InEdMode->UISettings->bApplyWithoutMovingSculpt;
+		check(InEdMode->CurrentTool != nullptr);
+		ELandscapeToolTargetType SupportedTargetType = InEdMode->CurrentTool->GetSupportedTargetType();
+		bUseContinuousApply = (SupportedTargetType == ELandscapeToolTargetType::Weightmap) ? InEdMode->UISettings->bApplyWithoutMovingPaint : InEdMode->UISettings->bApplyWithoutMovingSculpt;
 	}
 
 	virtual void SetEditLayer(const FGuid& EditLayerGUID)
@@ -1216,13 +1208,18 @@ public:
 		return TEXT("FLandscapeToolStrokeBase");
 	}
 
+	bool GetUseContinuousApply() const
+	{
+		return bUseContinuousApply;
+	}
+
 protected:
 	FEdModeLandscape* EdMode = nullptr;
-	const FLandscapeToolTarget& Target;
 	TObjectPtr<ULandscapeInfo> LandscapeInfo = nullptr;
+	// Whether to call Apply() every frame even if the mouse hasn't moved.  May be overwritten by child classes.
+	bool bUseContinuousApply = false;
 };
 
-
 /**
  * FLandscapeToolBase - base class for painting tools
  *		ToolTarget - the target for the tool (weight or heightmap)
@@ -1275,7 +1272,7 @@ public:
 			if (ShouldUpdateEditingLayer())
 			{
 				Landscape->RequestLayersContentUpdate(GetBeginToolContentUpdateFlag());
-				Landscape->SetEditingLayer(this->EdMode->GetCurrentLayerGuid());	// legacy way to set the edit layer, via Landscape state
+				Landscape->SetEditingLayer(this->EdMode->GetCurrentLayerGuid()); // legacy way to set the edit layer, via Landscape state
 			}
 			Landscape->SetGrassUpdateEnabled(false);
 		}
@@ -1287,14 +1284,14 @@ public:
 
 		if (ensure(!IsToolActive()))
 		{
-			ToolStroke.Emplace( EdMode, ViewportClient, InTarget );				// construct the tool stroke class
-			ToolStroke->SetEditLayer(this->EdMode->GetCurrentLayerGuid());		// set the edit layer explicitly (if the tool supports this path)
+			ToolStroke.Emplace(EdMode, ViewportClient, InTarget); // construct the tool stroke class
+			ToolStroke->SetEditLayer(this->EdMode->GetCurrentLayerGuid()); // set the edit layer explicitly (if the tool supports this path)
 			EdMode->CurrentBrush->BeginStroke(static_cast<float>(InHitLocation.X), static_cast<float>(InHitLocation.Y), this);
 		}
 
 		// Save the mouse position  
 		LastInteractorPosition = FVector2D(InHitLocation);
-		InteractorPositions.Emplace(LastInteractorPosition, ViewportClient ? IsModifierPressed(ViewportClient) : false); // Copy tool sometimes activates without a specific viewport via ctrl+c hotkey
+		InteractorPositions.Emplace(LastInteractorPosition, IsModifierPressed(ViewportClient), IsAlternateModifierPressed(ViewportClient));
 		TimeSinceLastInteractorMove = 0.0f;
 
 		ToolStroke->Apply(ViewportClient, EdMode->CurrentBrush, EdMode->UISettings, InteractorPositions);
@@ -1313,9 +1310,9 @@ public:
 				ViewportClient->Invalidate(false, false);
 				InteractorPositions.Empty(1);
 			}
-			else if (ToolStroke->bUseContinuousApply && TimeSinceLastInteractorMove >= 0.25f)
+			else if (ToolStroke->GetUseContinuousApply() && TimeSinceLastInteractorMove >= 0.25f)
 			{
-				InteractorPositions.Emplace(LastInteractorPosition, IsModifierPressed(ViewportClient));
+				InteractorPositions.Emplace(LastInteractorPosition, IsModifierPressed(ViewportClient), IsAlternateModifierPressed(ViewportClient));
 				ToolStroke->Apply(ViewportClient, EdMode->CurrentBrush, EdMode->UISettings, InteractorPositions);
 				ViewportClient->Invalidate(false, false);
 				InteractorPositions.Empty(1);
@@ -1379,7 +1376,7 @@ public:
 					if (InteractorPositions.Num() == 0 || LastInteractorPosition != FVector2D(HitLocation))
 					{
 						LastInteractorPosition = FVector2D(HitLocation);
-						InteractorPositions.Emplace(LastInteractorPosition, IsModifierPressed(ViewportClient));
+						InteractorPositions.Emplace(LastInteractorPosition, IsModifierPressed(ViewportClient), IsAlternateModifierPressed(ViewportClient));
 					}
 					TimeSinceLastInteractorMove = 0.0f;
 				}
@@ -1391,7 +1388,7 @@ public:
 			if (InteractorPositions.Num() == 0 || LastInteractorPosition != FVector2D(NewPosition))
 			{
 				LastInteractorPosition = FVector2D(NewPosition);
-				InteractorPositions.Emplace(LastInteractorPosition, IsModifierPressed());
+				InteractorPositions.Emplace(LastInteractorPosition, IsModifierPressed(ViewportClient), IsAlternateModifierPressed(ViewportClient));
 			}
 			TimeSinceLastInteractorMove = 0.0f;
 		}
@@ -1404,6 +1401,11 @@ public:
 	virtual void SetCanToolBeActivated(bool Value) { bCanToolBeActivated = Value; }
 	virtual bool CanToolBeActivated() const {	return bCanToolBeActivated; }
 
+	virtual ELandscapeTargetLayerActionType GetTargetLayerActionType(const ULandscapeLayerInfoObject* InLayerInfo, bool bInIsModifierPressed, bool bInIsAlternateModifierPressed) const override
+	{
+		return ELandscapeTargetLayerActionType::None;
+	}
+
 protected:
 	TArray<FLandscapeToolInteractorPosition> InteractorPositions;
 	FVector2D LastInteractorPosition;
@@ -1412,9 +1414,15 @@ protected:
 	bool bCanToolBeActivated;
 	TOptional<TStrokeClass> ToolStroke;
 
-	bool IsModifierPressed(const class FEditorViewportClient* ViewportClient = nullptr)
+	bool IsModifierPressed(const class FEditorViewportClient* ViewportClient)
 	{
-		UE_LOG(LogLandscapeTools, VeryVerbose, TEXT("ViewportClient = %d, IsShiftDown = %d"), (ViewportClient != nullptr), (ViewportClient != nullptr && IsShiftDown(ViewportClient->Viewport)));
+		UE_LOG(LogLandscapeTools, VeryVerbose, TEXT("ViewportClient = %d, IsModifierDown = %d"), (ViewportClient != nullptr), (ViewportClient != nullptr && IsShiftDown(ViewportClient->Viewport)));
 		return ViewportClient != nullptr && IsShiftDown(ViewportClient->Viewport);
 	}
+
+	bool IsAlternateModifierPressed(const FEditorViewportClient* ViewportClient)
+	{
+		UE_LOG(LogLandscapeTools, VeryVerbose, TEXT("ViewportClient = %d, IsAlternateModifierDown = %d"), (ViewportClient != nullptr), (ViewportClient != nullptr && IsCtrlDown(ViewportClient->Viewport)));
+		return ViewportClient != nullptr && IsCtrlDown(ViewportClient->Viewport);
+	}
 };
diff --git a/Engine/Source/Editor/LandscapeEditor/Private/LandscapeEditorDetailCustomization_ImportExport.cpp b/Engine/Source/Editor/LandscapeEditor/Private/LandscapeEditorDetailCustomization_ImportExport.cpp
--- a/Engine/Source/Editor/LandscapeEditor/Private/LandscapeEditorDetailCustomization_ImportExport.cpp
+++ b/Engine/Source/Editor/LandscapeEditor/Private/LandscapeEditorDetailCustomization_ImportExport.cpp
@@ -157,10 +157,17 @@ void FLandscapeEditorDetailCustomization_ImportExport::CustomizeDetails(IDetailL
 				+ SHorizontalBox::Slot()
 				.FillWidth(1)
 				[
-					SNew(SEditableTextBox)
-					.Font(DetailBuilder.GetDetailFont())
-					.Text_Static(&GetPropertyValueText, PropertyHandle_HeightmapFilename)
-					.OnTextCommitted_Static(&SetFilename, PropertyHandle_HeightmapFilename)
+					SNew(SBox)
+					// This path is edited via a SEditableTextBox (rather than using CreatePropertyValueWidget) so we have to specify a tooltip in order to display the full path if 
+					//  the text box is to small to display the entire thing
+					.ToolTipText_Static(&GetPropertyValueText, PropertyHandle_HeightmapFilename)
+					[
+					    SNew(SEditableTextBox)
+					    .Font(DetailBuilder.GetDetailFont())
+					    .Text_Static(&GetPropertyValueText, PropertyHandle_HeightmapFilename)
+					    .OnTextCommitted_Static(&SetFilename, PropertyHandle_HeightmapFilename)
+						.HintText(LOCTEXT("TargetLayerAssetPath_Hint", "(Please specify a heightmap)"))
+					]
 				]
 				+ SHorizontalBox::Slot()
 				.AutoWidth()
@@ -571,117 +578,120 @@ FReply FLandscapeEditorDetailCustomization_ImportExport::OnImportExportButtonCli
 
 		if (IsImporting())
 		{
-			FGuid CurrentLayerGuid = LandscapeEdMode->GetCurrentLayerGuid();
+			if (LandscapeEdMode->UISettings->CurrentGizmoActor.IsValid())
+			{
+				FGuid CurrentLayerGuid = LandscapeEdMode->GetCurrentLayerGuid();
 						
-			const ELandscapeLayerPaintingRestriction PaintRestriction = ELandscapeLayerPaintingRestriction::None;
+				const ELandscapeLayerPaintingRestriction PaintRestriction = ELandscapeLayerPaintingRestriction::None;
 						
-			ELandscapeImportTransformType TransformType = LandscapeEdMode->UISettings->ImportType;
+				ELandscapeImportTransformType TransformType = LandscapeEdMode->UISettings->ImportType;
 									
-			FVector LocalGizmoPosition = LandscapeInfo->GetLandscapeProxy()->LandscapeActorToWorld().InverseTransformPosition(LandscapeEdMode->UISettings->CurrentGizmoActor->GetActorLocation());
-			FIntPoint LocalGizmoPoint = FIntPoint(FMath::FloorToInt32(LocalGizmoPosition.X), FMath::FloorToInt32(LocalGizmoPosition.Y));
+				FVector LocalGizmoPosition = LandscapeInfo->GetLandscapeProxy()->LandscapeActorToWorld().InverseTransformPosition(LandscapeEdMode->UISettings->CurrentGizmoActor->GetActorLocation());
+				FIntPoint LocalGizmoPoint = FIntPoint(FMath::FloorToInt32(LocalGizmoPosition.X), FMath::FloorToInt32(LocalGizmoPosition.Y));
 			
-			// Update Gizmo Position if we exit and comeback into tool
-			LandscapeEdMode->UISettings->ImportLandscape_GizmoLocalPosition = LocalGizmoPoint;
+				// Update Gizmo Position if we exit and comeback into tool
+				LandscapeEdMode->UISettings->ImportLandscape_GizmoLocalPosition = LocalGizmoPoint;
 
-			FIntRect ImportRegion = LandscapeExtent;
-			FIntPoint ImportOffset(0, 0);
-			if (TransformType == ELandscapeImportTransformType::ExpandOffset)
-			{
-				ImportOffset = LocalGizmoPoint - FIntPoint(LandscapeExtent.Min.X, LandscapeExtent.Min.Y);
-			}
-			else if (TransformType == ELandscapeImportTransformType::None ||
-					 TransformType == ELandscapeImportTransformType::Subregion)
-			{
-				// We only know the image size here if a heightmap is being imported (which requires that all weightmaps match when imported at the same 
-				// time).  Pass in the correct offset consistently and let ImportHeightData/ImportWeightData figure out what that means for the import region.
-				ImportOffset = { LocalGizmoPoint.X, LocalGizmoPoint.Y };
-			}
-			// Note:  Resample mode is always aligned to the landscape and uses no offset.
+				FIntRect ImportRegion = LandscapeExtent;
+				FIntPoint ImportOffset(0, 0);
+				if (TransformType == ELandscapeImportTransformType::ExpandOffset)
+				{
+					ImportOffset = LocalGizmoPoint - FIntPoint(LandscapeExtent.Min.X, LandscapeExtent.Min.Y);
+				}
+				else if (TransformType == ELandscapeImportTransformType::None ||
+						 TransformType == ELandscapeImportTransformType::Subregion)
+				{
+					// We only know the image size here if a heightmap is being imported (which requires that all weightmaps match when imported at the same 
+					// time).  Pass in the correct offset consistently and let ImportHeightData/ImportWeightData figure out what that means for the import region.
+					ImportOffset = { LocalGizmoPoint.X, LocalGizmoPoint.Y };
+				}
+				// Note:  Resample mode is always aligned to the landscape and uses no offset.
 			
-			check(LandscapeEdMode->UISettings->ImportLandscape_HeightmapImportResult != ELandscapeImportResult::Error);
-			const bool bIsWorldPartition = LandscapeEdMode->GetWorld()->GetSubsystem<ULandscapeSubsystem>()->IsGridBased();
+				check(LandscapeEdMode->UISettings->ImportLandscape_HeightmapImportResult != ELandscapeImportResult::Error);
+				const bool bIsWorldPartition = LandscapeEdMode->GetWorld()->GetSubsystem<ULandscapeSubsystem>()->IsGridBased();
 
-			if (bIsWorldPartition && LandscapeEdMode->UISettings->ImportExportMode != ELandscapeImportExportMode::LoadedOnly && LandscapeRegionUtils::NumLandscapeRegions(LandscapeInfo) > 0) // If in world partition import region by region
-			{
-				TArray<ALocationVolume*> LandscapeRegions;
-				TArray<AActor*> Children;
-				LandscapeInfo->LandscapeActor->GetAttachedActors(Children);
-				for (AActor* Child : Children)
+				if (bIsWorldPartition && LandscapeEdMode->UISettings->ImportExportMode != ELandscapeImportExportMode::LoadedOnly && LandscapeRegionUtils::NumLandscapeRegions(LandscapeInfo) > 0) // If in world partition import region by region
 				{
-					if (Child->IsA<ALocationVolume>())
+					TArray<ALocationVolume*> LandscapeRegions;
+					TArray<AActor*> Children;
+					LandscapeInfo->LandscapeActor->GetAttachedActors(Children);
+					for (AActor* Child : Children)
 					{
-						LandscapeRegions.Add(Cast<ALocationVolume>(Child));
+						if (Child->IsA<ALocationVolume>())
+						{
+							LandscapeRegions.Add(Cast<ALocationVolume>(Child));
+						}
 					}
-				}
 
-				int32 NumRegions = LandscapeRegions.Num();
+					int32 NumRegions = LandscapeRegions.Num();
 					
-				FScopedSlowTask Progress(static_cast<float>(NumRegions), LOCTEXT("ImportingLandscapeRegions", "Importing Landscape Regions"));
-				Progress.MakeDialog(/*bShowCancelButton = */ true);
-
-				auto RegionImporter = [&Progress, LandscapeInfo, CurrentLayerGuid, LandscapeEdMode, ImportRegion, TransformType, ImportOffset, PaintRestriction](const FBox& RegionBounds, const TArray<ALandscapeProxy*>& Proxies)
-				{
-					// todo need a better way of working out the landscape bounds of the StreamingProxies in this region.
-					// What if some of the landscape is loaded?
-					FIntRect LandscapeLoadedExtent;
-					LandscapeInfo->GetLandscapeExtent(LandscapeLoadedExtent);
-					LandscapeLoadedExtent.Max.X += 1;
-					LandscapeLoadedExtent.Max.Y += 1;
-
-					Progress.EnterProgressFrame(1.0f, LOCTEXT("ImportingLandscapeRegions", "Importing Landscape Regions"));
+					FScopedSlowTask Progress(static_cast<float>(NumRegions), LOCTEXT("ImportingLandscapeRegions", "Importing Landscape Regions"));
+					Progress.MakeDialog(/*bShowCancelButton = */ true);
 
-					if (LandscapeEdMode->UISettings->bHeightmapSelected)
+					auto RegionImporter = [&Progress, LandscapeInfo, CurrentLayerGuid, LandscapeEdMode, ImportRegion, TransformType, ImportOffset, PaintRestriction](const FBox& RegionBounds, const TArray<ALandscapeProxy*>& Proxies)
 					{
-						LandscapeEdMode->ImportHeightData(LandscapeInfo, CurrentLayerGuid, LandscapeEdMode->UISettings->ImportLandscape_HeightmapFilename, LandscapeLoadedExtent, TransformType, ImportOffset, PaintRestriction, LandscapeEdMode->UISettings->bFlipYAxis);
-					}
+						// todo need a better way of working out the landscape bounds of the StreamingProxies in this region.
+						// What if some of the landscape is loaded?
+						FIntRect LandscapeLoadedExtent;
+						LandscapeInfo->GetLandscapeExtent(LandscapeLoadedExtent);
+						LandscapeLoadedExtent.Max.X += 1;
+						LandscapeLoadedExtent.Max.Y += 1;
 
-					for (const FLandscapeImportLayer& ImportLayer : LandscapeEdMode->UISettings->ImportLandscape_Layers)
-					{
-						if (ImportLayer.bSelected)
+						Progress.EnterProgressFrame(1.0f, LOCTEXT("ImportingLandscapeRegions", "Importing Landscape Regions"));
+
+						if (LandscapeEdMode->UISettings->bHeightmapSelected)
 						{
-							check(ImportLayer.ImportResult != ELandscapeImportResult::Error);
-							LandscapeEdMode->ImportWeightData(LandscapeInfo, CurrentLayerGuid, ImportLayer.LayerInfo, ImportLayer.SourceFilePath, LandscapeLoadedExtent, TransformType, ImportOffset, PaintRestriction, LandscapeEdMode->UISettings->bFlipYAxis);
+							LandscapeEdMode->ImportHeightData(LandscapeInfo, CurrentLayerGuid, LandscapeEdMode->UISettings->ImportLandscape_HeightmapFilename, LandscapeLoadedExtent, TransformType, ImportOffset, PaintRestriction, LandscapeEdMode->UISettings->bFlipYAxis);
 						}
-					}
 
-					return !Progress.ShouldCancel();
-				};
+						for (const FLandscapeImportLayer& ImportLayer : LandscapeEdMode->UISettings->ImportLandscape_Layers)
+						{
+							if (ImportLayer.bSelected)
+							{
+								check(ImportLayer.ImportResult != ELandscapeImportResult::Error);
+								LandscapeEdMode->ImportWeightData(LandscapeInfo, CurrentLayerGuid, ImportLayer.LayerInfo, ImportLayer.SourceFilePath, LandscapeLoadedExtent, TransformType, ImportOffset, PaintRestriction, LandscapeEdMode->UISettings->bFlipYAxis);
+							}
+						}
 
-				LandscapeRegionUtils::ForEachRegion_LoadProcessUnload(LandscapeInfo, World, RegionImporter);
+						return !Progress.ShouldCancel();
+					};
 
-				FLandscapeImageFileCache& LandscapeImageFileCache = FModuleManager::GetModuleChecked<ILandscapeEditorModule>("LandscapeEditor").GetImageFileCache();
-				LandscapeImageFileCache.Clear();
-			}
-			else
-			{
-				FScopedSlowTask Progress(static_cast<float>(1 + LandscapeEdMode->UISettings->ImportLandscape_Layers.Num()), LOCTEXT("ImportingLandscape", "Importing Landscape"));
-				Progress.MakeDialog(/*bShowCancelButton = */ true);
+					LandscapeRegionUtils::ForEachRegion_LoadProcessUnload(LandscapeInfo, World, RegionImporter);
 
-				if (LandscapeEdMode->UISettings->bHeightmapSelected)
-				{
-					Progress.EnterProgressFrame(1.0f, LOCTEXT("ImportingLandscapeHeight", "Importing Landscape Height"));
-					LandscapeEdMode->ImportHeightData(LandscapeInfo, CurrentLayerGuid, LandscapeEdMode->UISettings->ImportLandscape_HeightmapFilename, ImportRegion, TransformType, ImportOffset, PaintRestriction, LandscapeEdMode->UISettings->bFlipYAxis);
+					FLandscapeImageFileCache& LandscapeImageFileCache = FModuleManager::GetModuleChecked<ILandscapeEditorModule>("LandscapeEditor").GetImageFileCache();
+					LandscapeImageFileCache.Clear();
 				}
-
-				for (const FLandscapeImportLayer& ImportLayer : LandscapeEdMode->UISettings->ImportLandscape_Layers)
+				else
 				{
-					if (Progress.ShouldCancel())
+					FScopedSlowTask Progress(static_cast<float>(1 + LandscapeEdMode->UISettings->ImportLandscape_Layers.Num()), LOCTEXT("ImportingLandscape", "Importing Landscape"));
+					Progress.MakeDialog(/*bShowCancelButton = */ true);
+
+					if (LandscapeEdMode->UISettings->bHeightmapSelected)
 					{
-						continue;
+						Progress.EnterProgressFrame(1.0f, LOCTEXT("ImportingLandscapeHeight", "Importing Landscape Height"));
+						LandscapeEdMode->ImportHeightData(LandscapeInfo, CurrentLayerGuid, LandscapeEdMode->UISettings->ImportLandscape_HeightmapFilename, ImportRegion, TransformType, ImportOffset, PaintRestriction, LandscapeEdMode->UISettings->bFlipYAxis);
 					}
 
-					if (ImportLayer.bSelected)
+					for (const FLandscapeImportLayer& ImportLayer : LandscapeEdMode->UISettings->ImportLandscape_Layers)
 					{
-						Progress.EnterProgressFrame(1.0f, LOCTEXT("ImportingLandscapeWeight", "Importing Landscape Weight"));
-						check(ImportLayer.ImportResult != ELandscapeImportResult::Error);
-						LandscapeEdMode->ImportWeightData(LandscapeInfo, CurrentLayerGuid, ImportLayer.LayerInfo, ImportLayer.SourceFilePath, ImportRegion, TransformType, ImportOffset, PaintRestriction, LandscapeEdMode->UISettings->bFlipYAxis);
+						if (Progress.ShouldCancel())
+						{
+							continue;
+						}
+
+						if (ImportLayer.bSelected)
+						{
+							Progress.EnterProgressFrame(1.0f, LOCTEXT("ImportingLandscapeWeight", "Importing Landscape Weight"));
+							check(ImportLayer.ImportResult != ELandscapeImportResult::Error);
+							LandscapeEdMode->ImportWeightData(LandscapeInfo, CurrentLayerGuid, ImportLayer.LayerInfo, ImportLayer.SourceFilePath, ImportRegion, TransformType, ImportOffset, PaintRestriction, LandscapeEdMode->UISettings->bFlipYAxis);
+						}
 					}
-				}
 
-				LandscapeInfo->ForceLayersFullUpdate();
+					LandscapeInfo->ForceLayersFullUpdate();
+				}
 			}
 		}
-		else
+		else // !IsImporting()
 		{
 			auto BuildExportFileName = [](const FString& Filename, const FIntPoint& FileOffset, bool bUseOffset)
 			{
diff --git a/Engine/Source/Editor/LandscapeEditor/Private/LandscapeEditorDetailCustomization_NewLandscape.cpp b/Engine/Source/Editor/LandscapeEditor/Private/LandscapeEditorDetailCustomization_NewLandscape.cpp
--- a/Engine/Source/Editor/LandscapeEditor/Private/LandscapeEditorDetailCustomization_NewLandscape.cpp
+++ b/Engine/Source/Editor/LandscapeEditor/Private/LandscapeEditorDetailCustomization_NewLandscape.cpp
@@ -309,11 +309,17 @@ void FLandscapeEditorDetailCustomization_NewLandscape::CustomizeDetails(IDetailL
 		+ SHorizontalBox::Slot()
 		.FillWidth(1)
 		[
-			SNew(SEditableTextBox)
-			.Font(DetailBuilder.GetDetailFont())
-			.Text_Static(&GetPropertyValueText, PropertyHandle_HeightmapFilename)
-			.OnTextCommitted_Static(&SetImportHeightmapFilenameString, PropertyHandle_HeightmapFilename)
-			.HintText(LOCTEXT("Import_HeightmapNotSet", "(Please specify a heightmap)"))
+			SNew(SBox)
+			// This path is edited via a SEditableTextBox (rather than using CreatePropertyValueWidget) so we have to specify a tooltip in order to display the full path if 
+			//  the text box is to small to display the entire thing
+			.ToolTipText_Static(&GetPropertyValueText, PropertyHandle_HeightmapFilename)
+			[
+				SNew(SEditableTextBox)
+				.Font(DetailBuilder.GetDetailFont())
+				.Text_Static(&GetPropertyValueText, PropertyHandle_HeightmapFilename)
+				.OnTextCommitted_Static(&SetImportHeightmapFilenameString, PropertyHandle_HeightmapFilename)
+				.HintText(LOCTEXT("Import_HeightmapNotSet", "(Please specify a heightmap)"))
+			]
 		]
 		+ SHorizontalBox::Slot()
 		.AutoWidth()
@@ -656,7 +662,6 @@ void FLandscapeEditorDetailCustomization_NewLandscape::CustomizeDetails(IDetailL
 		]
 		+ SHorizontalBox::Slot()
 		.FillWidth(1)
-		.Padding(0,0,12,0) // Line up with the other properties due to having no reset to default button
 		[
 			SNew(SNumericEntryBox<int32>)
 			.Font(DetailBuilder.GetDetailFont())
diff --git a/Engine/Source/Editor/LandscapeEditor/Private/LandscapeEditorDetailCustomization_TargetLayers.cpp b/Engine/Source/Editor/LandscapeEditor/Private/LandscapeEditorDetailCustomization_TargetLayers.cpp
--- a/Engine/Source/Editor/LandscapeEditor/Private/LandscapeEditorDetailCustomization_TargetLayers.cpp
+++ b/Engine/Source/Editor/LandscapeEditor/Private/LandscapeEditorDetailCustomization_TargetLayers.cpp
@@ -3,7 +3,6 @@
 #include "LandscapeEditorDetailCustomization_TargetLayers.h"
 #include "IDetailChildrenBuilder.h"
 #include "Framework/Commands/UIAction.h"
-#include "Widgets/Text/STextBlock.h"
 #include "Framework/MultiBox/MultiBoxBuilder.h"
 #include "Misc/MessageDialog.h"
 #include "Modules/ModuleManager.h"
@@ -12,14 +11,18 @@
 #include "SlateOptMacros.h"
 #include "Framework/Application/MenuStack.h"
 #include "Framework/Application/SlateApplication.h"
+#include "Widgets/Colors/SColorBlock.h"
 #include "Widgets/Images/SImage.h"
-#include "Widgets/Layout/SBox.h"
 #include "Widgets/Input/SButton.h"
 #include "Widgets/Input/SComboButton.h"
 #include "Widgets/Input/SCheckBox.h"
 #include "Widgets/Input/SSearchBox.h"
+#include "Widgets/Text/STextBlock.h"
 #include "Widgets/Input/SMultiLineEditableTextBox.h"
+#include "Widgets/Layout/SBox.h"
+#include "Widgets/Layout/SSeparator.h"
 #include "Widgets/SBoxPanel.h"
+#include "Widgets/SToolTip.h"
 #include "Widgets/Text/SInlineEditableTextBlock.h"
 #include "Dialog/SCustomDialog.h"
 #include "SSimpleButton.h"
@@ -39,7 +42,9 @@
 #include "Algo/Count.h"
 #include "AssetRegistry/AssetRegistryModule.h"
 #include "Materials/MaterialExpressionLandscapeVisibilityMask.h"
+#include "Misc/Crc.h"
 
+#include "SDocumentationToolTip.h"
 #include "SLandscapeEditor.h"
 #include "Dialogs/DlgPickAssetPath.h"
 #include "Dialogs/DlgPickPath.h"
@@ -72,9 +77,9 @@ void FLandscapeEditorStructCustomization_FTargetLayerAssetPath::CustomizeHeader(
 				.Padding(5.0f, 0)
 				[
 					SNew(SCheckBox)
-						.IsChecked_Static(&FLandscapeEditorStructCustomization_FTargetLayerAssetPath::GetUseTargetLayerAssetPathCheckState)
-						.OnCheckStateChanged_Static(&FLandscapeEditorStructCustomization_FTargetLayerAssetPath::OnUseTargetLayerAssetPathCheckStateChanged)
-						.ToolTipText(LOCTEXT("TargetLayerAssetPathCheckbox_ToolTip", "Enable to override the default asset path"))
+					.IsChecked_Static(&FLandscapeEditorStructCustomization_FTargetLayerAssetPath::GetUseTargetLayerAssetPathCheckState)
+					.OnCheckStateChanged_Static(&FLandscapeEditorStructCustomization_FTargetLayerAssetPath::OnUseTargetLayerAssetPathCheckStateChanged)
+					.ToolTipText(LOCTEXT("TargetLayerAssetPathCheckbox_ToolTip", "Enable to override the default asset path"))
 				]
 				+ SHorizontalBox::Slot()
 				.FillWidth(1)
@@ -92,12 +97,18 @@ void FLandscapeEditorStructCustomization_FTargetLayerAssetPath::CustomizeHeader(
 				.FillWidth(1)
 				.VAlign(VAlign_Center)
 				[
-					SNew(SEditableTextBox)
+					SNew(SBox)
+					// This path is edited via a SEditableTextBox (rather than using CreatePropertyValueWidget) so we have to specify a tooltip in order to display the full path if 
+					//  the text box is to small to display the entire thing
+					.ToolTipText_Static(&FLandscapeEditorStructCustomization_FTargetLayerAssetPath::GetTargetLayerAssetFilePath)
+					[
+						SNew(SEditableTextBox)
 						.Font(IDetailLayoutBuilder::GetDetailFont())
 						.Text_Static(&FLandscapeEditorStructCustomization_FTargetLayerAssetPath::GetTargetLayerAssetFilePath)
 						.HintText(LOCTEXT("TargetLayerAssetPath_Hint", "(Specify a default path)"))
 						.IsEnabled_Static(&FLandscapeEditorStructCustomization_FTargetLayerAssetPath::IsUseTargetLayerAssetPathEnabled)
 						.IsReadOnly(true)
+					]
 				]
 				+ SHorizontalBox::Slot()
 				.HAlign(HAlign_Right)
@@ -105,9 +116,9 @@ void FLandscapeEditorStructCustomization_FTargetLayerAssetPath::CustomizeHeader(
 				.AutoWidth()
 				[
 					SNew(SSimpleButton)
-						.OnClicked_Static(&FLandscapeEditorStructCustomization_FTargetLayerAssetPath::OnSetTargetLayerAssetFilePath)
-						.IsEnabled_Static(&FLandscapeEditorStructCustomization_FTargetLayerAssetPath::IsUseTargetLayerAssetPathEnabled)
-						.Icon(FAppStyle::Get().GetBrush("Icons.FolderOpen"))
+					.OnClicked_Static(&FLandscapeEditorStructCustomization_FTargetLayerAssetPath::OnSetTargetLayerAssetFilePath)
+					.IsEnabled_Static(&FLandscapeEditorStructCustomization_FTargetLayerAssetPath::IsUseTargetLayerAssetPathEnabled)
+					.Icon(FAppStyle::Get().GetBrush("Icons.FolderOpen"))
 				]
 		];
 }
@@ -192,6 +203,8 @@ BEGIN_SLATE_FUNCTION_BUILD_OPTIMIZATION
 void FLandscapeEditorDetailCustomization_TargetLayers::CustomizeDetails(IDetailLayoutBuilder& DetailBuilder)
 {
 	TSharedRef<IPropertyHandle> PropertyHandle_PaintingRestriction = DetailBuilder.GetProperty(GET_MEMBER_NAME_CHECKED(ULandscapeEditorObject, PaintingRestriction));
+	TSharedRef<IPropertyHandle> PropertyHandle_WeightBlendedTargetLayerPaintMode = DetailBuilder.GetProperty(GET_MEMBER_NAME_CHECKED(ULandscapeEditorObject, WeightBlendedTargetLayerPaintMode));
+
 	TSharedRef<IPropertyHandle> PropertyHandle_TargetLayerAssetFilePath = DetailBuilder.GetProperty(GET_MEMBER_NAME_CHECKED(ULandscapeEditorObject, TargetLayerAssetFilePath));
 	TSharedRef<IPropertyHandle> PropertyHandle_TargetDisplayOrder = DetailBuilder.GetProperty(GET_MEMBER_NAME_CHECKED(ULandscapeEditorObject, TargetDisplayOrder));
 	PropertyHandle_TargetDisplayOrder->MarkHiddenByCustomization();
@@ -204,6 +217,7 @@ void FLandscapeEditorDetailCustomization_TargetLayers::CustomizeDetails(IDetailL
 	if (!ShouldShowTargetLayers())
 	{
 		PropertyHandle_PaintingRestriction->MarkHiddenByCustomization();
+		PropertyHandle_WeightBlendedTargetLayerPaintMode->MarkHiddenByCustomization();
 		PropertyHandle_TargetLayerAssetFilePath->MarkHiddenByCustomization();
 		return;
 	}
@@ -213,12 +227,15 @@ void FLandscapeEditorDetailCustomization_TargetLayers::CustomizeDetails(IDetailL
 	check(LandscapeEdMode);
 
 	TargetsCategory.AddProperty(PropertyHandle_PaintingRestriction)
-		.Visibility(TAttribute<EVisibility>::Create(TAttribute<EVisibility>::FGetter::CreateStatic(&FLandscapeEditorDetailCustomization_TargetLayers::GetPaintingRestrictionVisibility)));
+		.Visibility(TAttribute<EVisibility>::CreateStatic(&FLandscapeEditorDetailCustomization_TargetLayers::GetPaintingRestrictionVisibility));
+
+	TargetsCategory.AddProperty(PropertyHandle_WeightBlendedTargetLayerPaintMode)
+		.Visibility(TAttribute<EVisibility>::CreateStatic(&FLandscapeEditorDetailCustomization_TargetLayers::GetWeightBlendedTargetLayerPaintModeVisibility));
 
 	TargetsCategory.AddProperty(PropertyHandle_TargetLayerAssetFilePath);
 
 	TargetsCategory.AddCustomRow(FText())
-		.Visibility(TAttribute<EVisibility>::Create(TAttribute<EVisibility>::FGetter::CreateStatic(&FLandscapeEditorDetailCustomization_TargetLayers::GetVisibilityMaskTipVisibility)))
+		.Visibility(TAttribute<EVisibility>::CreateStatic(&FLandscapeEditorDetailCustomization_TargetLayers::GetVisibilityMaskTipVisibility))
 		[
 			SNew(SMultiLineEditableTextBox)
 				.IsReadOnly(true)
@@ -232,7 +249,7 @@ void FLandscapeEditorDetailCustomization_TargetLayers::CustomizeDetails(IDetailL
 		PropertyHandle_TargetDisplayOrder, PropertyHandle_TargetDisplayOrderIsAscendingPropertyHandle, PropertyHandle_TargetShowUnusedLayers)));
 
 	TargetsCategory.AddCustomRow(FText())
-		.Visibility(TAttribute<EVisibility>::Create(TAttribute<EVisibility>::FGetter::CreateStatic(&FLandscapeEditorDetailCustomization_TargetLayers::GetPopulateTargetLayersInfoTipVisibility)))
+		.Visibility(TAttribute<EVisibility>::CreateStatic(&FLandscapeEditorDetailCustomization_TargetLayers::GetPopulateTargetLayersInfoTipVisibility))
 		[
 			SNew(SMultiLineEditableTextBox)
 				.IsReadOnly(true)
@@ -243,7 +260,7 @@ void FLandscapeEditorDetailCustomization_TargetLayers::CustomizeDetails(IDetailL
 		];
 
 	TargetsCategory.AddCustomRow(FText())
-		.Visibility(TAttribute<EVisibility>::Create(TAttribute<EVisibility>::FGetter::CreateStatic(&FLandscapeEditorDetailCustomization_TargetLayers::GetTargetLayersInvalidInfoAssetTipVisibility)))
+		.Visibility(TAttribute<EVisibility>::CreateStatic(&FLandscapeEditorDetailCustomization_TargetLayers::GetTargetLayersInvalidInfoAssetTipVisibility))
 		[
 			SNew(SMultiLineEditableTextBox)
 				.IsReadOnly(true)
@@ -254,7 +271,7 @@ void FLandscapeEditorDetailCustomization_TargetLayers::CustomizeDetails(IDetailL
 		];
 
 	TargetsCategory.AddCustomRow(FText())
-		.Visibility(TAttribute<EVisibility>::Create(TAttribute<EVisibility>::FGetter::CreateStatic(&FLandscapeEditorDetailCustomization_TargetLayers::GetFilteredTargetLayersListInfoTipVisibility)))
+		.Visibility(TAttribute<EVisibility>::CreateStatic(&FLandscapeEditorDetailCustomization_TargetLayers::GetFilteredTargetLayersListInfoTipVisibility))
 		[
 			SNew(SMultiLineEditableTextBox)
 				.IsReadOnly(true)
@@ -269,7 +286,7 @@ END_SLATE_FUNCTION_BUILD_OPTIMIZATION
 bool FLandscapeEditorDetailCustomization_TargetLayers::ShouldShowTargetLayers()
 {
 	FEdModeLandscape* LandscapeEdMode = GetEditorMode();
-	if (LandscapeEdMode && LandscapeEdMode->CurrentToolMode)
+	if (LandscapeEdMode && LandscapeEdMode->CurrentToolMode && LandscapeEdMode->CurrentTool)
 	{
 		const FName CurrentToolName = LandscapeEdMode->CurrentTool->GetToolName();
 
@@ -289,7 +306,7 @@ EVisibility FLandscapeEditorDetailCustomization_TargetLayers::GetPaintingRestric
 {
 	FEdModeLandscape* LandscapeEdMode = GetEditorMode();
 
-	if (LandscapeEdMode && LandscapeEdMode->CurrentToolMode)
+	if (LandscapeEdMode && LandscapeEdMode->CurrentTool)
 	{
 		const FName CurrentToolName = LandscapeEdMode->CurrentTool->GetToolName();
 
@@ -305,6 +322,25 @@ EVisibility FLandscapeEditorDetailCustomization_TargetLayers::GetPaintingRestric
 	return EVisibility::Collapsed;
 }
 
+EVisibility FLandscapeEditorDetailCustomization_TargetLayers::GetWeightBlendedTargetLayerPaintModeVisibility()
+{
+	FEdModeLandscape* LandscapeEdMode = GetEditorMode();
+
+	if (LandscapeEdMode && LandscapeEdMode->CurrentTool)
+	{
+		const FName CurrentToolName = LandscapeEdMode->CurrentTool->GetToolName();
+
+		// Tool target type "Invalid" means Weightmap with no valid paint layer, so technically, it is weightmap and we therefore choose to show WeightBlendedTargetLayerPaintMode : 
+		if ((LandscapeEdMode->CurrentToolTarget.TargetType == ELandscapeToolTargetType::Weightmap && CurrentToolName != TEXT("BlueprintBrush"))
+			|| (LandscapeEdMode->CurrentToolTarget.TargetType == ELandscapeToolTargetType::Invalid))
+		{
+			return EVisibility::Visible;
+		}
+	}
+
+	return EVisibility::Collapsed;
+}
+
 EVisibility FLandscapeEditorDetailCustomization_TargetLayers::GetVisibilityMaskTipVisibility()
 {
 	if (FEdModeLandscape* LandscapeEdMode = GetEditorMode(); (LandscapeEdMode != nullptr) && LandscapeEdMode->CurrentToolTarget.LandscapeInfo.IsValid())
@@ -427,7 +463,7 @@ void FLandscapeEditorCustomNodeBuilder_TargetLayers::GenerateHeaderRowContent(FD
 {
 	FEdModeLandscape* LandscapeEdMode = GetEditorMode();
 	
-	if (LandscapeEdMode == nullptr)
+	if ((LandscapeEdMode == nullptr) || (LandscapeEdMode->CurrentToolMode == nullptr))
 	{
 		return;	
 	}
@@ -859,7 +895,7 @@ TSharedRef<SWidget> FLandscapeEditorCustomNodeBuilder_TargetLayers::GetTargetLay
 const FSlateBrush* FLandscapeEditorCustomNodeBuilder_TargetLayers::GetShowUnusedBrush() const
 {
 	const FSlateBrush* Brush = FAppStyle::GetBrush("Level.VisibleIcon16x");
-	if (FEdModeLandscape* LandscapeEdMode = GetEditorMode(); (LandscapeEdMode != nullptr) && !LandscapeEdMode->UISettings->ShowUnusedLayers)
+	if (FEdModeLandscape* LandscapeEdMode = GetEditorMode(); (LandscapeEdMode != nullptr) && LandscapeEdMode->UISettings && !LandscapeEdMode->UISettings->ShowUnusedLayers)
 	{
 		Brush = FAppStyle::GetBrush("Level.NotVisibleIcon16x");
 	}
@@ -1232,9 +1268,12 @@ TSharedPtr<SWidget> FLandscapeEditorCustomNodeBuilder_TargetLayers::GenerateRow(
 {
 	TSharedPtr<SWidget> RowWidget;
 
-	if (FEdModeLandscape* LandscapeEdMode = GetEditorMode())
+	if (FEdModeLandscape* LandscapeEdMode = GetEditorMode(); LandscapeEdMode && LandscapeEdMode->CurrentTool && LandscapeEdMode->CurrentToolMode)
 	{
-		if ((LandscapeEdMode->CurrentTool->GetSupportedTargetTypes() & LandscapeEdMode->CurrentToolMode->SupportedTargetTypes & UE::Landscape::GetLandscapeToolTargetTypeAsFlags(Target->TargetType)) == ELandscapeToolTargetTypeFlags::None)
+		ELandscapeToolTargetTypeFlags SupportedTargetTypes = UE::Landscape::GetLandscapeToolTargetTypeAsFlags(LandscapeEdMode->CurrentTool->GetSupportedTargetType())
+			& LandscapeEdMode->CurrentToolMode->SupportedTargetTypes
+			& UE::Landscape::GetLandscapeToolTargetTypeAsFlags(Target->TargetType);
+		if (SupportedTargetTypes == ELandscapeToolTargetTypeFlags::None)
 		{
 			return RowWidget;
 		}
@@ -1243,39 +1282,58 @@ TSharedPtr<SWidget> FLandscapeEditorCustomNodeBuilder_TargetLayers::GenerateRow(
 	if (Target->TargetType != ELandscapeToolTargetType::Weightmap)
 	{
 		RowWidget = SNew(SLandscapeEditorSelectableBorder)
-			.Padding(0.0f)
+		.Padding(0.0f)
+		.VAlign(VAlign_Center)
+		.OnContextMenuOpening(this, &FLandscapeEditorCustomNodeBuilder_TargetLayers::OnTargetLayerContextMenuOpening, Target, static_cast<int32>(INDEX_NONE))
+		.OnSelected_Static(&FLandscapeEditorCustomNodeBuilder_TargetLayers::OnTargetSelectionChanged, Target)
+		.IsSelected_Static(&FLandscapeEditorCustomNodeBuilder_TargetLayers::GetTargetLayerIsSelected, Target)
+		.Visibility(this, &FLandscapeEditorCustomNodeBuilder_TargetLayers::ShouldShowLayer, Target)
+		[
+			SNew(SHorizontalBox)
+			+ SHorizontalBox::Slot()
+			.AutoWidth()
 			.VAlign(VAlign_Center)
-			.OnContextMenuOpening(this, &FLandscapeEditorCustomNodeBuilder_TargetLayers::OnTargetLayerContextMenuOpening, Target, static_cast<int32>(INDEX_NONE))
-			.OnSelected_Static(&FLandscapeEditorCustomNodeBuilder_TargetLayers::OnTargetSelectionChanged, Target)
-			.IsSelected_Static(&FLandscapeEditorCustomNodeBuilder_TargetLayers::GetTargetLayerIsSelected, Target)
-			.Visibility(this, &FLandscapeEditorCustomNodeBuilder_TargetLayers::ShouldShowLayer, Target)
+			.Padding(FMargin(2))
 			[
-				SNew(SHorizontalBox)
-					+ SHorizontalBox::Slot()
-					.AutoWidth()
-					.VAlign(VAlign_Center)
-					.Padding(FMargin(2))
-					[
-						SNew(SLandscapeAssetThumbnail, /*Asset = */ nullptr, ThumbnailPool, Target->TargetType == ELandscapeToolTargetType::Heightmap ? FName("LandscapeEditor.Target_Heightmap") : FName("LandscapeEditor.Target_Visibility"))
-						.ThumbnailSize(FIntPoint(48, 48))
-					]
-					+ SHorizontalBox::Slot()
-					.VAlign(VAlign_Center)
-					.Padding(4, 0)
+				SNew(SOverlay)
+				+ SOverlay::Slot()
+				[
+					SNew(SLandscapeAssetThumbnail, /*Asset = */ nullptr, ThumbnailPool, Target->TargetType == ELandscapeToolTargetType::Heightmap ? FName("LandscapeEditor.Target_Heightmap") : FName("LandscapeEditor.Target_Visibility"))
+					.ThumbnailSize(FIntPoint(48, 48))
+				]
+
+				// Bottom-right corner text indicating how the current tool will affect the target layer
+				+ SOverlay::Slot()
+				.Padding(2)
+				.VAlign(VAlign_Bottom)
+				.HAlign(HAlign_Right)
+				[
+					SNew(SBorder)
+					.BorderImage(FAppStyle::GetBrush("AssetThumbnail", ".AssetThumbnailStatusBar"))
 					[
-						SNew(SVerticalBox)
-							+ SVerticalBox::Slot()
-							.AutoHeight()
-							.VAlign(VAlign_Center)
-							.Padding(0, 2)
-							[
-								SNew(STextBlock)
-									.Font(IDetailLayoutBuilder::GetDetailFont())
-									.Text(Target->TargetLayerDisplayName)
-									.ColorAndOpacity_Static(&FLandscapeEditorCustomNodeBuilder_TargetLayers::GetTargetTextColor, Target)
-							]
+						SNew(SImage)
+						.Image_Static(&FLandscapeEditorCustomNodeBuilder_TargetLayers::GetTargetLayerActionTypeBrush, Target)
+						.DesiredSizeOverride(FVector2D(12.0f, 12.0f))
 					]
-			];
+				]
+			]
+			+ SHorizontalBox::Slot()
+			.VAlign(VAlign_Center)
+			.Padding(4, 0)
+			[
+				SNew(SVerticalBox)
+				+ SVerticalBox::Slot()
+				.AutoHeight()
+				.VAlign(VAlign_Center)
+				.Padding(0, 2)
+				[
+					SNew(STextBlock)
+					.Font(IDetailLayoutBuilder::GetDetailFont())
+					.Text(Target->TargetLayerDisplayName)
+					.ColorAndOpacity_Static(&FLandscapeEditorCustomNodeBuilder_TargetLayers::GetTargetTextColor, Target)
+				]
+			]
+		];
 	}
 	else
 	{
@@ -1291,6 +1349,7 @@ TSharedPtr<SWidget> FLandscapeEditorCustomNodeBuilder_TargetLayers::GenerateRow(
 			ThumbnailObject = ((LayerInfo != nullptr) && ThumbnailTools::AssetHasCustomThumbnail(LayerInfo->GetFullName())) ? static_cast<UObject*>(LayerInfo) : static_cast<UObject*>(ThumbnailMIC);
 		}
 
+
 		TSharedRef<SWidget> DeleteButton = PropertyCustomizationHelpers::MakeDeleteButton(
 			FSimpleDelegate::CreateStatic(&FLandscapeEditorCustomNodeBuilder_TargetLayers::OnTargetLayerDeleteClicked, Target),
 			LOCTEXT("Tooltip_Delete", "Delete Layer"));
@@ -1299,9 +1358,108 @@ TSharedPtr<SWidget> FLandscapeEditorCustomNodeBuilder_TargetLayers::GenerateRow(
 			FAppStyle::GetBrush(TEXT("Icons.Plus")),
 			FSimpleDelegate::CreateStatic(&FLandscapeEditorCustomNodeBuilder_TargetLayers::OnTargetLayerCreateClicked, Target),
 			LOCTEXT("Tooltip_Create", "Create Layer Info"),
-			TAttribute<bool>::Create(TAttribute<bool>::FGetter::CreateStatic(&FLandscapeEditorCustomNodeBuilder_TargetLayers::GetTargetLayerCreateEnabled, Target)));
+			TAttribute<bool>::CreateStatic(&FLandscapeEditorCustomNodeBuilder_TargetLayers::GetTargetLayerCreateEnabled, Target));
+
+		// The tooltip mimics the asset thumbnail tooltips since we're using SLandscapeAssetThumbnail. It's displaying information about the layer info, 
+		//  especially its blend group
+		TSharedRef<SToolTip> TargetLayerToolTip = 
+			SNew(SToolTip)
+			.TextMargin(FMargin(1.f, -3.f))
+			.BorderImage(FAppStyle::GetBrush("AssetThumbnail.Tooltip.Border"))
+			.Visibility((LayerInfo != nullptr) ? EVisibility::Visible : EVisibility::Collapsed)
+			.Content()
+			[
+				SNew(SDocumentationToolTip)
+				[
+					SNew(SVerticalBox)
 
-		RowWidget = SNew(SLandscapeEditorSelectableBorder)
+					// Layer name
+					+ SVerticalBox::Slot()
+					.Padding(0.f, 0.f, 0.f, 6.f)
+					.AutoHeight()
+					[
+						SNew(STextBlock)
+						.Text(Target->TargetLayerDisplayName)
+						.Font(FAppStyle::GetFontStyle("AssetThumbnail.Font"))
+						.ColorAndOpacity(FStyleColors::White)
+					]
+
+					// Separator
+					+ SVerticalBox::Slot()
+					.Padding(0.f, 0.f, 0.f, 6.f)
+					.AutoHeight()
+					[
+						SNew(SSeparator)
+						.Orientation(Orient_Horizontal)
+						.Thickness(1.f)
+						.ColorAndOpacity(COLOR("#484848FF"))
+						.SeparatorImage(FAppStyle::Get().GetBrush("WhiteBrush"))
+					]
+
+					// Blend Method
+					+ SVerticalBox::Slot()
+					.Padding(0.f, 0.f, 0.f, 6.f)
+					.AutoHeight()
+					[
+						SNew(SHorizontalBox)
+
+						+SHorizontalBox::Slot()
+						.AutoWidth()
+						.Padding(0, 0, 4, 0)
+						[
+							SNew(STextBlock)
+							.Text(LOCTEXT("BlendMethodHeader", "Blend Method:"))
+						]
+
+						+ SHorizontalBox::Slot()
+						.AutoWidth()
+						[
+							SNew(STextBlock)
+							.ColorAndOpacity(FStyleColors::White)
+							.Text_Static(&FLandscapeEditorCustomNodeBuilder_TargetLayers::GetTargetBlendMethodText, Target, /*bInIncludeBlendGroup = */true)
+						]
+					]
+
+					// Blend Group (if applicable)
+					+ SVerticalBox::Slot()
+					.Padding(0.f, 0.f, 0.f, 6.f)
+					.AutoHeight()
+					[
+						SNew(SHorizontalBox)
+						.Visibility_Static(&FLandscapeEditorCustomNodeBuilder_TargetLayers::GetTargetLayerBlendGroupStripVisibility, Target)
+
+						+ SHorizontalBox::Slot()
+						.AutoWidth()
+						.Padding(0.f, 0.f, 4.f, 0.f)
+						[
+							SNew(STextBlock)
+							.Text(LOCTEXT("BlendGroupHeader", "Blend Group:"))
+						]
+
+						+SHorizontalBox::Slot()
+						.AutoWidth()
+						.Padding(0.f, 0.f, 4.f, 0.f)
+						.VAlign(VAlign_Center)
+						[
+							// Display a color box to explain the color in the thumbnail color strip
+							SNew(SColorBlock)
+							.Color_Static(&FLandscapeEditorCustomNodeBuilder_TargetLayers::GetTargetLayerBlendGroupStripColor, Target)
+							.Size(12.f)
+						]
+
+						+ SHorizontalBox::Slot()
+						.AutoWidth()
+						[
+							SNew(STextBlock)
+							.ColorAndOpacity(FStyleColors::White)
+							.Text_Static(&FLandscapeEditorCustomNodeBuilder_TargetLayers::GetTargetBlendGroupText, Target, /*bInAllowNoName = */true)
+						]
+					]
+				]
+			];
+
+		RowWidget = 
+			SNew(SLandscapeEditorSelectableBorder)
 			.Padding(0.0f)
 			.VAlign(VAlign_Center)
 			.OnContextMenuOpening(this, &FLandscapeEditorCustomNodeBuilder_TargetLayers::OnTargetLayerContextMenuOpening, Target, InlineTextBlocks.Num() - 1)
@@ -1310,39 +1468,42 @@ TSharedPtr<SWidget> FLandscapeEditorCustomNodeBuilder_TargetLayers::GenerateRow(
 			.Visibility(this, &FLandscapeEditorCustomNodeBuilder_TargetLayers::ShouldShowLayer, Target)
 			[
 				SNew(SHorizontalBox)
-
-					+ SHorizontalBox::Slot()
-					.AutoWidth()
-					.VAlign(VAlign_Center)
+				+ SHorizontalBox::Slot()
+				.AutoWidth()
+				.VAlign(VAlign_Center)
+				[
+					SNew(SBox)
+					.Padding(FMargin(2.0f, 0.0f, 2.0f, 0.0f))
 					[
-						SNew(SBox)
-							.Padding(FMargin(2.0f, 0.0f, 2.0f, 0.0f))
-							[
-								SNew(SImage)
-									.Image(FCoreStyle::Get().GetBrush("VerticalBoxDragIndicator"))
-							]
+						SNew(SImage)
+						.Image(FCoreStyle::Get().GetBrush("VerticalBoxDragIndicator"))
 					]
+				]
 
-					+ SHorizontalBox::Slot()
-					.AutoWidth()
-					.VAlign(VAlign_Center)
-					.Padding(FMargin(2))
+				+ SHorizontalBox::Slot()
+				.AutoWidth()
+				.VAlign(VAlign_Center)
+				.Padding(FMargin(2))
+				[
+					SNew(SBox)
+					.Visibility_Static(&FLandscapeEditorCustomNodeBuilder_TargetLayers::GetDebugModeLayerUsageVisibility, Target)
+					.WidthOverride(48.0f)
+					.HeightOverride(48.0f)
 					[
-						SNew(SBox)
-						.Visibility_Static(&FLandscapeEditorCustomNodeBuilder_TargetLayers::GetDebugModeLayerUsageVisibility, Target)
-						.WidthOverride(48.0f)
-						.HeightOverride(48.0f)
-						[
-							SNew(SImage)
-							.Image(FCoreStyle::Get().GetBrush("WhiteBrush"))
-							.ColorAndOpacity_Static(&FLandscapeEditorCustomNodeBuilder_TargetLayers::GetLayerUsageDebugColor, Target)
-						]
+						SNew(SImage)
+						.Image(FCoreStyle::Get().GetBrush("WhiteBrush"))
+						.ColorAndOpacity_Static(&FLandscapeEditorCustomNodeBuilder_TargetLayers::GetLayerUsageDebugColor, Target)
 					]
+				]
 
-					+ SHorizontalBox::Slot()
-					.AutoWidth()
-					.VAlign(VAlign_Center)
-					.Padding(FMargin(2))
+				+ SHorizontalBox::Slot()
+				.AutoWidth()
+				.VAlign(VAlign_Center)
+				.Padding(FMargin(2))
+				[
+					SNew(SOverlay)
+					.ToolTip(TargetLayerToolTip)
+					+ SOverlay::Slot()
 					[
 						SNew(SLandscapeAssetThumbnail, ThumbnailObject, ThumbnailPool, Target->bIsLayerReferencedByMaterial ? FName("LandscapeEditor.Target_Weightmap") : FName("LandscapeEditor.Target_Unknown"))
 						.Visibility_Static(&FLandscapeEditorCustomNodeBuilder_TargetLayers::GetDebugModeLayerUsageVisibility_Invert, Target)
@@ -1359,192 +1520,220 @@ TSharedPtr<SWidget> FLandscapeEditorCustomNodeBuilder_TargetLayers::GenerateRow(
 							return false;
 						})
 					]
-					+ SHorizontalBox::Slot()
+
+					// Mimic the asset color strip to show a color for the blend group :
+					+ SOverlay::Slot()
+					.HAlign(HAlign_Fill)
+					.VAlign(VAlign_Bottom)
+					[
+						SNew(SColorBlock)
+						.Color_Static(&FLandscapeEditorCustomNodeBuilder_TargetLayers::GetTargetLayerBlendGroupStripColor, Target)
+						.Size(2.f)
+						.Visibility_Static(&FLandscapeEditorCustomNodeBuilder_TargetLayers::GetTargetLayerBlendGroupStripVisibility, Target)
+					]
+
+					// Bottom-right corner text indicating how the current tool will affect the target layer
+					+ SOverlay::Slot()
+					.Padding(2)
+					.VAlign(VAlign_Bottom)
+					.HAlign(HAlign_Right)
+					[
+						SNew(SBorder)
+						.BorderImage(FAppStyle::GetBrush("AssetThumbnail", ".AssetThumbnailStatusBar"))
+						[
+							SNew(SImage)
+							.Image_Static(&FLandscapeEditorCustomNodeBuilder_TargetLayers::GetTargetLayerActionTypeBrush, Target)
+							.DesiredSizeOverride(FVector2D(12.0f, 12.0f))
+						]
+					]
+				]
+				
+				+ SHorizontalBox::Slot()
+				.VAlign(VAlign_Center)
+				.Padding(4, 0)
+				[
+					SNew(SVerticalBox)
+					+ SVerticalBox::Slot()
+					.AutoHeight()
 					.VAlign(VAlign_Center)
-					.Padding(4, 0)
+					.Padding(4, 3, 0, 3)
 					[
-						SNew(SVerticalBox)
-							+ SVerticalBox::Slot()
-							.AutoHeight()
-							.VAlign(VAlign_Center)
-							.Padding(4, 3, 0, 3)
-							[
-								SNew(SHorizontalBox)
-									+ SHorizontalBox::Slot()
-									[
-										SAssignNew(InlineTextBlocks.Last(), SInlineEditableTextBlock)
-											.Font(IDetailLayoutBuilder::GetDetailFontBold())
-											.Text(Target->TargetLayerDisplayName)
-											.ColorAndOpacity_Static(&FLandscapeEditorCustomNodeBuilder_TargetLayers::GetTargetTextColor, Target)
-											.OnVerifyTextChanged_Lambda([Target](const FText& InNewText, FText& OutErrorMessage)
-											{
-												const FName NewName(InNewText.ToString());
-
-												if (Target->LayerName == NewName)
-												{
-													return true;
-												}
-
-												if (NewName == UMaterialExpressionLandscapeVisibilityMask::ParameterName)
-												{
-													OutErrorMessage = LOCTEXT("LandscapeTargetLayer_RenameFailed_ReservedName", "This target layer name is reserved for internal usage");
-													return false;
-												}
-
-												if (NewName.IsNone())
-												{
-													OutErrorMessage = LOCTEXT("LandscapeTargetLayer_RenameFailed_EmptyName", "Target layer name cannot be empty");
-													return false;
-												}
-
-												ALandscape* Landscape = Cast<ALandscape>(Target->Owner);
-												if (Landscape->HasTargetLayer(NewName))
-												{
-													OutErrorMessage = LOCTEXT("LandscapeTargetLayer_RenameFailed_AlreadyExists", "This target layer name already exists");
-													return false;
-												}
-
-												return true;
-											})
-											.OnTextCommitted_Lambda([Target](const FText& Text, ETextCommit::Type Type)
-											{
-												const FName NewName(Text.ToString());
-												if (Target->LayerName == NewName)
-												{
-													return;
-												}
-
-												FScopedTransaction Transaction(LOCTEXT("LandscapeTargetLayer_Rename", "Rename Target Layer"));
-												ALandscape* Landscape = Cast<ALandscape>(Target->Owner);
+						SNew(SHorizontalBox)
+						+ SHorizontalBox::Slot()
+						[
+							SAssignNew(InlineTextBlocks.Last(), SInlineEditableTextBlock)
+							.Font(IDetailLayoutBuilder::GetDetailFontBold())
+							.Text(Target->TargetLayerDisplayName)
+							.ColorAndOpacity_Static(&FLandscapeEditorCustomNodeBuilder_TargetLayers::GetTargetTextColor, Target)
+							.OnVerifyTextChanged_Lambda([Target](const FText& InNewText, FText& OutErrorMessage)
+							{
+								const FName NewName(InNewText.ToString());
+
+								if (Target->LayerName == NewName)
+								{
+									return true;
+								}
+
+								if (NewName == UMaterialExpressionLandscapeVisibilityMask::ParameterName)
+								{
+									OutErrorMessage = LOCTEXT("LandscapeTargetLayer_RenameFailed_ReservedName", "This target layer name is reserved for internal usage");
+									return false;
+								}
+
+								if (NewName.IsNone())
+								{
+									OutErrorMessage = LOCTEXT("LandscapeTargetLayer_RenameFailed_EmptyName", "Target layer name cannot be empty");
+									return false;
+								}
+
+								ALandscape* Landscape = Cast<ALandscape>(Target->Owner);
+								if (Landscape->HasTargetLayer(NewName))
+								{
+									OutErrorMessage = LOCTEXT("LandscapeTargetLayer_RenameFailed_AlreadyExists", "This target layer name already exists");
+									return false;
+								}
+
+								return true;
+							})
+							.OnTextCommitted_Lambda([Target](const FText& Text, ETextCommit::Type Type)
+							{
+								const FName NewName(Text.ToString());
+								if (Target->LayerName == NewName)
+								{
+									return;
+								}
+
+								FScopedTransaction Transaction(LOCTEXT("LandscapeTargetLayer_Rename", "Rename Target Layer"));
+								ALandscape* Landscape = Cast<ALandscape>(Target->Owner);
 												
-												const TMap<FName, FLandscapeTargetLayerSettings>& TargetLayers = Landscape->GetTargetLayers();
-												Landscape->RemoveTargetLayer(FName(Target->TargetLayerDisplayName.ToString()));
+								const TMap<FName, FLandscapeTargetLayerSettings>& TargetLayers = Landscape->GetTargetLayers();
+								Landscape->RemoveTargetLayer(FName(Target->TargetLayerDisplayName.ToString()));
 											
-												Target->TargetLayerDisplayName = Text;
-												Target->LayerName = FName(Text.ToString());
-												Landscape->AddTargetLayer(Target->LayerName, FLandscapeTargetLayerSettings());	
+								Target->TargetLayerDisplayName = Text;
+								Target->LayerName = FName(Text.ToString());
+								Landscape->AddTargetLayer(Target->LayerName, FLandscapeTargetLayerSettings());	
 												
-												Target->LandscapeInfo->UpdateLayerInfoMap();
-												if (FEdModeLandscape* LandscapeEdMode = GetEditorMode())
-												{
-													LandscapeEdMode->UpdateTargetList();
-												}
-											})
-									]
-									+ SHorizontalBox::Slot()
-									.HAlign(HAlign_Right)
-									[
-										SNew(STextBlock)
-											.Font(IDetailLayoutBuilder::GetDetailFont())
-											.Text_Static(&FLandscapeEditorCustomNodeBuilder_TargetLayers::GetTargetBlendMethodText, Target)
-											.ToolTipText_Static(&FLandscapeEditorCustomNodeBuilder_TargetLayers::GetTargetBlendMethodTooltipText, Target)
-											.ColorAndOpacity_Static(&FLandscapeEditorCustomNodeBuilder_TargetLayers::GetTargetTextColor, Target)
-									]
+								Target->LandscapeInfo->UpdateLayerInfoMap();
+								if (FEdModeLandscape* LandscapeEdMode = GetEditorMode())
+								{
+									LandscapeEdMode->UpdateTargetList();
+								}
+							})
+						]
+						+ SHorizontalBox::Slot()
+						.HAlign(HAlign_Right)
+						[
+							SNew(STextBlock)
+							.Font(IDetailLayoutBuilder::GetDetailFont())
+							.Text_Static(&FLandscapeEditorCustomNodeBuilder_TargetLayers::GetTargetBlendMethodText, Target, /*bInIncludeBlendGroup = */true)
+							.ToolTipText_Static(&FLandscapeEditorCustomNodeBuilder_TargetLayers::GetTargetBlendMethodTooltipText, Target)
+							.ColorAndOpacity_Static(&FLandscapeEditorCustomNodeBuilder_TargetLayers::GetTargetTextColor, Target)
+						]
+					]
+					+ SVerticalBox::Slot()
+					.AutoHeight()
+					.VAlign(VAlign_Center)
+					[
+						SNew(SHorizontalBox)
+						.Visibility_Static(&FLandscapeEditorCustomNodeBuilder_TargetLayers::GetTargetLayerInfoSelectorVisibility, Target)
+						+ SHorizontalBox::Slot()
+						.FillWidth(1)
+						.VAlign(VAlign_Center)
+						[
+							SNew(SObjectPropertyEntryBox)
+							.ObjectPath(Target->LayerInfoObj != nullptr ? Target->LayerInfoObj->GetPathName() : FString())
+							.AllowedClass(ULandscapeLayerInfoObject::StaticClass())
+							.OnObjectChanged_Static(&FLandscapeEditorCustomNodeBuilder_TargetLayers::OnTargetLayerSetObject, Target)
+							.OnShouldFilterAsset_Static(&FLandscapeEditorCustomNodeBuilder_TargetLayers::ShouldFilterLayerInfo, Target->LayerName)
+							.AllowCreate(false)
+						]
+						+ SHorizontalBox::Slot()
+						.AutoWidth()
+						.VAlign(VAlign_Center)
+						[
+							CreateButton
+						]
+						+ SHorizontalBox::Slot()
+						.AutoWidth()
+						.VAlign(VAlign_Center)
+						[
+							DeleteButton
+						]
+					]
+					+ SVerticalBox::Slot()
+					.AutoHeight()
+					[
+						SNew(SHorizontalBox)
+						.Visibility_Static(&FLandscapeEditorCustomNodeBuilder_TargetLayers::GetLayersSubstractiveBlendVisibility, Target)
+						+ SHorizontalBox::Slot()
+						.AutoWidth()
+						.Padding(0, 2, 2, 2)
+						[
+							SNew(SCheckBox)
+							.IsChecked_Static(&FLandscapeEditorCustomNodeBuilder_TargetLayers::IsLayersSubstractiveBlendChecked, Target)
+							.OnCheckStateChanged_Static(&FLandscapeEditorCustomNodeBuilder_TargetLayers::OnLayersSubstractiveBlendChanged, Target)
+							[
+								SNew(STextBlock)
+								.Text(LOCTEXT("SubtractiveBlend", "Subtractive Blend"))
+								.ColorAndOpacity_Static(&FLandscapeEditorCustomNodeBuilder_TargetLayers::GetTargetTextColor, Target)
 							]
-							+ SVerticalBox::Slot()
-							.AutoHeight()
-							.VAlign(VAlign_Center)
+						]
+					]
+						+ SVerticalBox::Slot()
+						.AutoHeight()
+					[
+						SNew(SHorizontalBox)
+						.Visibility_Static(&FLandscapeEditorCustomNodeBuilder_TargetLayers::GetDebugModeColorChannelVisibility, Target)
+						+ SHorizontalBox::Slot()
+						.AutoWidth()
+						.Padding(0, 2, 2, 2)
+						[
+							SNew(SCheckBox)
+							.IsChecked_Static(&FLandscapeEditorCustomNodeBuilder_TargetLayers::DebugModeColorChannelIsChecked, Target, 0)
+							.OnCheckStateChanged_Static(&FLandscapeEditorCustomNodeBuilder_TargetLayers::OnDebugModeColorChannelChanged, Target, 0)
 							[
-								SNew(SHorizontalBox)
-									.Visibility_Static(&FLandscapeEditorCustomNodeBuilder_TargetLayers::GetTargetLayerInfoSelectorVisibility, Target)
-									+ SHorizontalBox::Slot()
-									.FillWidth(1)
-									.VAlign(VAlign_Center)
-									[
-										SNew(SObjectPropertyEntryBox)
-											.ObjectPath(Target->LayerInfoObj != nullptr ? Target->LayerInfoObj->GetPathName() : FString())
-											.AllowedClass(ULandscapeLayerInfoObject::StaticClass())
-											.OnObjectChanged_Static(&FLandscapeEditorCustomNodeBuilder_TargetLayers::OnTargetLayerSetObject, Target)
-											.OnShouldFilterAsset_Static(&FLandscapeEditorCustomNodeBuilder_TargetLayers::ShouldFilterLayerInfo, Target->LayerName)
-											.AllowCreate(false)
-									]
-									+ SHorizontalBox::Slot()
-									.AutoWidth()
-									.VAlign(VAlign_Center)
-									[
-										CreateButton
-									]
-									+ SHorizontalBox::Slot()
-									.AutoWidth()
-									.VAlign(VAlign_Center)
-									[
-										DeleteButton
-									]
+								SNew(STextBlock)
+								.Text(LOCTEXT("ViewMode.Debug_None", "None"))
 							]
-							+ SVerticalBox::Slot()
-							.AutoHeight()
+						]
+						+ SHorizontalBox::Slot()
+						.AutoWidth()
+						.Padding(2)
+						[
+							SNew(SCheckBox)
+							.IsChecked_Static(&FLandscapeEditorCustomNodeBuilder_TargetLayers::DebugModeColorChannelIsChecked, Target, 1)
+							.OnCheckStateChanged_Static(&FLandscapeEditorCustomNodeBuilder_TargetLayers::OnDebugModeColorChannelChanged, Target, 1)
 							[
-								SNew(SHorizontalBox)
-									.Visibility_Static(&FLandscapeEditorCustomNodeBuilder_TargetLayers::GetLayersSubstractiveBlendVisibility, Target)
-									+ SHorizontalBox::Slot()
-									.AutoWidth()
-									.Padding(0, 2, 2, 2)
-									[
-										SNew(SCheckBox)
-											.IsChecked_Static(&FLandscapeEditorCustomNodeBuilder_TargetLayers::IsLayersSubstractiveBlendChecked, Target)
-											.OnCheckStateChanged_Static(&FLandscapeEditorCustomNodeBuilder_TargetLayers::OnLayersSubstractiveBlendChanged, Target)
-											[
-												SNew(STextBlock)
-													.Text(LOCTEXT("SubtractiveBlend", "Subtractive Blend"))
-													.ColorAndOpacity_Static(&FLandscapeEditorCustomNodeBuilder_TargetLayers::GetTargetTextColor, Target)
-											]
-									]
+								SNew(STextBlock)
+								.Text(LOCTEXT("ViewMode.Debug_R", "R"))
 							]
-							+ SVerticalBox::Slot()
-							.AutoHeight()
+						]
+						+ SHorizontalBox::Slot()
+						.AutoWidth()
+						.Padding(2)
+						[
+							SNew(SCheckBox)
+							.IsChecked_Static(&FLandscapeEditorCustomNodeBuilder_TargetLayers::DebugModeColorChannelIsChecked, Target, 2)
+							.OnCheckStateChanged_Static(&FLandscapeEditorCustomNodeBuilder_TargetLayers::OnDebugModeColorChannelChanged, Target, 2)
 							[
-								SNew(SHorizontalBox)
-									.Visibility_Static(&FLandscapeEditorCustomNodeBuilder_TargetLayers::GetDebugModeColorChannelVisibility, Target)
-									+ SHorizontalBox::Slot()
-									.AutoWidth()
-									.Padding(0, 2, 2, 2)
-									[
-										SNew(SCheckBox)
-											.IsChecked_Static(&FLandscapeEditorCustomNodeBuilder_TargetLayers::DebugModeColorChannelIsChecked, Target, 0)
-											.OnCheckStateChanged_Static(&FLandscapeEditorCustomNodeBuilder_TargetLayers::OnDebugModeColorChannelChanged, Target, 0)
-											[
-												SNew(STextBlock)
-													.Text(LOCTEXT("ViewMode.Debug_None", "None"))
-											]
-									]
-									+ SHorizontalBox::Slot()
-									.AutoWidth()
-									.Padding(2)
-									[
-										SNew(SCheckBox)
-											.IsChecked_Static(&FLandscapeEditorCustomNodeBuilder_TargetLayers::DebugModeColorChannelIsChecked, Target, 1)
-											.OnCheckStateChanged_Static(&FLandscapeEditorCustomNodeBuilder_TargetLayers::OnDebugModeColorChannelChanged, Target, 1)
-											[
-												SNew(STextBlock)
-													.Text(LOCTEXT("ViewMode.Debug_R", "R"))
-											]
-									]
-									+ SHorizontalBox::Slot()
-									.AutoWidth()
-									.Padding(2)
-									[
-										SNew(SCheckBox)
-											.IsChecked_Static(&FLandscapeEditorCustomNodeBuilder_TargetLayers::DebugModeColorChannelIsChecked, Target, 2)
-											.OnCheckStateChanged_Static(&FLandscapeEditorCustomNodeBuilder_TargetLayers::OnDebugModeColorChannelChanged, Target, 2)
-											[
-												SNew(STextBlock)
-													.Text(LOCTEXT("ViewMode.Debug_G", "G"))
-											]
-									]
-									+ SHorizontalBox::Slot()
-									.AutoWidth()
-									.Padding(2)
-									[
-										SNew(SCheckBox)
-											.IsChecked_Static(&FLandscapeEditorCustomNodeBuilder_TargetLayers::DebugModeColorChannelIsChecked, Target, 4)
-											.OnCheckStateChanged_Static(&FLandscapeEditorCustomNodeBuilder_TargetLayers::OnDebugModeColorChannelChanged, Target, 4)
-											[
-												SNew(STextBlock)
-													.Text(LOCTEXT("ViewMode.Debug_B", "B"))
-											]
-									]
+								SNew(STextBlock)
+								.Text(LOCTEXT("ViewMode.Debug_G", "G"))
+							]
+						]
+						+ SHorizontalBox::Slot()
+						.AutoWidth()
+						.Padding(2)
+						[
+							SNew(SCheckBox)
+							.IsChecked_Static(&FLandscapeEditorCustomNodeBuilder_TargetLayers::DebugModeColorChannelIsChecked, Target, 4)
+							.OnCheckStateChanged_Static(&FLandscapeEditorCustomNodeBuilder_TargetLayers::OnDebugModeColorChannelChanged, Target, 4)
+							[
+								SNew(STextBlock)
+								.Text(LOCTEXT("ViewMode.Debug_B", "B"))
 							]
+						]
 					]
+				]
 			];
 	}
 
@@ -1636,6 +1825,72 @@ bool FLandscapeEditorCustomNodeBuilder_TargetLayers::GetTargetLayerIsSelected(co
 	return false;
 }
 
+ELandscapeTargetLayerActionType FLandscapeEditorCustomNodeBuilder_TargetLayers::GetTargetLayerActionType(const TSharedRef<FLandscapeTargetListInfo> Target)
+{
+	FEdModeLandscape* LandscapeEdMode = GetEditorMode();
+	if (LandscapeEdMode == nullptr)
+	{
+		return ELandscapeTargetLayerActionType::None;
+	}
+
+	return LandscapeEdMode->GetEstimatedTargetLayerActionType(Target);
+}
+
+const FSlateBrush* FLandscapeEditorCustomNodeBuilder_TargetLayers::GetTargetLayerActionTypeBrush(const TSharedRef<FLandscapeTargetListInfo> Target)
+{
+	switch (GetTargetLayerActionType(Target))
+	{
+		case ELandscapeTargetLayerActionType::Normal:
+		{
+			return FAppStyle::GetBrush("LandscapeEditor.Brush.Action_Regular");
+		}
+		case ELandscapeTargetLayerActionType::Inverted:
+		{
+			return FAppStyle::GetBrush("LandscapeEditor.Brush.Action_Inverse");
+		}
+		default:
+			break;
+	}
+
+	return nullptr;
+}
+
+FLinearColor FLandscapeEditorCustomNodeBuilder_TargetLayers::GetTargetLayerBlendGroupStripColor(const TSharedRef<FLandscapeTargetListInfo> Target)
+{
+	FLinearColor Color = FLinearColor::Gray;
+	if (Target->LandscapeInfo.IsValid() && Target->LayerInfoObj.IsValid())
+	{
+		TOptional<FTargetLayerBlendGroup> BlendGroup = Target->LandscapeInfo->GetBlendGroupForTargetLayer(Target->LayerInfoObj.Get());
+		if (BlendGroup.IsSet())
+		{
+			// Generate a random color name based on the blend group's data (in order to have a consistent result for a given blend method / group name)
+			const FString BlendGroupName = BlendGroup->GetName().ToString();
+			uint32 Hash = FCrc::TypeCrc32<uint32>(static_cast<uint32>(BlendGroup->GetBlendMethod()));
+			Hash = FCrc::MemCrc32(*BlendGroupName, BlendGroupName.Len() * sizeof(BlendGroupName[0]), Hash);
+
+			const uint8 H = Hash & 0xFF;
+			const uint8 S = 128 + ((Hash >> 8) & 0x7F);
+			const uint8 V = 128 + ((Hash >> 16) & 0x7F);
+			return FLinearColor::MakeFromHSV8(H, S, V);
+		}
+	}
+	return Color;
+}
+
+EVisibility FLandscapeEditorCustomNodeBuilder_TargetLayers::GetTargetLayerBlendGroupStripVisibility(const TSharedRef<FLandscapeTargetListInfo> Target)
+{
+	EVisibility Visibility = EVisibility::Collapsed;
+	if (Target->LandscapeInfo.IsValid() && Target->LayerInfoObj.IsValid())
+	{
+		TOptional<FTargetLayerBlendGroup> BlendGroup = Target->LandscapeInfo->GetBlendGroupForTargetLayer(Target->LayerInfoObj.Get());
+		if (BlendGroup.IsSet())
+		{
+			Visibility = EVisibility::Visible;
+		}
+	}
+	return Visibility;
+}
+
 void FLandscapeEditorCustomNodeBuilder_TargetLayers::OnTargetSelectionChanged(const TSharedRef<FLandscapeTargetListInfo> Target)
 {
 	FEdModeLandscape* LandscapeEdMode = GetEditorMode();
@@ -2205,8 +2460,10 @@ void FLandscapeEditorCustomNodeBuilder_TargetLayers::OnLayersSubstractiveBlendCh
 	{
 		FScopedTransaction Transaction(LOCTEXT("Undo_SubtractiveBlend", "Set Subtractive Blend Layer"));
 		ULandscapeEditLayerBase* EditLayer = LandscapeEdMode->GetEditLayer(LandscapeEdMode->GetSelectedEditLayerIndex());
-
-		EditLayer->AddOrUpdateWeightmapAllocationLayerBlend(Target.Get().LayerInfoObj.Get(), NewCheckedState == ECheckBoxState::Checked, /*bInModify =*/true);
+		if (EditLayer != nullptr)
+		{
+			EditLayer->AddOrUpdateWeightmapAllocationLayerBlend(Target.Get().LayerInfoObj.Get(), NewCheckedState == ECheckBoxState::Checked, /*bInModify =*/true);
+		}
 	}
 }
 
@@ -2264,7 +2521,7 @@ void FLandscapeEditorCustomNodeBuilder_TargetLayers::OnDebugModeColorChannelChan
 	}
 }
 
-FText FLandscapeEditorCustomNodeBuilder_TargetLayers::GetTargetBlendMethodText(const TSharedRef<FLandscapeTargetListInfo> InTarget)
+FText FLandscapeEditorCustomNodeBuilder_TargetLayers::GetTargetBlendMethodText(const TSharedRef<FLandscapeTargetListInfo> InTarget, bool bInIncludeBlendGroup)
 {
 	const ULandscapeLayerInfoObject* LayerInfo = InTarget->LayerInfoObj.Get();
 	if (LayerInfo == nullptr)
@@ -2275,13 +2532,42 @@ FText FLandscapeEditorCustomNodeBuilder_TargetLayers::GetTargetBlendMethodText(c
 	const UEnum* ModeEnum = StaticEnum<ELandscapeTargetLayerBlendMethod>();
 	check(ModeEnum != nullptr);
 	FText Result = ModeEnum->GetDisplayNameTextByValue(static_cast<int64>(LayerInfo->GetBlendMethod()));
-	if ((LayerInfo->GetBlendMethod() == ELandscapeTargetLayerBlendMethod::PremultipliedAlphaBlending) && !LayerInfo->GetBlendGroup().IsNone())
+	if (bInIncludeBlendGroup && (LayerInfo->GetBlendMethod() == ELandscapeTargetLayerBlendMethod::PremultipliedAlphaBlending) && !LayerInfo->GetBlendGroup().IsNone())
 	{
-		Result = FText::Format(LOCTEXT("PremultipliedAlphaBlendingBlendGroupName", "{0} ({1})"), Result, FText::FromName(LayerInfo->GetBlendGroup()));
+		FText BlendGroupName = GetTargetBlendGroupText(InTarget, /*bInAllowNoName = */false);
+		if (!BlendGroupName.IsEmpty())
+		{
+			Result = FText::Format(LOCTEXT("PremultipliedAlphaBlendingBlendGroupName", "{0} ({1})"), Result, BlendGroupName);
+		}
 	}
 	return Result;
 }
 
+FText FLandscapeEditorCustomNodeBuilder_TargetLayers::GetTargetBlendGroupText(const TSharedRef<FLandscapeTargetListInfo> InTarget, bool bInAllowNoName)
+{
+	const ULandscapeLayerInfoObject* LayerInfo = InTarget->LayerInfoObj.Get();
+	if (LayerInfo == nullptr)
+	{
+		return FText::GetEmpty();
+	}
+
+	FName LayerGroupName;
+	if (LayerInfo->GetBlendMethod() == ELandscapeTargetLayerBlendMethod::PremultipliedAlphaBlending)
+	{
+		LayerGroupName = LayerInfo->GetBlendGroup();
+	}
+
+	if (!LayerGroupName.IsNone())
+	{
+		return FText::FromName(LayerGroupName);
+	}
+	else if (bInAllowNoName)
+	{
+		return LOCTEXT("UnnamedBlendGroupName", "(unnamed)");
+	}
+	return FText::GetEmpty();
+}
+
 FText FLandscapeEditorCustomNodeBuilder_TargetLayers::GetTargetBlendMethodTooltipText(const TSharedRef<FLandscapeTargetListInfo> InTarget)
 {
 	const ULandscapeLayerInfoObject* LayerInfo = InTarget->LayerInfoObj.Get();
diff --git a/Engine/Source/Editor/LandscapeEditor/Private/LandscapeEditorDetailCustomization_TargetLayers.h b/Engine/Source/Editor/LandscapeEditor/Private/LandscapeEditorDetailCustomization_TargetLayers.h
--- a/Engine/Source/Editor/LandscapeEditor/Private/LandscapeEditorDetailCustomization_TargetLayers.h
+++ b/Engine/Source/Editor/LandscapeEditor/Private/LandscapeEditorDetailCustomization_TargetLayers.h
@@ -56,6 +56,7 @@ public:
 protected:
 	static bool ShouldShowTargetLayers();
 	static EVisibility GetPaintingRestrictionVisibility();
+	static EVisibility GetWeightBlendedTargetLayerPaintModeVisibility();
 	static EVisibility GetVisibilityMaskTipVisibility();
 	static EVisibility GetPopulateTargetLayersInfoTipVisibility();
 	static EVisibility GetTargetLayersInvalidInfoAssetTipVisibility();
@@ -92,6 +93,10 @@ protected:
 
 	TSharedPtr<SWidget> GenerateRow(const TSharedRef<FLandscapeTargetListInfo> Target);
 
+	static ELandscapeTargetLayerActionType GetTargetLayerActionType(const TSharedRef<FLandscapeTargetListInfo> Target);
+	static const FSlateBrush* GetTargetLayerActionTypeBrush(const TSharedRef<FLandscapeTargetListInfo> Target);
+	static FLinearColor GetTargetLayerBlendGroupStripColor(const TSharedRef<FLandscapeTargetListInfo> Target);
+	static EVisibility GetTargetLayerBlendGroupStripVisibility(const TSharedRef<FLandscapeTargetListInfo> Target);
 	static bool GetTargetLayerIsSelected(const TSharedRef<FLandscapeTargetListInfo> Target);
 	static void OnTargetSelectionChanged(const TSharedRef<FLandscapeTargetListInfo> Target);
 	TSharedPtr<SWidget> OnTargetLayerContextMenuOpening(const TSharedRef<FLandscapeTargetListInfo> Target, const int32 InLayerIndex);
@@ -117,7 +122,8 @@ protected:
 	static EVisibility GetDebugModeColorChannelVisibility(const TSharedRef<FLandscapeTargetListInfo> Target);
 	static ECheckBoxState DebugModeColorChannelIsChecked(const TSharedRef<FLandscapeTargetListInfo> Target, int32 Channel);
 	static void OnDebugModeColorChannelChanged(ECheckBoxState NewCheckedState, const TSharedRef<FLandscapeTargetListInfo> Target, int32 Channel);
-	static FText GetTargetBlendMethodText(const TSharedRef<FLandscapeTargetListInfo> InTarget);
+	static FText GetTargetBlendMethodText(const TSharedRef<FLandscapeTargetListInfo> InTarget, bool bInIncludeBlendGroup);
+	static FText GetTargetBlendGroupText(const TSharedRef<FLandscapeTargetListInfo> InTarget, bool bInAllowNoName);
 	static FText GetTargetBlendMethodTooltipText(const TSharedRef<FLandscapeTargetListInfo> InTarget);
 	static FSlateColor GetTargetTextColor(const TSharedRef<FLandscapeTargetListInfo> InTarget);
 
diff --git a/Engine/Source/Editor/LandscapeEditor/Private/LandscapeEditorObject.cpp b/Engine/Source/Editor/LandscapeEditor/Private/LandscapeEditorObject.cpp
--- a/Engine/Source/Editor/LandscapeEditor/Private/LandscapeEditorObject.cpp
+++ b/Engine/Source/Editor/LandscapeEditor/Private/LandscapeEditorObject.cpp
@@ -211,36 +211,36 @@ void ULandscapeEditorObject::PostEditChangeProperty(FPropertyChangedEvent& Prope
 /** Load UI settings from ini file */
 void ULandscapeEditorObject::Load()
 {
-	GConfig->GetFloat(TEXT("LandscapeEdit"), TEXT("ToolStrength"), ToolStrength, GEditorPerProjectIni);
-	GConfig->GetFloat(TEXT("LandscapeEdit"), TEXT("PaintToolStrength"), PaintToolStrength, GEditorPerProjectIni);
-	GConfig->GetFloat(TEXT("LandscapeEdit"), TEXT("WeightTargetValue"), WeightTargetValue, GEditorPerProjectIni);
+	GConfig->GetFloat(TEXT("LandscapeEdit"), *GET_MEMBER_NAME_CHECKED(ULandscapeEditorObject, ToolStrength).ToString(), ToolStrength, GEditorPerProjectIni);
+	GConfig->GetFloat(TEXT("LandscapeEdit"), *GET_MEMBER_NAME_CHECKED(ULandscapeEditorObject, PaintToolStrength).ToString(), PaintToolStrength, GEditorPerProjectIni);
+	GConfig->GetFloat(TEXT("LandscapeEdit"), *GET_MEMBER_NAME_CHECKED(ULandscapeEditorObject, WeightTargetValue).ToString(), WeightTargetValue, GEditorPerProjectIni);
 	bool InbUseWeightTargetValue = bUseWeightTargetValue;
-	GConfig->GetBool(TEXT("LandscapeEdit"), TEXT("bUseWeightTargetValue"), InbUseWeightTargetValue, GEditorPerProjectIni);
+	GConfig->GetBool(TEXT("LandscapeEdit"), *GET_MEMBER_NAME_CHECKED(ULandscapeEditorObject, bUseWeightTargetValue).ToString(), InbUseWeightTargetValue, GEditorPerProjectIni);
 	bUseWeightTargetValue = InbUseWeightTargetValue;
 
-	GConfig->GetFloat(TEXT("LandscapeEdit"), TEXT("BrushRadius"), BrushRadius, GEditorPerProjectIni);
-	GConfig->GetFloat(TEXT("LandscapeEdit"), TEXT("PaintBrushRadius"), PaintBrushRadius, GEditorPerProjectIni);
-	GConfig->GetInt(TEXT("LandscapeEdit"), TEXT("BrushComponentSize"), BrushComponentSize, GEditorPerProjectIni);
-	GConfig->GetFloat(TEXT("LandscapeEdit"), TEXT("BrushFalloff"), BrushFalloff, GEditorPerProjectIni);
-	GConfig->GetFloat(TEXT("LandscapeEdit"), TEXT("PaintBrushFalloff"), PaintBrushFalloff, GEditorPerProjectIni);
+	GConfig->GetFloat(TEXT("LandscapeEdit"), *GET_MEMBER_NAME_CHECKED(ULandscapeEditorObject, BrushRadius).ToString(), BrushRadius, GEditorPerProjectIni);
+	GConfig->GetFloat(TEXT("LandscapeEdit"), *GET_MEMBER_NAME_CHECKED(ULandscapeEditorObject, PaintBrushRadius).ToString(), PaintBrushRadius, GEditorPerProjectIni);
+	GConfig->GetInt(TEXT("LandscapeEdit"), *GET_MEMBER_NAME_CHECKED(ULandscapeEditorObject, BrushComponentSize).ToString(), BrushComponentSize, GEditorPerProjectIni);
+	GConfig->GetFloat(TEXT("LandscapeEdit"), *GET_MEMBER_NAME_CHECKED(ULandscapeEditorObject, BrushFalloff).ToString(), BrushFalloff, GEditorPerProjectIni);
+	GConfig->GetFloat(TEXT("LandscapeEdit"), *GET_MEMBER_NAME_CHECKED(ULandscapeEditorObject, PaintBrushFalloff).ToString(), PaintBrushFalloff, GEditorPerProjectIni);
 	GConfig->GetString(TEXT("LandscapeEdit"), TEXT("TargetLayerAssetFilePath"), TargetLayerAssetFilePath.DirectoryPath.Path, GEditorPerProjectIni);
 	GConfig->GetBool(TEXT("LandscapeEdit"), TEXT("bUseTargetLayerAssetFilePath"), TargetLayerAssetFilePath.bUseAssetDirectoryPath, GEditorPerProjectIni);
 	bool InbUseClayBrush = bUseClayBrush;
-	GConfig->GetBool(TEXT("LandscapeEdit"), TEXT("bUseClayBrush"), InbUseClayBrush, GEditorPerProjectIni);
+	GConfig->GetBool(TEXT("LandscapeEdit"), *GET_MEMBER_NAME_CHECKED(ULandscapeEditorObject, bUseClayBrush).ToString(), InbUseClayBrush, GEditorPerProjectIni);
 	bUseClayBrush = InbUseClayBrush;
-	GConfig->GetFloat(TEXT("LandscapeEdit"), TEXT("AlphaBrushScale"), AlphaBrushScale, GEditorPerProjectIni);
-	GConfig->GetBool(TEXT("LandscapeEdit"), TEXT("AlphaBrushAutoRotate"), bAlphaBrushAutoRotate, GEditorPerProjectIni);
-	GConfig->GetFloat(TEXT("LandscapeEdit"), TEXT("AlphaBrushRotation"), AlphaBrushRotation, GEditorPerProjectIni);
-	GConfig->GetFloat(TEXT("LandscapeEdit"), TEXT("AlphaBrushPanU"), AlphaBrushPanU, GEditorPerProjectIni);
-	GConfig->GetFloat(TEXT("LandscapeEdit"), TEXT("AlphaBrushPanV"), AlphaBrushPanV, GEditorPerProjectIni);
-	GConfig->GetBool(TEXT("LandscapeEdit"), TEXT("bUseWorldSpacePatternBrush"), bUseWorldSpacePatternBrush, GEditorPerProjectIni);
+	GConfig->GetFloat(TEXT("LandscapeEdit"), *GET_MEMBER_NAME_CHECKED(ULandscapeEditorObject, AlphaBrushScale).ToString(), AlphaBrushScale, GEditorPerProjectIni);
+	GConfig->GetBool(TEXT("LandscapeEdit"), *GET_MEMBER_NAME_CHECKED(ULandscapeEditorObject, bAlphaBrushAutoRotate).ToString(), bAlphaBrushAutoRotate, GEditorPerProjectIni);
+	GConfig->GetFloat(TEXT("LandscapeEdit"), *GET_MEMBER_NAME_CHECKED(ULandscapeEditorObject, AlphaBrushRotation).ToString(), AlphaBrushRotation, GEditorPerProjectIni);
+	GConfig->GetFloat(TEXT("LandscapeEdit"), *GET_MEMBER_NAME_CHECKED(ULandscapeEditorObject, AlphaBrushPanU).ToString(), AlphaBrushPanU, GEditorPerProjectIni);
+	GConfig->GetFloat(TEXT("LandscapeEdit"), *GET_MEMBER_NAME_CHECKED(ULandscapeEditorObject, AlphaBrushPanV).ToString(), AlphaBrushPanV, GEditorPerProjectIni);
+	GConfig->GetBool(TEXT("LandscapeEdit"), *GET_MEMBER_NAME_CHECKED(ULandscapeEditorObject, bUseWorldSpacePatternBrush).ToString(), bUseWorldSpacePatternBrush, GEditorPerProjectIni);
 	GConfig->GetVector2D(TEXT("LandscapeEdit"), TEXT("WorldSpacePatternBrushSettings.Origin"), WorldSpacePatternBrushSettings.Origin, GEditorPerProjectIni);
 	GConfig->GetBool(TEXT("LandscapeEdit"), TEXT("WorldSpacePatternBrushSettings.bCenterTextureOnOrigin"), WorldSpacePatternBrushSettings.bCenterTextureOnOrigin, GEditorPerProjectIni);
 	GConfig->GetFloat(TEXT("LandscapeEdit"), TEXT("WorldSpacePatternBrushSettings.RepeatSize"), WorldSpacePatternBrushSettings.RepeatSize, GEditorPerProjectIni);
 	FString AlphaTextureName = (AlphaTexture != nullptr) ? AlphaTexture->GetPathName() : FString();
 	int32 InAlphaTextureChannel = static_cast<int32>(AlphaTextureChannel);
 	GConfig->GetString(TEXT("LandscapeEdit"), TEXT("AlphaTextureName"), AlphaTextureName, GEditorPerProjectIni);
-	GConfig->GetInt(TEXT("LandscapeEdit"), TEXT("AlphaTextureChannel"), InAlphaTextureChannel, GEditorPerProjectIni);
+	GConfig->GetInt(TEXT("LandscapeEdit"), *GET_MEMBER_NAME_CHECKED(ULandscapeEditorObject, AlphaTextureChannel).ToString(), InAlphaTextureChannel, GEditorPerProjectIni);
 	AlphaTextureChannel = (ELandscapeTextureColorChannel)InAlphaTextureChannel;
 	UTexture2D* LoadedTexture = LoadObject<UTexture2D>(nullptr, *AlphaTextureName, nullptr, LOAD_NoWarn);
 	if ((LoadedTexture == nullptr) && !AlphaTextureName.IsEmpty())
@@ -250,96 +250,93 @@ void ULandscapeEditorObject::Load()
 	SetAlphaTexture(LoadedTexture, AlphaTextureChannel);
 
 	int32 InFlattenMode = (int32)ELandscapeToolFlattenMode::Both;
-	GConfig->GetInt(TEXT("LandscapeEdit"), TEXT("FlattenMode"), InFlattenMode, GEditorPerProjectIni);
+	GConfig->GetInt(TEXT("LandscapeEdit"), *GET_MEMBER_NAME_CHECKED(ULandscapeEditorObject, FlattenMode).ToString(), InFlattenMode, GEditorPerProjectIni);
 	FlattenMode = (ELandscapeToolFlattenMode)InFlattenMode;
 
 	bool InbUseSlopeFlatten = bUseSlopeFlatten;
-	GConfig->GetBool(TEXT("LandscapeEdit"), TEXT("bUseSlopeFlatten"), InbUseSlopeFlatten, GEditorPerProjectIni);
+	GConfig->GetBool(TEXT("LandscapeEdit"), *GET_MEMBER_NAME_CHECKED(ULandscapeEditorObject, bUseSlopeFlatten).ToString(), InbUseSlopeFlatten, GEditorPerProjectIni);
 	bUseSlopeFlatten = InbUseSlopeFlatten;
 
 	bool InbPickValuePerApply = bPickValuePerApply;
-	GConfig->GetBool(TEXT("LandscapeEdit"), TEXT("bPickValuePerApply"), InbPickValuePerApply, GEditorPerProjectIni);
+	GConfig->GetBool(TEXT("LandscapeEdit"), *GET_MEMBER_NAME_CHECKED(ULandscapeEditorObject, bPickValuePerApply).ToString(), InbPickValuePerApply, GEditorPerProjectIni);
 	bPickValuePerApply = InbPickValuePerApply;
 
 	bool InbUseFlattenTarget = bUseFlattenTarget;
-	GConfig->GetBool(TEXT("LandscapeEdit"), TEXT("bUseFlattenTarget"), InbUseFlattenTarget, GEditorPerProjectIni);
+	GConfig->GetBool(TEXT("LandscapeEdit"), *GET_MEMBER_NAME_CHECKED(ULandscapeEditorObject, bUseFlattenTarget).ToString(), InbUseFlattenTarget, GEditorPerProjectIni);
 	bUseFlattenTarget = InbUseFlattenTarget;
 
-	GConfig->GetFloat(TEXT("LandscapeEdit"), TEXT("FlattenTarget"), FlattenTarget, GEditorPerProjectIni);
+	GConfig->GetFloat(TEXT("LandscapeEdit"), *GET_MEMBER_NAME_CHECKED(ULandscapeEditorObject, FlattenTarget).ToString(), FlattenTarget, GEditorPerProjectIni);
 
 	bool InbShowFlattenTargetPreview = bShowFlattenTargetPreview;
-	GConfig->GetBool(TEXT("LandscapeEdit"), TEXT("bShowFlattenTargetPreview"), InbShowFlattenTargetPreview, GEditorPerProjectIni);
+	GConfig->GetBool(TEXT("LandscapeEdit"), *GET_MEMBER_NAME_CHECKED(ULandscapeEditorObject, bShowFlattenTargetPreview).ToString(), InbShowFlattenTargetPreview, GEditorPerProjectIni);
 	bShowFlattenTargetPreview = InbShowFlattenTargetPreview;
 
-	GConfig->GetFloat(TEXT("LandscapeEdit"), TEXT("TerraceSmooth"), TerraceSmooth, GEditorPerProjectIni);
-	GConfig->GetFloat(TEXT("LandscapeEdit"), TEXT("TerraceInterval"), TerraceInterval, GEditorPerProjectIni);
+	GConfig->GetFloat(TEXT("LandscapeEdit"), *GET_MEMBER_NAME_CHECKED(ULandscapeEditorObject, TerraceSmooth).ToString(), TerraceSmooth, GEditorPerProjectIni);
+	GConfig->GetFloat(TEXT("LandscapeEdit"), *GET_MEMBER_NAME_CHECKED(ULandscapeEditorObject, TerraceInterval).ToString(), TerraceInterval, GEditorPerProjectIni);
 
-	GConfig->GetFloat(TEXT("LandscapeEdit"), TEXT("RampWidth"), RampWidth, GEditorPerProjectIni);
-	GConfig->GetFloat(TEXT("LandscapeEdit"), TEXT("RampSideFalloff"), RampSideFalloff, GEditorPerProjectIni);
+	GConfig->GetFloat(TEXT("LandscapeEdit"), *GET_MEMBER_NAME_CHECKED(ULandscapeEditorObject, RampWidth).ToString(), RampWidth, GEditorPerProjectIni);
+	GConfig->GetFloat(TEXT("LandscapeEdit"), *GET_MEMBER_NAME_CHECKED(ULandscapeEditorObject, RampSideFalloff).ToString(), RampSideFalloff, GEditorPerProjectIni);
 
-	GConfig->GetBool(TEXT("LandscapeEdit"), TEXT("bCombinedLayersOperation"), bCombinedLayersOperation, GEditorPerProjectIni);
-	GConfig->GetBool(TEXT("LandscapeEdit"), TEXT("bApplyWithoutMovingSculpt"), bApplyWithoutMovingSculpt, GEditorPerProjectIni);
-	GConfig->GetBool(TEXT("LandscapeEdit"), TEXT("bApplyWithoutMovingPaint"), bApplyWithoutMovingPaint, GEditorPerProjectIni);
+	GConfig->GetBool(TEXT("LandscapeEdit"), *GET_MEMBER_NAME_CHECKED(ULandscapeEditorObject, bCombinedLayersOperation).ToString(), bCombinedLayersOperation, GEditorPerProjectIni);
+	GConfig->GetBool(TEXT("LandscapeEdit"), *GET_MEMBER_NAME_CHECKED(ULandscapeEditorObject, bApplyWithoutMovingSculpt).ToString(), bApplyWithoutMovingSculpt, GEditorPerProjectIni);
+	GConfig->GetBool(TEXT("LandscapeEdit"), *GET_MEMBER_NAME_CHECKED(ULandscapeEditorObject, bApplyWithoutMovingPaint).ToString(), bApplyWithoutMovingPaint, GEditorPerProjectIni);
 
-	GConfig->GetInt(TEXT("LandscapeEdit"), TEXT("ErodeThresh"), ErodeThresh, GEditorPerProjectIni);
-	GConfig->GetInt(TEXT("LandscapeEdit"), TEXT("ErodeIterationNum"), ErodeIterationNum, GEditorPerProjectIni);
-	GConfig->GetInt(TEXT("LandscapeEdit"), TEXT("ErodeSurfaceThickness"), ErodeSurfaceThickness, GEditorPerProjectIni);
+	GConfig->GetInt(TEXT("LandscapeEdit"), *GET_MEMBER_NAME_CHECKED(ULandscapeEditorObject, ErodeThresh).ToString(), ErodeThresh, GEditorPerProjectIni);
+	GConfig->GetInt(TEXT("LandscapeEdit"), *GET_MEMBER_NAME_CHECKED(ULandscapeEditorObject, ErodeIterationNum).ToString(), ErodeIterationNum, GEditorPerProjectIni);
+	GConfig->GetInt(TEXT("LandscapeEdit"), *GET_MEMBER_NAME_CHECKED(ULandscapeEditorObject, ErodeSurfaceThickness).ToString(), ErodeSurfaceThickness, GEditorPerProjectIni);
 	int32 InErosionNoiseMode = (int32)ErosionNoiseMode;
-	GConfig->GetInt(TEXT("LandscapeEdit"), TEXT("ErosionNoiseMode"), InErosionNoiseMode, GEditorPerProjectIni);
+	GConfig->GetInt(TEXT("LandscapeEdit"), *GET_MEMBER_NAME_CHECKED(ULandscapeEditorObject, ErosionNoiseMode).ToString(), InErosionNoiseMode, GEditorPerProjectIni);
 	ErosionNoiseMode = (ELandscapeToolErosionMode)InErosionNoiseMode;
-	GConfig->GetFloat(TEXT("LandscapeEdit"), TEXT("ErosionNoiseScale"), ErosionNoiseScale, GEditorPerProjectIni);
-	GConfig->GetBool(TEXT("LandscapeEdit"), TEXT("bErosionUseLayerHardness"), bErosionUseLayerHardness, GEditorPerProjectIni);
+	GConfig->GetFloat(TEXT("LandscapeEdit"), *GET_MEMBER_NAME_CHECKED(ULandscapeEditorObject, ErosionNoiseScale).ToString(), ErosionNoiseScale, GEditorPerProjectIni);
+	GConfig->GetBool(TEXT("LandscapeEdit"), *GET_MEMBER_NAME_CHECKED(ULandscapeEditorObject, bErosionUseLayerHardness).ToString(), bErosionUseLayerHardness, GEditorPerProjectIni);
 
-	GConfig->GetInt(TEXT("LandscapeEdit"), TEXT("RainAmount"), RainAmount, GEditorPerProjectIni);
-	GConfig->GetFloat(TEXT("LandscapeEdit"), TEXT("SedimentCapacity"), SedimentCapacity, GEditorPerProjectIni);
-	GConfig->GetInt(TEXT("LandscapeEdit"), TEXT("HErodeIterationNum"), HErodeIterationNum, GEditorPerProjectIni);
+	GConfig->GetInt(TEXT("LandscapeEdit"), *GET_MEMBER_NAME_CHECKED(ULandscapeEditorObject, RainAmount).ToString(), RainAmount, GEditorPerProjectIni);
+	GConfig->GetFloat(TEXT("LandscapeEdit"), *GET_MEMBER_NAME_CHECKED(ULandscapeEditorObject, SedimentCapacity).ToString(), SedimentCapacity, GEditorPerProjectIni);
+	GConfig->GetInt(TEXT("LandscapeEdit"), *GET_MEMBER_NAME_CHECKED(ULandscapeEditorObject, HErodeIterationNum).ToString(), HErodeIterationNum, GEditorPerProjectIni);
 	int32 InRainDistMode = (int32)RainDistMode;
-	GConfig->GetInt(TEXT("LandscapeEdit"), TEXT("RainDistNoiseMode"), InRainDistMode, GEditorPerProjectIni);
+	GConfig->GetInt(TEXT("LandscapeEdit"), *GET_MEMBER_NAME_CHECKED(ULandscapeEditorObject, RainDistMode).ToString(), InRainDistMode, GEditorPerProjectIni);
 	RainDistMode = (ELandscapeToolHydroErosionMode)InRainDistMode;
-	GConfig->GetFloat(TEXT("LandscapeEdit"), TEXT("RainDistScale"), RainDistScale, GEditorPerProjectIni);
-	GConfig->GetFloat(TEXT("LandscapeEdit"), TEXT("HErosionDetailScale"), HErosionDetailScale, GEditorPerProjectIni);
+	GConfig->GetFloat(TEXT("LandscapeEdit"), *GET_MEMBER_NAME_CHECKED(ULandscapeEditorObject, RainDistScale).ToString(), RainDistScale, GEditorPerProjectIni);
+	GConfig->GetFloat(TEXT("LandscapeEdit"), *GET_MEMBER_NAME_CHECKED(ULandscapeEditorObject, HErosionDetailScale).ToString(), HErosionDetailScale, GEditorPerProjectIni);
 	bool InbHErosionDetailSmooth = bHErosionDetailSmooth;
-	GConfig->GetBool(TEXT("LandscapeEdit"), TEXT("bHErosionDetailSmooth"), InbHErosionDetailSmooth, GEditorPerProjectIni);
+	GConfig->GetBool(TEXT("LandscapeEdit"), *GET_MEMBER_NAME_CHECKED(ULandscapeEditorObject, bHErosionDetailSmooth).ToString(), InbHErosionDetailSmooth, GEditorPerProjectIni);
 	bHErosionDetailSmooth = InbHErosionDetailSmooth;
 
 	int32 InNoiseMode = (int32)NoiseMode;
-	GConfig->GetInt(TEXT("LandscapeEdit"), TEXT("NoiseMode"), InNoiseMode, GEditorPerProjectIni);
+	GConfig->GetInt(TEXT("LandscapeEdit"), *GET_MEMBER_NAME_CHECKED(ULandscapeEditorObject, NoiseMode).ToString(), InNoiseMode, GEditorPerProjectIni);
 	NoiseMode = (ELandscapeToolNoiseMode)InNoiseMode;
-	GConfig->GetFloat(TEXT("LandscapeEdit"), TEXT("NoiseScale"), NoiseScale, GEditorPerProjectIni);
+	GConfig->GetFloat(TEXT("LandscapeEdit"), *GET_MEMBER_NAME_CHECKED(ULandscapeEditorObject, NoiseScale).ToString(), NoiseScale, GEditorPerProjectIni);
 
-	GConfig->GetInt(TEXT("LandscapeEdit"), TEXT("SmoothFilterKernelSize"), SmoothFilterKernelSize, GEditorPerProjectIni);
-	GConfig->GetFloat(TEXT("LandscapeEdit"), TEXT("DetailScale"), DetailScale, GEditorPerProjectIni);
+	GConfig->GetInt(TEXT("LandscapeEdit"), *GET_MEMBER_NAME_CHECKED(ULandscapeEditorObject, SmoothFilterKernelSize).ToString(), SmoothFilterKernelSize, GEditorPerProjectIni);
+	GConfig->GetFloat(TEXT("LandscapeEdit"), *GET_MEMBER_NAME_CHECKED(ULandscapeEditorObject, DetailScale).ToString(), DetailScale, GEditorPerProjectIni);
 	bool InbDetailSmooth = bDetailSmooth;
-	GConfig->GetBool(TEXT("LandscapeEdit"), TEXT("bDetailSmooth"), InbDetailSmooth, GEditorPerProjectIni);
+	GConfig->GetBool(TEXT("LandscapeEdit"), *GET_MEMBER_NAME_CHECKED(ULandscapeEditorObject, bDetailSmooth).ToString(), InbDetailSmooth, GEditorPerProjectIni);
 	bDetailSmooth = InbDetailSmooth;
 
-	GConfig->GetFloat(TEXT("LandscapeEdit"), TEXT("MaximumValueRadius"), MaximumValueRadius, GEditorPerProjectIni);
+	GConfig->GetFloat(TEXT("LandscapeEdit"), *GET_MEMBER_NAME_CHECKED(ULandscapeEditorObject, MaximumValueRadius).ToString(), MaximumValueRadius, GEditorPerProjectIni);
 
 	bool InbSmoothGizmoBrush = bSmoothGizmoBrush;
-	GConfig->GetBool(TEXT("LandscapeEdit"), TEXT("bSmoothGizmoBrush"), InbSmoothGizmoBrush, GEditorPerProjectIni);
+	GConfig->GetBool(TEXT("LandscapeEdit"), *GET_MEMBER_NAME_CHECKED(ULandscapeEditorObject, bSmoothGizmoBrush).ToString(), InbSmoothGizmoBrush, GEditorPerProjectIni);
 	bSmoothGizmoBrush = InbSmoothGizmoBrush;
 
 	int32 InPasteMode = (int32)ELandscapeToolPasteMode::Both;
-	GConfig->GetInt(TEXT("LandscapeEdit"), TEXT("PasteMode"), InPasteMode, GEditorPerProjectIni);
-	//PasteMode = (ELandscapeToolPasteMode)InPasteMode;
+	GConfig->GetInt(TEXT("LandscapeEdit"), *GET_MEMBER_NAME_CHECKED(ULandscapeEditorObject, PasteMode).ToString(), InPasteMode, GEditorPerProjectIni);
 	SetPasteMode((ELandscapeToolPasteMode)InPasteMode);
 
 	int32 InMirrorOp = (int32)ELandscapeMirrorOperation::MinusXToPlusX;
-	GConfig->GetInt(TEXT("LandscapeEdit"), TEXT("MirrorOp"), InMirrorOp, GEditorPerProjectIni);
+	GConfig->GetInt(TEXT("LandscapeEdit"), *GET_MEMBER_NAME_CHECKED(ULandscapeEditorObject, MirrorOp).ToString(), InMirrorOp, GEditorPerProjectIni);
 	MirrorOp = (ELandscapeMirrorOperation)InMirrorOp;
 
-	int32 InConvertMode = (int32)ResizeLandscape_ConvertMode;
-	GConfig->GetInt(TEXT("LandscapeEdit"), TEXT("ConvertMode"), InConvertMode, GEditorPerProjectIni);
-	ResizeLandscape_ConvertMode = (ELandscapeConvertMode)InConvertMode;
+	int32 InResizeLandscape_ConvertMode = (int32)ResizeLandscape_ConvertMode;
+	GConfig->GetInt(TEXT("LandscapeEdit"), *GET_MEMBER_NAME_CHECKED(ULandscapeEditorObject, ResizeLandscape_ConvertMode).ToString(), InResizeLandscape_ConvertMode, GEditorPerProjectIni);
+	ResizeLandscape_ConvertMode = (ELandscapeConvertMode)InResizeLandscape_ConvertMode;
 
 	// Region
-	//GConfig->GetBool(TEXT("LandscapeEdit"), TEXT("bUseSelectedRegion"), bUseSelectedRegion, GEditorPerProjectIni);
-	//GConfig->GetBool(TEXT("LandscapeEdit"), TEXT("bUseNegativeMask"), bUseNegativeMask, GEditorPerProjectIni);
 	bool InbApplyToAllTargets = bApplyToAllTargets;
-	GConfig->GetBool(TEXT("LandscapeEdit"), TEXT("bApplyToAllTargets"), InbApplyToAllTargets, GEditorPerProjectIni);
+	GConfig->GetBool(TEXT("LandscapeEdit"), *GET_MEMBER_NAME_CHECKED(ULandscapeEditorObject, bApplyToAllTargets).ToString(), InbApplyToAllTargets, GEditorPerProjectIni);
 	bApplyToAllTargets = InbApplyToAllTargets;
 
-	GConfig->GetBool(TEXT("LandscapeEdit"), TEXT("ShowUnusedLayers"), ShowUnusedLayers, GEditorPerProjectIni);
+	GConfig->GetBool(TEXT("LandscapeEdit"), *GET_MEMBER_NAME_CHECKED(ULandscapeEditorObject, ShowUnusedLayers).ToString(), ShowUnusedLayers, GEditorPerProjectIni);
 
 	// Set EditRenderMode
 	SetbUseSelectedRegion(bUseSelectedRegion);
@@ -376,9 +373,13 @@ void ULandscapeEditorObject::Load()
 		NewLandscape_Material = LoadObject<UMaterialInterface>(nullptr, *NewLandscapeMaterialName, nullptr, LOAD_NoWarn);
 	}
 	
-	int32 AlphamapType = (uint8)ImportLandscape_AlphamapType;
-	GConfig->GetInt(TEXT("LandscapeEdit"), TEXT("ImportLandscape_AlphamapType"), AlphamapType, GEditorPerProjectIni);
-	ImportLandscape_AlphamapType = (ELandscapeImportAlphamapType)AlphamapType;
+	int32 InImportLandscape_AlphamapType = (uint8)ImportLandscape_AlphamapType;
+	GConfig->GetInt(TEXT("LandscapeEdit"), *GET_MEMBER_NAME_CHECKED(ULandscapeEditorObject, ImportLandscape_AlphamapType).ToString(), InImportLandscape_AlphamapType, GEditorPerProjectIni);
+	ImportLandscape_AlphamapType = (ELandscapeImportAlphamapType)InImportLandscape_AlphamapType;
+
+	int32 InWeightBlendedTargetLayerPaintMode = (int32)WeightBlendedTargetLayerPaintMode;
+	GConfig->GetInt(TEXT("LandscapeEdit"), *GET_MEMBER_NAME_CHECKED(ULandscapeEditorObject, WeightBlendedTargetLayerPaintMode).ToString(), InWeightBlendedTargetLayerPaintMode, GEditorPerProjectIni);
+	WeightBlendedTargetLayerPaintMode = (ELandscapeWeightBlendPaintMode)InWeightBlendedTargetLayerPaintMode;
 
 	RefreshImportLayersList();
 }
@@ -386,86 +387,86 @@ void ULandscapeEditorObject::Load()
 /** Save UI settings to ini file */
 void ULandscapeEditorObject::Save()
 {
-	GConfig->SetFloat(TEXT("LandscapeEdit"), TEXT("ToolStrength"), ToolStrength, GEditorPerProjectIni);
-	GConfig->SetFloat(TEXT("LandscapeEdit"), TEXT("PaintToolStrength"), PaintToolStrength, GEditorPerProjectIni);
-	GConfig->SetFloat(TEXT("LandscapeEdit"), TEXT("WeightTargetValue"), WeightTargetValue, GEditorPerProjectIni);
-	GConfig->SetBool(TEXT("LandscapeEdit"), TEXT("bUseWeightTargetValue"), bUseWeightTargetValue, GEditorPerProjectIni);
-
-	GConfig->SetFloat(TEXT("LandscapeEdit"), TEXT("BrushRadius"), BrushRadius, GEditorPerProjectIni);
-	GConfig->SetFloat(TEXT("LandscapeEdit"), TEXT("PaintBrushRadius"), PaintBrushRadius, GEditorPerProjectIni);
-	GConfig->SetInt(TEXT("LandscapeEdit"), TEXT("BrushComponentSize"), BrushComponentSize, GEditorPerProjectIni);
-	GConfig->SetFloat(TEXT("LandscapeEdit"), TEXT("BrushFalloff"), BrushFalloff, GEditorPerProjectIni);
-	GConfig->SetFloat(TEXT("LandscapeEdit"), TEXT("PaintBrushFalloff"), PaintBrushFalloff, GEditorPerProjectIni);
+	GConfig->SetFloat(TEXT("LandscapeEdit"), *GET_MEMBER_NAME_CHECKED(ULandscapeEditorObject, ToolStrength).ToString(), ToolStrength, GEditorPerProjectIni);
+	GConfig->SetFloat(TEXT("LandscapeEdit"), *GET_MEMBER_NAME_CHECKED(ULandscapeEditorObject, PaintToolStrength).ToString(), PaintToolStrength, GEditorPerProjectIni);
+	GConfig->SetFloat(TEXT("LandscapeEdit"), *GET_MEMBER_NAME_CHECKED(ULandscapeEditorObject, WeightTargetValue).ToString(), WeightTargetValue, GEditorPerProjectIni);
+	GConfig->SetBool(TEXT("LandscapeEdit"), *GET_MEMBER_NAME_CHECKED(ULandscapeEditorObject, bUseWeightTargetValue).ToString(), bUseWeightTargetValue, GEditorPerProjectIni);
+
+	GConfig->SetFloat(TEXT("LandscapeEdit"), *GET_MEMBER_NAME_CHECKED(ULandscapeEditorObject, BrushRadius).ToString(), BrushRadius, GEditorPerProjectIni);
+	GConfig->SetFloat(TEXT("LandscapeEdit"), *GET_MEMBER_NAME_CHECKED(ULandscapeEditorObject, PaintBrushRadius).ToString(), PaintBrushRadius, GEditorPerProjectIni);
+	GConfig->SetInt(TEXT("LandscapeEdit"), *GET_MEMBER_NAME_CHECKED(ULandscapeEditorObject, BrushComponentSize).ToString(), BrushComponentSize, GEditorPerProjectIni);
+	GConfig->SetFloat(TEXT("LandscapeEdit"), *GET_MEMBER_NAME_CHECKED(ULandscapeEditorObject, BrushFalloff).ToString(), BrushFalloff, GEditorPerProjectIni);
+	GConfig->SetFloat(TEXT("LandscapeEdit"), *GET_MEMBER_NAME_CHECKED(ULandscapeEditorObject, PaintBrushFalloff).ToString(), PaintBrushFalloff, GEditorPerProjectIni);
 	GConfig->SetString(TEXT("LandscapeEdit"), TEXT("TargetLayerAssetFilePath"), *TargetLayerAssetFilePath.DirectoryPath.Path, GEditorPerProjectIni);
 	GConfig->SetBool(TEXT("LandscapeEdit"), TEXT("bUseTargetLayerAssetFilePath"), TargetLayerAssetFilePath.bUseAssetDirectoryPath, GEditorPerProjectIni);
-	GConfig->SetBool(TEXT("LandscapeEdit"), TEXT("bUseClayBrush"), bUseClayBrush, GEditorPerProjectIni);
-	GConfig->SetFloat(TEXT("LandscapeEdit"), TEXT("AlphaBrushScale"), AlphaBrushScale, GEditorPerProjectIni);
-	GConfig->SetBool(TEXT("LandscapeEdit"), TEXT("AlphaBrushAutoRotate"), bAlphaBrushAutoRotate, GEditorPerProjectIni);
-	GConfig->SetFloat(TEXT("LandscapeEdit"), TEXT("AlphaBrushRotation"), AlphaBrushRotation, GEditorPerProjectIni);
-	GConfig->SetFloat(TEXT("LandscapeEdit"), TEXT("AlphaBrushPanU"), AlphaBrushPanU, GEditorPerProjectIni);
-	GConfig->SetFloat(TEXT("LandscapeEdit"), TEXT("AlphaBrushPanV"), AlphaBrushPanV, GEditorPerProjectIni);
+	GConfig->SetBool(TEXT("LandscapeEdit"), *GET_MEMBER_NAME_CHECKED(ULandscapeEditorObject, bUseClayBrush).ToString(), bUseClayBrush, GEditorPerProjectIni);
+	GConfig->SetFloat(TEXT("LandscapeEdit"), *GET_MEMBER_NAME_CHECKED(ULandscapeEditorObject, AlphaBrushScale).ToString(), AlphaBrushScale, GEditorPerProjectIni);
+	GConfig->SetBool(TEXT("LandscapeEdit"), *GET_MEMBER_NAME_CHECKED(ULandscapeEditorObject, bAlphaBrushAutoRotate).ToString(), bAlphaBrushAutoRotate, GEditorPerProjectIni);
+	GConfig->SetFloat(TEXT("LandscapeEdit"), *GET_MEMBER_NAME_CHECKED(ULandscapeEditorObject, AlphaBrushRotation).ToString(), AlphaBrushRotation, GEditorPerProjectIni);
+	GConfig->SetFloat(TEXT("LandscapeEdit"), *GET_MEMBER_NAME_CHECKED(ULandscapeEditorObject, AlphaBrushPanU).ToString(), AlphaBrushPanU, GEditorPerProjectIni);
+	GConfig->SetFloat(TEXT("LandscapeEdit"), *GET_MEMBER_NAME_CHECKED(ULandscapeEditorObject, AlphaBrushPanV).ToString(), AlphaBrushPanV, GEditorPerProjectIni);
+	GConfig->SetBool(TEXT("LandscapeEdit"), *GET_MEMBER_NAME_CHECKED(ULandscapeEditorObject, bUseWorldSpacePatternBrush).ToString(), bUseWorldSpacePatternBrush, GEditorPerProjectIni);
 	GConfig->SetVector2D(TEXT("LandscapeEdit"), TEXT("WorldSpacePatternBrushSettings.Origin"), WorldSpacePatternBrushSettings.Origin, GEditorPerProjectIni);
 	GConfig->SetBool(TEXT("LandscapeEdit"), TEXT("WorldSpacePatternBrushSettings.bCenterTextureOnOrigin"), WorldSpacePatternBrushSettings.bCenterTextureOnOrigin, GEditorPerProjectIni);
 	GConfig->SetFloat(TEXT("LandscapeEdit"), TEXT("WorldSpacePatternBrushSettings.RepeatSize"), WorldSpacePatternBrushSettings.RepeatSize, GEditorPerProjectIni);
 	const FString AlphaTextureName = (AlphaTexture != nullptr) ? AlphaTexture->GetPathName() : FString();
 	GConfig->SetString(TEXT("LandscapeEdit"), TEXT("AlphaTextureName"), *AlphaTextureName, GEditorPerProjectIni);
-	GConfig->SetInt(TEXT("LandscapeEdit"), TEXT("AlphaTextureChannel"), (int32)AlphaTextureChannel, GEditorPerProjectIni);
-
-	GConfig->SetInt(TEXT("LandscapeEdit"), TEXT("FlattenMode"), (int32)FlattenMode, GEditorPerProjectIni);
-	GConfig->SetBool(TEXT("LandscapeEdit"), TEXT("bUseSlopeFlatten"), bUseSlopeFlatten, GEditorPerProjectIni);
-	GConfig->SetBool(TEXT("LandscapeEdit"), TEXT("bPickValuePerApply"), bPickValuePerApply, GEditorPerProjectIni);
-	GConfig->SetBool(TEXT("LandscapeEdit"), TEXT("bUseFlattenTarget"), bUseFlattenTarget, GEditorPerProjectIni);
-	GConfig->SetFloat(TEXT("LandscapeEdit"), TEXT("FlattenTarget"), FlattenTarget, GEditorPerProjectIni);
-	GConfig->SetBool(TEXT("LandscapeEdit"), TEXT("bShowFlattenTargetPreview"), bShowFlattenTargetPreview, GEditorPerProjectIni);
-	GConfig->SetFloat(TEXT("LandscapeEdit"), TEXT("TerraceSmooth"), TerraceSmooth, GEditorPerProjectIni);
-	GConfig->SetFloat(TEXT("LandscapeEdit"), TEXT("TerraceInterval"), TerraceInterval, GEditorPerProjectIni);
-
-	GConfig->SetFloat(TEXT("LandscapeEdit"), TEXT("RampWidth"), RampWidth, GEditorPerProjectIni);
-	GConfig->SetFloat(TEXT("LandscapeEdit"), TEXT("RampSideFalloff"), RampSideFalloff, GEditorPerProjectIni);
-
-	GConfig->SetBool(TEXT("LandscapeEdit"), TEXT("bCombinedLayersOperation"), bCombinedLayersOperation, GEditorPerProjectIni);
-	GConfig->SetBool(TEXT("LandscapeEdit"), TEXT("bApplyWithoutMovingSculpt"), bApplyWithoutMovingSculpt, GEditorPerProjectIni);
-	GConfig->SetBool(TEXT("LandscapeEdit"), TEXT("bApplyWithoutMovingPaint"), bApplyWithoutMovingPaint, GEditorPerProjectIni);
-
-	GConfig->SetInt(TEXT("LandscapeEdit"), TEXT("ErodeThresh"), ErodeThresh, GEditorPerProjectIni);
-	GConfig->SetInt(TEXT("LandscapeEdit"), TEXT("ErodeIterationNum"), ErodeIterationNum, GEditorPerProjectIni);
-	GConfig->SetInt(TEXT("LandscapeEdit"), TEXT("ErodeSurfaceThickness"), ErodeSurfaceThickness, GEditorPerProjectIni);
-	GConfig->SetInt(TEXT("LandscapeEdit"), TEXT("ErosionNoiseMode"), (int32)ErosionNoiseMode, GEditorPerProjectIni);
-	GConfig->SetFloat(TEXT("LandscapeEdit"), TEXT("ErosionNoiseScale"), ErosionNoiseScale, GEditorPerProjectIni);
-	GConfig->SetBool(TEXT("LandscapeEdit"), TEXT("bErosionUseLayerHardness"), bErosionUseLayerHardness, GEditorPerProjectIni);
-
-	GConfig->SetInt(TEXT("LandscapeEdit"), TEXT("RainAmount"), RainAmount, GEditorPerProjectIni);
-	GConfig->SetFloat(TEXT("LandscapeEdit"), TEXT("SedimentCapacity"), SedimentCapacity, GEditorPerProjectIni);
-	GConfig->SetInt(TEXT("LandscapeEdit"), TEXT("HErodeIterationNum"), ErodeIterationNum, GEditorPerProjectIni);
-	GConfig->SetInt(TEXT("LandscapeEdit"), TEXT("RainDistMode"), (int32)RainDistMode, GEditorPerProjectIni);
-	GConfig->SetFloat(TEXT("LandscapeEdit"), TEXT("RainDistScale"), RainDistScale, GEditorPerProjectIni);
-	GConfig->SetFloat(TEXT("LandscapeEdit"), TEXT("HErosionDetailScale"), HErosionDetailScale, GEditorPerProjectIni);
-	GConfig->SetBool(TEXT("LandscapeEdit"), TEXT("bHErosionDetailSmooth"), bHErosionDetailSmooth, GEditorPerProjectIni);
-
-	GConfig->SetInt(TEXT("LandscapeEdit"), TEXT("NoiseMode"), (int32)NoiseMode, GEditorPerProjectIni);
-	GConfig->SetFloat(TEXT("LandscapeEdit"), TEXT("NoiseScale"), NoiseScale, GEditorPerProjectIni);
-	GConfig->SetInt(TEXT("LandscapeEdit"), TEXT("SmoothFilterKernelSize"), SmoothFilterKernelSize, GEditorPerProjectIni);
-	GConfig->SetFloat(TEXT("LandscapeEdit"), TEXT("DetailScale"), DetailScale, GEditorPerProjectIni);
-	GConfig->SetBool(TEXT("LandscapeEdit"), TEXT("bDetailSmooth"), bDetailSmooth, GEditorPerProjectIni);
-
-	GConfig->SetFloat(TEXT("LandscapeEdit"), TEXT("MaximumValueRadius"), MaximumValueRadius, GEditorPerProjectIni);
-
-	GConfig->SetBool(TEXT("LandscapeEdit"), TEXT("bSmoothGizmoBrush"), bSmoothGizmoBrush, GEditorPerProjectIni);
-	GConfig->SetInt(TEXT("LandscapeEdit"), TEXT("PasteMode"), (int32)PasteMode, GEditorPerProjectIni);
-
-	GConfig->SetInt(TEXT("LandscapeEdit"), TEXT("MirrorOp"), (int32)MirrorOp, GEditorPerProjectIni);
-
-	GConfig->SetInt(TEXT("LandscapeEdit"), TEXT("ConvertMode"), (int32)ResizeLandscape_ConvertMode, GEditorPerProjectIni);
-	//GConfig->SetBool(TEXT("LandscapeEdit"), TEXT("bUseSelectedRegion"), bUseSelectedRegion, GEditorPerProjectIni);
-	//GConfig->SetBool(TEXT("LandscapeEdit"), TEXT("bUseNegativeMask"), bUseNegativeMask, GEditorPerProjectIni);
-	GConfig->SetBool(TEXT("LandscapeEdit"), TEXT("bApplyToAllTargets"), bApplyToAllTargets, GEditorPerProjectIni);
+	GConfig->SetInt(TEXT("LandscapeEdit"), *GET_MEMBER_NAME_CHECKED(ULandscapeEditorObject, AlphaTextureChannel).ToString(), (int32)AlphaTextureChannel, GEditorPerProjectIni);
+
+	GConfig->SetInt(TEXT("LandscapeEdit"), *GET_MEMBER_NAME_CHECKED(ULandscapeEditorObject, FlattenMode).ToString(), (int32)FlattenMode, GEditorPerProjectIni);
+	GConfig->SetBool(TEXT("LandscapeEdit"), *GET_MEMBER_NAME_CHECKED(ULandscapeEditorObject, bUseSlopeFlatten).ToString(), bUseSlopeFlatten, GEditorPerProjectIni);
+	GConfig->SetBool(TEXT("LandscapeEdit"), *GET_MEMBER_NAME_CHECKED(ULandscapeEditorObject, bPickValuePerApply).ToString(), bPickValuePerApply, GEditorPerProjectIni);
+	GConfig->SetBool(TEXT("LandscapeEdit"), *GET_MEMBER_NAME_CHECKED(ULandscapeEditorObject, bUseFlattenTarget).ToString(), bUseFlattenTarget, GEditorPerProjectIni);
+	GConfig->SetFloat(TEXT("LandscapeEdit"), *GET_MEMBER_NAME_CHECKED(ULandscapeEditorObject, FlattenTarget).ToString(), FlattenTarget, GEditorPerProjectIni);
+	GConfig->SetBool(TEXT("LandscapeEdit"), *GET_MEMBER_NAME_CHECKED(ULandscapeEditorObject, bShowFlattenTargetPreview).ToString(), bShowFlattenTargetPreview, GEditorPerProjectIni);
+	GConfig->SetFloat(TEXT("LandscapeEdit"), *GET_MEMBER_NAME_CHECKED(ULandscapeEditorObject, TerraceSmooth).ToString(), TerraceSmooth, GEditorPerProjectIni);
+	GConfig->SetFloat(TEXT("LandscapeEdit"), *GET_MEMBER_NAME_CHECKED(ULandscapeEditorObject, TerraceInterval).ToString(), TerraceInterval, GEditorPerProjectIni);
+
+	GConfig->SetFloat(TEXT("LandscapeEdit"), *GET_MEMBER_NAME_CHECKED(ULandscapeEditorObject, RampWidth).ToString(), RampWidth, GEditorPerProjectIni);
+	GConfig->SetFloat(TEXT("LandscapeEdit"), *GET_MEMBER_NAME_CHECKED(ULandscapeEditorObject, RampSideFalloff).ToString(), RampSideFalloff, GEditorPerProjectIni);
+
+	GConfig->SetBool(TEXT("LandscapeEdit"), *GET_MEMBER_NAME_CHECKED(ULandscapeEditorObject, bCombinedLayersOperation).ToString(), bCombinedLayersOperation, GEditorPerProjectIni);
+	GConfig->SetBool(TEXT("LandscapeEdit"), *GET_MEMBER_NAME_CHECKED(ULandscapeEditorObject, bApplyWithoutMovingSculpt).ToString(), bApplyWithoutMovingSculpt, GEditorPerProjectIni);
+	GConfig->SetBool(TEXT("LandscapeEdit"), *GET_MEMBER_NAME_CHECKED(ULandscapeEditorObject, bApplyWithoutMovingPaint).ToString(), bApplyWithoutMovingPaint, GEditorPerProjectIni);
+
+	GConfig->SetInt(TEXT("LandscapeEdit"), *GET_MEMBER_NAME_CHECKED(ULandscapeEditorObject, ErodeThresh).ToString(), ErodeThresh, GEditorPerProjectIni);
+	GConfig->SetInt(TEXT("LandscapeEdit"), *GET_MEMBER_NAME_CHECKED(ULandscapeEditorObject, ErodeIterationNum).ToString(), ErodeIterationNum, GEditorPerProjectIni);
+	GConfig->SetInt(TEXT("LandscapeEdit"), *GET_MEMBER_NAME_CHECKED(ULandscapeEditorObject, ErodeSurfaceThickness).ToString(), ErodeSurfaceThickness, GEditorPerProjectIni);
+	GConfig->SetInt(TEXT("LandscapeEdit"), *GET_MEMBER_NAME_CHECKED(ULandscapeEditorObject, ErosionNoiseMode).ToString(), (int32)ErosionNoiseMode, GEditorPerProjectIni);
+	GConfig->SetFloat(TEXT("LandscapeEdit"), *GET_MEMBER_NAME_CHECKED(ULandscapeEditorObject, ErosionNoiseScale).ToString(), ErosionNoiseScale, GEditorPerProjectIni);
+	GConfig->SetBool(TEXT("LandscapeEdit"), *GET_MEMBER_NAME_CHECKED(ULandscapeEditorObject, bErosionUseLayerHardness).ToString(), bErosionUseLayerHardness, GEditorPerProjectIni);
+
+	GConfig->SetInt(TEXT("LandscapeEdit"), *GET_MEMBER_NAME_CHECKED(ULandscapeEditorObject, RainAmount).ToString(), RainAmount, GEditorPerProjectIni);
+	GConfig->SetFloat(TEXT("LandscapeEdit"), *GET_MEMBER_NAME_CHECKED(ULandscapeEditorObject, SedimentCapacity).ToString(), SedimentCapacity, GEditorPerProjectIni);
+	GConfig->SetInt(TEXT("LandscapeEdit"), *GET_MEMBER_NAME_CHECKED(ULandscapeEditorObject, HErodeIterationNum).ToString(), HErodeIterationNum, GEditorPerProjectIni);
+	GConfig->SetInt(TEXT("LandscapeEdit"), *GET_MEMBER_NAME_CHECKED(ULandscapeEditorObject, RainDistMode).ToString(), (int32)RainDistMode, GEditorPerProjectIni);
+	GConfig->SetFloat(TEXT("LandscapeEdit"), *GET_MEMBER_NAME_CHECKED(ULandscapeEditorObject, RainDistScale).ToString(), RainDistScale, GEditorPerProjectIni);
+	GConfig->SetFloat(TEXT("LandscapeEdit"), *GET_MEMBER_NAME_CHECKED(ULandscapeEditorObject, HErosionDetailScale).ToString(), HErosionDetailScale, GEditorPerProjectIni);
+	GConfig->SetBool(TEXT("LandscapeEdit"), *GET_MEMBER_NAME_CHECKED(ULandscapeEditorObject, bHErosionDetailSmooth).ToString(), bHErosionDetailSmooth, GEditorPerProjectIni);
+
+	GConfig->SetInt(TEXT("LandscapeEdit"), *GET_MEMBER_NAME_CHECKED(ULandscapeEditorObject, NoiseMode).ToString(), (int32)NoiseMode, GEditorPerProjectIni);
+	GConfig->SetFloat(TEXT("LandscapeEdit"), *GET_MEMBER_NAME_CHECKED(ULandscapeEditorObject, NoiseScale).ToString(), NoiseScale, GEditorPerProjectIni);
+	GConfig->SetInt(TEXT("LandscapeEdit"), *GET_MEMBER_NAME_CHECKED(ULandscapeEditorObject, SmoothFilterKernelSize).ToString(), SmoothFilterKernelSize, GEditorPerProjectIni);
+	GConfig->SetFloat(TEXT("LandscapeEdit"), *GET_MEMBER_NAME_CHECKED(ULandscapeEditorObject, DetailScale).ToString(), DetailScale, GEditorPerProjectIni);
+	GConfig->SetBool(TEXT("LandscapeEdit"), *GET_MEMBER_NAME_CHECKED(ULandscapeEditorObject, bDetailSmooth).ToString(), bDetailSmooth, GEditorPerProjectIni);
+
+	GConfig->SetFloat(TEXT("LandscapeEdit"), *GET_MEMBER_NAME_CHECKED(ULandscapeEditorObject, MaximumValueRadius).ToString(), MaximumValueRadius, GEditorPerProjectIni);
+
+	GConfig->SetBool(TEXT("LandscapeEdit"), *GET_MEMBER_NAME_CHECKED(ULandscapeEditorObject, bSmoothGizmoBrush).ToString(), bSmoothGizmoBrush, GEditorPerProjectIni);
+	GConfig->SetInt(TEXT("LandscapeEdit"), *GET_MEMBER_NAME_CHECKED(ULandscapeEditorObject, PasteMode).ToString(), (int32)PasteMode, GEditorPerProjectIni);
+
+	GConfig->SetInt(TEXT("LandscapeEdit"), *GET_MEMBER_NAME_CHECKED(ULandscapeEditorObject, MirrorOp).ToString(), (int32)MirrorOp, GEditorPerProjectIni);
+
+	GConfig->SetInt(TEXT("LandscapeEdit"), *GET_MEMBER_NAME_CHECKED(ULandscapeEditorObject, ResizeLandscape_ConvertMode).ToString(), (int32)ResizeLandscape_ConvertMode, GEditorPerProjectIni);
+	GConfig->SetBool(TEXT("LandscapeEdit"), *GET_MEMBER_NAME_CHECKED(ULandscapeEditorObject, bApplyToAllTargets).ToString(), bApplyToAllTargets, GEditorPerProjectIni);
 
 	const FString NewLandscapeMaterialName = (NewLandscape_Material != nullptr) ? NewLandscape_Material->GetPathName() : FString();
 	GConfig->SetString(TEXT("LandscapeEdit"), TEXT("NewLandscapeMaterialName"), *NewLandscapeMaterialName, GEditorPerProjectIni);
 
-	GConfig->SetInt(TEXT("LandscapeEdit"), TEXT("ImportLandscape_AlphamapType"), (uint8)ImportLandscape_AlphamapType, GEditorPerProjectIni);
+	GConfig->SetInt(TEXT("LandscapeEdit"), *GET_MEMBER_NAME_CHECKED(ULandscapeEditorObject, ImportLandscape_AlphamapType).ToString(), (uint8)ImportLandscape_AlphamapType, GEditorPerProjectIni);
 
-	GConfig->SetBool(TEXT("LandscapeEdit"), TEXT("ShowUnusedLayers"), ShowUnusedLayers, GEditorPerProjectIni);
+	GConfig->SetBool(TEXT("LandscapeEdit"), *GET_MEMBER_NAME_CHECKED(ULandscapeEditorObject, ShowUnusedLayers).ToString(), ShowUnusedLayers, GEditorPerProjectIni);
+	GConfig->SetInt(TEXT("LandscapeEdit"), *GET_MEMBER_NAME_CHECKED(ULandscapeEditorObject, WeightBlendedTargetLayerPaintMode).ToString(), (int32)WeightBlendedTargetLayerPaintMode, GEditorPerProjectIni);
 }
 
 // Region
@@ -983,7 +984,7 @@ void ULandscapeEditorObject::AutoFillTargetLayers(const FString& DefaultTargetLa
 
 void ULandscapeEditorObject::UpdateComponentLayerAllowList()
 {
-	if (ParentMode->CurrentToolTarget.LandscapeInfo.IsValid())
+	if (ParentMode && ParentMode->CurrentToolTarget.LandscapeInfo.IsValid())
 	{
 		ParentMode->CurrentToolTarget.LandscapeInfo->UpdateComponentLayerAllowList();
 	}
diff --git a/Engine/Source/Editor/LandscapeEditor/Private/SLandscapeEditor.cpp b/Engine/Source/Editor/LandscapeEditor/Private/SLandscapeEditor.cpp
--- a/Engine/Source/Editor/LandscapeEditor/Private/SLandscapeEditor.cpp
+++ b/Engine/Source/Editor/LandscapeEditor/Private/SLandscapeEditor.cpp
@@ -49,11 +49,8 @@ void SLandscapeAssetThumbnail::Construct(const FArguments& InArgs, UObject* Asse
 
 	TSharedRef<SWidget> AssetThumbnailWidget = AssetThumbnail->MakeThumbnailWidget(AssetThumbnailConfig);
 
-	// Disable the tooltip instead of showing a "None" tooltip for invalid assets
-	if (Asset == nullptr)
-	{
-		AssetThumbnailWidget->SetToolTip(TSharedPtr<IToolTip>());
-	}
+	// Disable the asset tooltip altogether, the calling code will add its own if needed
+	AssetThumbnailWidget->SetToolTip(TSharedPtr<IToolTip>());
 
 	ChildSlot
 	[
@@ -377,7 +374,7 @@ bool FLandscapeToolKit::IsModeEnabled(FName ModeName) const
 bool FLandscapeToolKit::IsModeActive(FName ModeName) const
 {
 	FEdModeLandscape* LandscapeEdMode = GetEditorMode();
-	if (LandscapeEdMode && LandscapeEdMode->CurrentTool)
+	if (LandscapeEdMode && LandscapeEdMode->CurrentToolMode)
 	{
 		return LandscapeEdMode->CurrentToolMode->ToolModeName == ModeName;
 	}
@@ -461,7 +458,7 @@ bool FLandscapeToolKit::IsBrushSetEnabled(FName BrushSetName) const
 bool FLandscapeToolKit::IsBrushSetActive(FName BrushSetName) const
 {
 	FEdModeLandscape* LandscapeEdMode = GetEditorMode();
-	if (LandscapeEdMode && LandscapeEdMode->CurrentBrushSetIndex >= 0)
+	if (LandscapeEdMode && LandscapeEdMode->LandscapeBrushSets.IsValidIndex(LandscapeEdMode->CurrentBrushSetIndex))
 	{
 		const FName CurrentBrushSetName = LandscapeEdMode->LandscapeBrushSets[LandscapeEdMode->CurrentBrushSetIndex].BrushSetName;
 		return CurrentBrushSetName == BrushSetName;
@@ -590,6 +587,11 @@ FEdModeLandscape* SLandscapeEditor::GetEditorMode() const
 FText SLandscapeEditor::GetErrorText() const
 {
 	const FEdModeLandscape* LandscapeEdMode = GetEditorMode();
+	if (LandscapeEdMode == nullptr)
+	{
+		return FText::GetEmpty();
+	}
+
 	ELandscapeEditingState EditState = LandscapeEdMode->GetEditingState();
 	switch (EditState)
 	{
@@ -703,7 +705,7 @@ bool FLandscapeToolKit::GetIsPropertyVisibleFromProperty(const FProperty& Proper
 				return false;
 			}
 		}
-		if (Property.HasMetaData("ShowForBrushes"))
+		if (Property.HasMetaData("ShowForBrushes") && LandscapeEdMode->LandscapeBrushSets.IsValidIndex(LandscapeEdMode->CurrentBrushSetIndex))
 		{
 			const FName CurrentBrushSetName = LandscapeEdMode->LandscapeBrushSets[LandscapeEdMode->CurrentBrushSetIndex].BrushSetName;
 			// const FName CurrentBrushName = LandscapeEdMode->CurrentBrush->GetBrushName();
diff --git a/Engine/Source/Editor/LandscapeEditor/Public/LandscapeEditorObject.h b/Engine/Source/Editor/LandscapeEditor/Public/LandscapeEditorObject.h
--- a/Engine/Source/Editor/LandscapeEditor/Public/LandscapeEditorObject.h
+++ b/Engine/Source/Editor/LandscapeEditor/Public/LandscapeEditorObject.h
@@ -290,6 +290,28 @@ enum class ELandscapeImportExportMode
 	// Selected, // TODO: don.boogert   
 };
 
+UENUM()
+enum class ELandscapeWeightBlendPaintMode
+{
+	/** 
+	 * Exclusive Mode for weight-blended layers will be active only when the Ctrl key is held. Non-Exclusive Mode is used otherwise.
+	 * Note : this mode is incompatible with LandscapeEditorControlType "RequireCtrl" since the Ctrl key is then reserved for activating the landscape tools in that case. Use one of the manual modes in that case. 
+	 */
+	ExclusiveRequiresCtrl UMETA(DisplayName = "Hold Ctrl"),
+
+	/** 
+	 * Exclusive Mode for weight-blended layers will be active only when the Ctrl key is not held. Non-Exclusive Mode is used otherwise.
+	 * Note : this mode is incompatible with LandscapeEditorControlType "RequireNoCtrl" since the Ctrl key is then reserved for activating the landscape tools in that case. Use one of the manual modes in that case. 
+	 */
+	ExclusiveRequiresNoCtrl UMETA(DisplayName = "Don't Hold Ctrl"),
+
+	/** Force use Exclusive Mode for weight-blended layers. The Ctrl key is ineffective in that mode. */
+	ForceExclusive UMETA(DisplayName = "Exclusive"),
+
+	/** Force use Non-Exclusive Mode for weight-blended layers. The Ctrl key is ineffective in that mode. */
+	ForceNonExclusive UMETA(DisplayName = "Non-Exclusive"),
+};
+
 
 UCLASS(MinimalAPI)
 class ULandscapeEditorObject : public UObject
@@ -733,10 +755,18 @@ public:
 
 	// Target Layer Settings:
 
-	// Limits painting to only the components that already have the selected layer
+	/**
+	 * Limits painting to only the components that already have the selected layer.
+	 * Use the '=' key and activate the landscape brush to add the selected layer to the painted components' allow list. 
+	 * Use the '-' key and activate the landscape brush to remove the selected layer from the painted components' allow list. 
+	 */ 
 	UPROPERTY(Category = "Target Layers", EditAnywhere, NonTransactional, meta=(ShowForTargetTypes="Weightmap,Visibility"))
 	ELandscapeLayerPaintingRestriction PaintingRestriction;
 
+	/** Defines how to activate the Exclusive Paint mode for weight-blended target layers */
+	UPROPERTY(Category = "Target Layers", EditAnywhere, NonTransactional, meta = (DisplayName = "Exclusive Paint Mode", ShowForTools="Paint", ShowForTargetTypes = "Weightmap"))
+	ELandscapeWeightBlendPaintMode WeightBlendedTargetLayerPaintMode = ELandscapeWeightBlendPaintMode::ExclusiveRequiresCtrl;
+
 	UPROPERTY(Category = "Target Layers", EditAnywhere, NonTransactional, meta = (ShowForTargetTypes = "Weightmap"))
 	FLandscapeTargetLayerAssetFilePath TargetLayerAssetFilePath;
 
diff --git a/Engine/Source/Editor/LandscapeEditor/Public/LandscapeToolInterface.h b/Engine/Source/Editor/LandscapeEditor/Public/LandscapeToolInterface.h
--- a/Engine/Source/Editor/LandscapeEditor/Public/LandscapeToolInterface.h
+++ b/Engine/Source/Editor/LandscapeEditor/Public/LandscapeToolInterface.h
@@ -25,12 +25,14 @@ struct FViewportClick;
 struct FLandscapeToolInteractorPosition
 {
 	// Stored in heightmap space.
-	FVector2D Position;
-	bool bModifierPressed;
+	FVector2D Position = FVector2D::Zero();
+	bool bModifierPressed = false;
+	bool bAlternateModifierPressed = false;
 
-	FLandscapeToolInteractorPosition(FVector2D InPosition, const bool bInModifierPressed)
+	FLandscapeToolInteractorPosition(FVector2D InPosition, const bool bInModifierPressed, const bool bInAlternateModifierPressed)
 		: Position(InPosition)
 		, bModifierPressed(bInModifierPressed)
+		, bAlternateModifierPressed(bInAlternateModifierPressed)
 	{
 	}
 };
@@ -102,6 +104,26 @@ public:
 	virtual void MouseMove(float LandscapeX, float LandscapeY) = 0;
 	virtual TOptional<FVector2D> GetLastMousePosition() const { return TOptional<FVector2D>(); }
 	virtual FLandscapeBrushData ApplyBrush(const TArray<FLandscapeToolInteractorPosition>& InteractorPositions) = 0;
+
+	struct FOverlapInfo
+	{
+		FORCEINLINE explicit operator bool() const
+		{
+			return ExclusiveBounds.Area() > 0;
+		}
+
+		FORCEINLINE bool operator!() const
+		{
+			return !(bool)*this;
+		}
+
+		FIntRect ExclusiveBounds;
+		bool bOverlapsUnloadedComponents = false;
+		TSet<ULandscapeComponent*> OverlappedLoadedComponents;
+	};
+	/** @return landscape overlap information about the components being overlapped by the tool interactor positions. */
+	virtual FOverlapInfo GetOverlapInfo(const TArray<FLandscapeToolInteractorPosition>& InteractorPositions) const { return FOverlapInfo(); };
+
 	virtual TOptional<bool> InputKey(FEditorViewportClient* InViewportClient, FViewport* InViewport, FKey InKey, EInputEvent InEvent) { return TOptional<bool>(); }
 	virtual void Tick(FEditorViewportClient* ViewportClient, float DeltaTime) {};
 	virtual void BeginStroke(float LandscapeX, float LandscapeY, class FLandscapeTool* CurrentTool);
@@ -109,7 +131,7 @@ public:
 	virtual void EnterBrush() {}
 	virtual void LeaveBrush() {}
 	virtual ~FLandscapeBrush() = default;
-	virtual UMaterialInterface* GetBrushMaterial() { return NULL; }
+	virtual UMaterialInterface* GetBrushMaterial() { return nullptr; }
 	virtual const TCHAR* GetBrushName() = 0;
 	virtual FText GetDisplayName() = 0;
 	virtual ELandscapeBrushType GetBrushType() { return ELandscapeBrushType::Normal; }
@@ -172,9 +194,16 @@ struct FLandscapeToolTarget
 {
 	TWeakObjectPtr<ULandscapeInfo> LandscapeInfo;
 	ELandscapeToolTargetType TargetType;
+
+	// TODO [jonathan.bard] : eventually, will become multiple layers (multi-paint) : don't assume there will be a single one of these
 	TWeakObjectPtr<ULandscapeLayerInfoObject> LayerInfo;
 	FName LayerName;
 	
+	TArray<ULandscapeLayerInfoObject*> GetLayerInfos() const
+	{
+		return { LayerInfo.Get() };
+	}
+
 	FLandscapeToolTarget()
 		: LandscapeInfo()
 		, TargetType(ELandscapeToolTargetType::Heightmap)
@@ -190,6 +219,14 @@ enum class ELandscapeToolType
 	Mask,
 };
 
+/** Describes the type of action that a landscape tool will perform on a given target layer */
+enum class ELandscapeTargetLayerActionType
+{
+	None, // The target layer will not be affected by any action
+	Normal, // The target layer will be affected by the action (e.g. the paint action by the Paint brush)
+	Inverted, // The target layer will be inversely affected by the action (e.g. the inverse paint action by the Paint brush in exclusive mode)
+};
+
 /**
  * FLandscapeTool
  */
@@ -218,9 +255,8 @@ public:
 	virtual void SetEditRenderType();
 	virtual void Render(const FSceneView* View, FViewport* Viewport, FPrimitiveDrawInterface* PDI) {}
 	virtual bool HitTrace(const FVector& TraceStart, const FVector& TraceEnd, FVector& OutHitLocation) { return false; }
-	virtual bool UseSphereTrace() { return true; }
-	virtual bool SupportsMask() { return true; }
-	virtual bool SupportsComponentSelection() { return false; }
+	virtual bool UseSphereTrace() const { return true; }
+	virtual bool SupportsMask() const { return true; }
 	virtual bool OverrideSelection() const { return false; }
 	virtual bool IsSelectionAllowed(AActor* InActor, bool bInSelection) const { return false; }
 	virtual bool UsesTransformWidget() const { return false; }
@@ -261,8 +297,25 @@ public:
 
 	// Functions which doesn't need Viewport data...
 	virtual void Process(int32 Index, int32 Arg) {}
-	virtual ELandscapeToolType GetToolType() { return ELandscapeToolType::Normal; }
-	virtual ELandscapeToolTargetTypeFlags GetSupportedTargetTypes() { return ELandscapeToolTargetTypeFlags::None; };
+	virtual ELandscapeToolType GetToolType() const { return ELandscapeToolType::Normal; }
+	/** 
+	 * @return the type of target type that this tool supports. Some tools might support both heightmaps and weightmaps, for example, but they should be specialized 
+	 *  to return only one type (e.g. FLandscapeToolSmooth, that is templated to either FHeightmapToolTarget (ELandscapeToolTargetType::Heightmap) or FWeightmapToolTarget (ELandscapeToolTargetType::Weightmap)
+	 *  That is the reason why we return ELandscapeToolTargetType here, rather than ELandscapeToolTargetTypeFlags
+	 */
+	virtual ELandscapeToolTargetType GetSupportedTargetType() const { return ELandscapeToolTargetType::Invalid; }
+
+	/** 
+	 * Returns the type of action that this tool will perform on the target layer
+	 * @param InLayerInfo - the target layer for which to return the action type
+	 * @param bInIsModifierPressed indicates the status of the standard modifier key (Shift) : some tools might return a different action type depending on the modifier pressed
+	 * @param bInIsAlternateModifierPressed indicates the status of the alternate modifier key (Ctrl) : some tools might return a different action type depending on the modifier pressed
+	 * @return the action type 
+	 */
+	virtual ELandscapeTargetLayerActionType GetTargetLayerActionType(const ULandscapeLayerInfoObject* InLayerInfo, bool bInIsModifierPressed, bool bInIsAlternateModifierPressed) const
+	{ 
+		return ELandscapeTargetLayerActionType::None; 
+	}
 
 	// FGCObject interface
 	virtual void AddReferencedObjects(FReferenceCollector& Collector) override {}
diff --git a/Engine/Source/Runtime/Landscape/Classes/LandscapeInfo.h b/Engine/Source/Runtime/Landscape/Classes/LandscapeInfo.h
--- a/Engine/Source/Runtime/Landscape/Classes/LandscapeInfo.h
+++ b/Engine/Source/Runtime/Landscape/Classes/LandscapeInfo.h
@@ -9,6 +9,7 @@
 #include "UObject/LazyObjectPtr.h"
 #include "UObject/ScriptInterface.h"
 #include "ILandscapeSplineInterface.h"
+#include "LandscapeEditTypes.h"
 
 #if WITH_EDITOR
 #include "WorldPartition/WorldPartitionHandle.h"
@@ -24,13 +25,14 @@ class ULandscapeComponent;
 class ULandscapeLayerInfoObject;
 class ULevel;
 class UMaterialInstanceConstant;
-struct FLandscapeEditorLayerSettings;
-struct FLandscapeTargetLayerSettings;
 class ULandscapeSplinesComponent;
 class ULandscapeSplineControlPoint;
 class ULandscapeSplineSegment;
 class ULandscapeHeightfieldCollisionComponent;
+class ULandscapeInfo;
 class FModulateAlpha;
+struct FLandscapeEditorLayerSettings;
+struct FLandscapeTargetLayerSettings;
 
 /** Structure storing Collision for LandscapeComponent Add */
 #if WITH_EDITORONLY_DATA
@@ -45,7 +47,6 @@ struct FLandscapeAddCollision
 };
 #endif // WITH_EDITORONLY_DATA
 
-class ULandscapeInfo;
 
 #if WITH_EDITOR
 struct FLandscapeDirtyOnlyInModeScope
@@ -104,6 +105,48 @@ struct FLandscapeInfoLayerSettings
 #endif // WITH_EDITORONLY_DATA
 };
 
+USTRUCT()
+struct FTargetLayerBlendGroup
+{
+	GENERATED_USTRUCT_BODY()
+
+public:
+	friend class ULandscapeInfo;
+
+	FTargetLayerBlendGroup() = default;
+	FTargetLayerBlendGroup(const FName& InName, ELandscapeTargetLayerBlendMethod InBlendMethod, const TSet<ULandscapeLayerInfoObject*>& InTargetLayers)
+		: Name(InName)
+		, BlendMethod(InBlendMethod)
+		, TargetLayers(ObjectPtrWrap(InTargetLayers))
+	{
+	}
+
+	const FName& GetName() const
+	{
+		return Name;
+	}
+
+	ELandscapeTargetLayerBlendMethod GetBlendMethod() const
+	{
+		return BlendMethod;
+	}
+
+	const TSet<TObjectPtr<ULandscapeLayerInfoObject>>& GetTargetLayers() const
+	{
+		return TargetLayers;
+	}
+
+private:
+	UPROPERTY()
+	FName Name;
+
+	UPROPERTY()
+	ELandscapeTargetLayerBlendMethod BlendMethod = ELandscapeTargetLayerBlendMethod::Count;
+
+	UPROPERTY()
+	TSet<TObjectPtr<ULandscapeLayerInfoObject>> TargetLayers;
+};
+
 UCLASS(Transient)
 class ULandscapeInfo : public UObject
 {
@@ -440,6 +483,9 @@ public:
 	LANDSCAPE_API bool GetOverlappedComponents(const FTransform& InAreaWorldTransform, const FBox2D& InAreaExtents, TMap<FIntPoint, ULandscapeComponent*>& OutOverlappedComponents, FIntRect& OutComponentIndicesBoundingRect);
 
 #if WITH_EDITOR
+	LANDSCAPE_API TArray<FTargetLayerBlendGroup> GetTargetLayerBlendGroups() const;
+	LANDSCAPE_API TOptional<FTargetLayerBlendGroup> GetBlendGroupForTargetLayer(ULandscapeLayerInfoObject* InLayerInfo) const;
+
 	LANDSCAPE_API ALandscapeSplineActor* CreateSplineActor(const FVector& Location);
 
 	LANDSCAPE_API void ForAllSplineActors(TFunctionRef<void(TScriptInterface<ILandscapeSplineInterface>)> Fn) const;
diff --git a/Engine/Source/Runtime/Landscape/Classes/LandscapeLayerInfoObject.h b/Engine/Source/Runtime/Landscape/Classes/LandscapeLayerInfoObject.h
--- a/Engine/Source/Runtime/Landscape/Classes/LandscapeLayerInfoObject.h
+++ b/Engine/Source/Runtime/Landscape/Classes/LandscapeLayerInfoObject.h
@@ -123,10 +123,7 @@ public:
 	};
 	LANDSCAPE_API void SetBlendMethod(ELandscapeTargetLayerBlendMethod InBlendMethod, bool bInModify);
 
-	inline FName GetBlendGroup() const 
-	{	
-		return BlendGroup; 
-	}
+	LANDSCAPE_API FName GetBlendGroup() const; 
 	static FName GetBlendGroupMemberName()
 	{
 		return GET_MEMBER_NAME_CHECKED(ULandscapeLayerInfoObject, BlendGroup);
diff --git a/Engine/Source/Runtime/Landscape/Private/Landscape.cpp b/Engine/Source/Runtime/Landscape/Private/Landscape.cpp
--- a/Engine/Source/Runtime/Landscape/Private/Landscape.cpp
+++ b/Engine/Source/Runtime/Landscape/Private/Landscape.cpp
@@ -6503,6 +6503,51 @@ void ULandscapeInfo::UnregisterActor(ALandscapeProxy* Proxy)
 }
 
 #if WITH_EDITOR
+TArray<FTargetLayerBlendGroup> ULandscapeInfo::GetTargetLayerBlendGroups() const
+{
+	TArray<FTargetLayerBlendGroup> BlendGroups;
+	for (const FLandscapeInfoLayerSettings& LayerSettings : Layers)
+	{
+		if (LayerSettings.LayerInfoObj != nullptr)
+		{
+			if (LayerSettings.LayerInfoObj->GetBlendMethod() != ELandscapeTargetLayerBlendMethod::None)
+			{
+				FTargetLayerBlendGroup* BlendGroup = BlendGroups.FindByPredicate([&LayerSettings](const FTargetLayerBlendGroup& InGroup)
+					{
+						return (InGroup.GetBlendMethod() == LayerSettings.LayerInfoObj->GetBlendMethod())
+							&& (InGroup.GetName() == LayerSettings.LayerInfoObj->GetBlendGroup());
+					});
+
+				if (BlendGroup == nullptr)
+				{
+					BlendGroup = &BlendGroups.Add_GetRef(FTargetLayerBlendGroup(LayerSettings.LayerInfoObj->GetBlendGroup(), LayerSettings.LayerInfoObj->GetBlendMethod(), { LayerSettings.LayerInfoObj }));
+				}
+
+				BlendGroup->TargetLayers.Add(LayerSettings.LayerInfoObj);
+			}
+
+		}
+	}
+
+	return BlendGroups;
+}
+
+TOptional<FTargetLayerBlendGroup> ULandscapeInfo::GetBlendGroupForTargetLayer(ULandscapeLayerInfoObject* InLayerInfo) const
+{
+	TOptional<FTargetLayerBlendGroup> BlendGroup;
+
+	TArray<FTargetLayerBlendGroup> BlendGroups = GetTargetLayerBlendGroups();
+	if (FTargetLayerBlendGroup* FoundBlendGroup = BlendGroups.FindByPredicate([InLayerInfo](const FTargetLayerBlendGroup& InGroup)
+		{
+			return InGroup.GetTargetLayers().Contains(InLayerInfo);
+		}))
+	{
+		BlendGroup = *FoundBlendGroup;
+	}
+
+	return BlendGroup;
+}
+
 ALandscapeSplineActor* ULandscapeInfo::CreateSplineActor(const FVector& Location)
 {
 	check(LandscapeActor.Get());
@@ -7313,7 +7358,10 @@ void ALandscapeProxy::OnLayerInfoObjectDataChanged(const FOnLandscapeLayerInfoDa
 		LandscapeActor->RequestLayersContentUpdateForceAll(ELandscapeLayerUpdateMode::Update_All, InParams.bUserTriggered);
 	}
 
-	if (bUpdateAll && (InParams.PropertyChangedEvent.ChangeType != EPropertyChangeType::Interactive))
+	if ((InParams.PropertyChangedEvent.ChangeType != EPropertyChangeType::Interactive) 
+		&& ((MemberPropertyName == ULandscapeLayerInfoObject::GetBlendMethodMemberName())
+			|| (MemberPropertyName == ULandscapeLayerInfoObject::GetBlendGroupMemberName())
+			|| bUpdateAll))
 	{
 		// There's no real hook to attach to when a custom thumbnail is generated (or cleared) for a given asset and we need to regenerate the thumbnails when that happens.
 		//  PostEditChange is called on the layer info object when this happens, though, so we can ask for regenerating thumbnails when that happens. Not the most efficient
diff --git a/Engine/Source/Runtime/Landscape/Private/LandscapeLayerInfoObject.cpp b/Engine/Source/Runtime/Landscape/Private/LandscapeLayerInfoObject.cpp
--- a/Engine/Source/Runtime/Landscape/Private/LandscapeLayerInfoObject.cpp
+++ b/Engine/Source/Runtime/Landscape/Private/LandscapeLayerInfoObject.cpp
@@ -54,6 +54,11 @@ void ULandscapeLayerInfoObject::SetBlendMethod(ELandscapeTargetLayerBlendMethod
 #endif // WITH_EDITOR
 }
 
+FName ULandscapeLayerInfoObject::GetBlendGroup() const
+{
+	return (BlendMethod == ELandscapeTargetLayerBlendMethod::PremultipliedAlphaBlending) ? BlendGroup : NAME_None;
+}
+
 void ULandscapeLayerInfoObject::SetBlendGroup(const FName& InBlendGroup, bool bInModify)
 {
 	const bool bHasValueChanged = (InBlendGroup != BlendGroup);
diff --git a/Engine/Source/Runtime/Landscape/Private/LandscapeUtils.cpp b/Engine/Source/Runtime/Landscape/Private/LandscapeUtils.cpp
--- a/Engine/Source/Runtime/Landscape/Private/LandscapeUtils.cpp
+++ b/Engine/Source/Runtime/Landscape/Private/LandscapeUtils.cpp
@@ -65,6 +65,11 @@ void ModifyShaderCompilerEnvironmentForDebug(FShaderCompilerEnvironment& OutEnvi
 
 ELandscapeToolTargetTypeFlags GetLandscapeToolTargetTypeAsFlags(ELandscapeToolTargetType InTargetType)
 {
+	if (InTargetType == ELandscapeToolTargetType::Invalid)
+	{
+		return ELandscapeToolTargetTypeFlags::None;
+	}
+
 	uint8 TargetTypeValue = static_cast<uint8>(InTargetType);
 	check(TargetTypeValue < static_cast<uint8>(ELandscapeToolTargetType::Count));
 	return static_cast<ELandscapeToolTargetTypeFlags>(1 << TargetTypeValue);
diff --git a/Engine/Source/Runtime/Landscape/Public/LandscapeEdit.h b/Engine/Source/Runtime/Landscape/Public/LandscapeEdit.h
--- a/Engine/Source/Runtime/Landscape/Public/LandscapeEdit.h
+++ b/Engine/Source/Runtime/Landscape/Public/LandscapeEdit.h
@@ -397,41 +397,48 @@ template<bool bInUseInterp>
 struct FHeightmapAccessor
 {
 	enum { bUseInterp = bInUseInterp };
-	FHeightmapAccessor(ULandscapeInfo* InLandscapeInfo)
+	FHeightmapAccessor(ULandscapeInfo* InLandscapeInfo, ULandscapeLayerInfoObject* InLayerInfo = nullptr)
+		: LandscapeInfo(InLandscapeInfo)
+		, LandscapeEdit(InLandscapeInfo)
 	{
-		LandscapeInfo = InLandscapeInfo;
-		LandscapeEdit = new FLandscapeEditDataInterface(InLandscapeInfo);
+		// We don't care about InLayerInfo since heightmaps are not a target layer per se
+		check(InLayerInfo == nullptr);
 	}
 
+	FHeightmapAccessor(FHeightmapAccessor&& InOther) = default;
+	FHeightmapAccessor& operator = (FHeightmapAccessor&& InOther) = default;
+	FHeightmapAccessor(const FHeightmapAccessor& InOther) = delete;
+	FHeightmapAccessor& operator = (const FHeightmapAccessor& InOther) = delete;
+
 	void SetEditLayer(const FGuid& InEditLayerGUID)
 	{
-		LandscapeEdit->SetEditLayer(InEditLayerGUID);
+		LandscapeEdit.SetEditLayer(InEditLayerGUID);
 	}
 
 	FGuid GetEditLayer() const
 	{
-		return LandscapeEdit->GetEditLayer();
+		return LandscapeEdit.GetEditLayer();
 	}
 
 	// accessors
 	void GetData(int32& X1, int32& Y1, int32& X2, int32& Y2, TMap<FIntPoint, uint16>& Data)
 	{
-		LandscapeEdit->GetHeightData(X1, Y1, X2, Y2, Data);
+		LandscapeEdit.GetHeightData(X1, Y1, X2, Y2, Data);
 	}
 
 	void GetData(int32& X1, int32& Y1, int32& X2, int32& Y2, uint16* Data)
 	{
-		LandscapeEdit->GetHeightData(X1, Y1, X2, Y2, Data, 0);
+		LandscapeEdit.GetHeightData(X1, Y1, X2, Y2, Data, 0);
 	}
 
 	void GetDataFast(int32 X1, int32 Y1, int32 X2, int32 Y2, TMap<FIntPoint, uint16>& Data)
 	{
-		LandscapeEdit->GetHeightDataFast(X1, Y1, X2, Y2, Data);
+		LandscapeEdit.GetHeightDataFast(X1, Y1, X2, Y2, Data);
 	}
 
 	void GetDataFast(int32 X1, int32 Y1, int32 X2, int32 Y2, uint16* Data)
 	{
-		LandscapeEdit->GetHeightDataFast(X1, Y1, X2, Y2, Data, 0);
+		LandscapeEdit.GetHeightDataFast(X1, Y1, X2, Y2, Data, 0);
 	}
 
 	// Render the heightmap with overrides on edit layer visibility.  Does not modify persistent landscape data like FLandscapeEditDataInterface does.
@@ -450,7 +457,7 @@ struct FHeightmapAccessor
 	void SetData(int32 X1, int32 Y1, int32 X2, int32 Y2, const uint16* Data, ELandscapeLayerPaintingRestriction PaintingRestriction = ELandscapeLayerPaintingRestriction::None)
 	{
 		TSet<ULandscapeComponent*> Components;
-		if (LandscapeInfo && LandscapeEdit->GetComponentsInRegion(X1, Y1, X2, Y2, &Components))
+		if (LandscapeInfo && LandscapeEdit.GetComponentsInRegion(X1, Y1, X2, Y2, &Components))
 		{
 			// Update data
 			ChangedComponents.Append(Components);
@@ -460,27 +467,24 @@ struct FHeightmapAccessor
 				Component->RequestHeightmapUpdate();
 			}
 
-			LandscapeEdit->SetHeightData(X1, Y1, X2, Y2, Data, 0, /*InCalcNormals=*/ false);
+			LandscapeEdit.SetHeightData(X1, Y1, X2, Y2, Data, 0, /*InCalcNormals=*/ false);
 		}
 	}
 
 	void Flush()
 	{
-		LandscapeEdit->Flush();
+		LandscapeEdit.Flush();
 	}
 
 	virtual ~FHeightmapAccessor()
 	{
 		// Flush here manually so it will release the lock of the textures, as we will re lock for other things afterward
 		Flush();
-
-		delete LandscapeEdit;
-		LandscapeEdit = NULL;
 	}
 
 private:
 	ULandscapeInfo* LandscapeInfo;
-	FLandscapeEditDataInterface* LandscapeEdit;
+	FLandscapeEditDataInterface LandscapeEdit;
 	TSet<ULandscapeComponent*> ChangedComponents;
 };
 
@@ -495,7 +499,14 @@ struct TAlphamapAccessor
 		: LandscapeInfo(InLandscapeInfo)
 		, LandscapeEdit(InLandscapeInfo)
 		, LayerInfo(InLayerInfo)
-	{}
+	{
+		check(InLayerInfo != nullptr);
+	}
+
+	TAlphamapAccessor(TAlphamapAccessor&& InOther) = default;
+	TAlphamapAccessor& operator = (TAlphamapAccessor&& InOther) = default;
+	TAlphamapAccessor(const TAlphamapAccessor& InOther) = delete;
+	TAlphamapAccessor& operator = (const TAlphamapAccessor& InOther) = delete;
 
 	void GetData(int32& X1, int32& Y1, int32& X2, int32& Y2, TMap<FIntPoint, uint8>& Data)
 	{
