Landscape editor :
* Adjusted tooltip for Exclusive Paint Mode to better explain how it works
* Made "Don't Hold Ctrl" the default, since it's how it used to work in UE5.6

#Synced-CL 48911231

diff --git a/Engine/Source/Editor/LandscapeEditor/Public/LandscapeEditorObject.h b/Engine/Source/Editor/LandscapeEditor/Public/LandscapeEditorObject.h
--- a/Engine/Source/Editor/LandscapeEditor/Public/LandscapeEditorObject.h
+++ b/Engine/Source/Editor/LandscapeEditor/Public/LandscapeEditorObject.h
@@ -763,9 +763,14 @@ public:
 	UPROPERTY(Category = "Target Layers", EditAnywhere, NonTransactional, meta=(ShowForTargetTypes="Weightmap,Visibility"))
 	ELandscapeLayerPaintingRestriction PaintingRestriction;
 
-	/** Defines how to activate the Exclusive Paint mode for weight-blended target layers */
+	/*
+	 * Defines how to activate the Exclusive Paint mode for weight-blended target layers.
+	 * Exclusive Paint mode allows all weight-blended target layers from the same blend group to be painted in the same brush stroke. 
+	 * The selected target layer is painted while the unselected ones are "inversely" painted (erased). 
+	 * This allows the painted target layer to "override" the other ones instead of being merely being "blended" with them. It therefore (destructively) removes the weight from the other target layers. 
+	 */
 	UPROPERTY(Category = "Target Layers", EditAnywhere, NonTransactional, meta = (DisplayName = "Exclusive Paint Mode", ShowForTools="Paint", ShowForTargetTypes = "Weightmap"))
-	ELandscapeWeightBlendPaintMode WeightBlendedTargetLayerPaintMode = ELandscapeWeightBlendPaintMode::ExclusiveRequiresCtrl;
+	ELandscapeWeightBlendPaintMode WeightBlendedTargetLayerPaintMode = ELandscapeWeightBlendPaintMode::ExclusiveRequiresNoCtrl;
 
 	UPROPERTY(Category = "Target Layers", EditAnywhere, NonTransactional, meta = (ShowForTargetTypes = "Weightmap"))
 	FLandscapeTargetLayerAssetFilePath TargetLayerAssetFilePath;
