* Fixed potential crash when a user modifies the user settings by hand with invalid landscape enum values
* Fixed a couple of landscape settings that were not saved in the config

#Synced-CL 48954958

diff --git a/Engine/Source/Editor/LandscapeEditor/Private/LandscapeEditorObject.cpp b/Engine/Source/Editor/LandscapeEditor/Private/LandscapeEditorObject.cpp
--- a/Engine/Source/Editor/LandscapeEditor/Private/LandscapeEditorObject.cpp
+++ b/Engine/Source/Editor/LandscapeEditor/Private/LandscapeEditorObject.cpp
@@ -208,16 +208,42 @@ void ULandscapeEditorObject::PostEditChangeProperty(FPropertyChangedEvent& Prope
 	}
 }
 
+namespace UE::Landscape::Editor::Private
+{
+	template <typename EnumType>
+	EnumType GetConfigEnum(const TCHAR* InSection, const TCHAR* InKey, const EnumType& InDefaultValue, const TOptional<EnumType>& InInvalidValue = TOptional<EnumType>())
+	{
+		EnumType Result = InDefaultValue;
+		int32 ParsedValue = static_cast<int32>(Result);
+		if (GConfig->GetInt(InSection, InKey, ParsedValue, GEditorPerProjectIni))
+		{
+			UEnum* Enum = StaticEnum<EnumType>();
+			check(Enum != nullptr);
+			if (Enum->IsValidEnumValue(ParsedValue)
+				&& (!InInvalidValue.IsSet() || (static_cast<int32>(*InInvalidValue) != ParsedValue)))
+			{
+				Result = static_cast<EnumType>(ParsedValue);
+			}
+			else
+			{
+				UE_LOG(LogLandscapeTools, Error, TEXT("Invalid enum value read from config file for enum %s : %i, defaulting to : %s"), 
+					*Enum->GetName(), ParsedValue, *Enum->GetDisplayNameTextByValue(static_cast<int64>(InDefaultValue)).ToString());
+			}
+		}
+
+		return Result;
+	}
+} // namespace UE::Landscape::Editor::Private
+
 /** Load UI settings from ini file */
 void ULandscapeEditorObject::Load()
 {
+	using namespace UE::Landscape::Editor::Private;
+
 	GConfig->GetFloat(TEXT("LandscapeEdit"), *GET_MEMBER_NAME_CHECKED(ULandscapeEditorObject, ToolStrength).ToString(), ToolStrength, GEditorPerProjectIni);
 	GConfig->GetFloat(TEXT("LandscapeEdit"), *GET_MEMBER_NAME_CHECKED(ULandscapeEditorObject, PaintToolStrength).ToString(), PaintToolStrength, GEditorPerProjectIni);
 	GConfig->GetFloat(TEXT("LandscapeEdit"), *GET_MEMBER_NAME_CHECKED(ULandscapeEditorObject, WeightTargetValue).ToString(), WeightTargetValue, GEditorPerProjectIni);
-	bool InbUseWeightTargetValue = bUseWeightTargetValue;
-	GConfig->GetBool(TEXT("LandscapeEdit"), *GET_MEMBER_NAME_CHECKED(ULandscapeEditorObject, bUseWeightTargetValue).ToString(), InbUseWeightTargetValue, GEditorPerProjectIni);
-	bUseWeightTargetValue = InbUseWeightTargetValue;
-
+	GConfig->GetBool(TEXT("LandscapeEdit"), *GET_MEMBER_NAME_CHECKED(ULandscapeEditorObject, bUseWeightTargetValue).ToString(), bUseWeightTargetValue, GEditorPerProjectIni);
 	GConfig->GetFloat(TEXT("LandscapeEdit"), *GET_MEMBER_NAME_CHECKED(ULandscapeEditorObject, BrushRadius).ToString(), BrushRadius, GEditorPerProjectIni);
 	GConfig->GetFloat(TEXT("LandscapeEdit"), *GET_MEMBER_NAME_CHECKED(ULandscapeEditorObject, PaintBrushRadius).ToString(), PaintBrushRadius, GEditorPerProjectIni);
 	GConfig->GetInt(TEXT("LandscapeEdit"), *GET_MEMBER_NAME_CHECKED(ULandscapeEditorObject, BrushComponentSize).ToString(), BrushComponentSize, GEditorPerProjectIni);
@@ -225,9 +251,7 @@ void ULandscapeEditorObject::Load()
 	GConfig->GetFloat(TEXT("LandscapeEdit"), *GET_MEMBER_NAME_CHECKED(ULandscapeEditorObject, PaintBrushFalloff).ToString(), PaintBrushFalloff, GEditorPerProjectIni);
 	GConfig->GetString(TEXT("LandscapeEdit"), TEXT("TargetLayerAssetFilePath"), TargetLayerAssetFilePath.DirectoryPath.Path, GEditorPerProjectIni);
 	GConfig->GetBool(TEXT("LandscapeEdit"), TEXT("bUseTargetLayerAssetFilePath"), TargetLayerAssetFilePath.bUseAssetDirectoryPath, GEditorPerProjectIni);
-	bool InbUseClayBrush = bUseClayBrush;
-	GConfig->GetBool(TEXT("LandscapeEdit"), *GET_MEMBER_NAME_CHECKED(ULandscapeEditorObject, bUseClayBrush).ToString(), InbUseClayBrush, GEditorPerProjectIni);
-	bUseClayBrush = InbUseClayBrush;
+	GConfig->GetBool(TEXT("LandscapeEdit"), *GET_MEMBER_NAME_CHECKED(ULandscapeEditorObject, bUseClayBrush).ToString(), bUseClayBrush, GEditorPerProjectIni);
 	GConfig->GetFloat(TEXT("LandscapeEdit"), *GET_MEMBER_NAME_CHECKED(ULandscapeEditorObject, AlphaBrushScale).ToString(), AlphaBrushScale, GEditorPerProjectIni);
 	GConfig->GetBool(TEXT("LandscapeEdit"), *GET_MEMBER_NAME_CHECKED(ULandscapeEditorObject, bAlphaBrushAutoRotate).ToString(), bAlphaBrushAutoRotate, GEditorPerProjectIni);
 	GConfig->GetFloat(TEXT("LandscapeEdit"), *GET_MEMBER_NAME_CHECKED(ULandscapeEditorObject, AlphaBrushRotation).ToString(), AlphaBrushRotation, GEditorPerProjectIni);
@@ -237,11 +261,9 @@ void ULandscapeEditorObject::Load()
 	GConfig->GetVector2D(TEXT("LandscapeEdit"), TEXT("WorldSpacePatternBrushSettings.Origin"), WorldSpacePatternBrushSettings.Origin, GEditorPerProjectIni);
 	GConfig->GetBool(TEXT("LandscapeEdit"), TEXT("WorldSpacePatternBrushSettings.bCenterTextureOnOrigin"), WorldSpacePatternBrushSettings.bCenterTextureOnOrigin, GEditorPerProjectIni);
 	GConfig->GetFloat(TEXT("LandscapeEdit"), TEXT("WorldSpacePatternBrushSettings.RepeatSize"), WorldSpacePatternBrushSettings.RepeatSize, GEditorPerProjectIni);
+	AlphaTextureChannel = GetConfigEnum(TEXT("LandscapeEdit"), *GET_MEMBER_NAME_CHECKED(ULandscapeEditorObject, AlphaTextureChannel).ToString(), /*InDefaultValue = */ELandscapeTextureColorChannel::Red);
 	FString AlphaTextureName = (AlphaTexture != nullptr) ? AlphaTexture->GetPathName() : FString();
-	int32 InAlphaTextureChannel = static_cast<int32>(AlphaTextureChannel);
 	GConfig->GetString(TEXT("LandscapeEdit"), TEXT("AlphaTextureName"), AlphaTextureName, GEditorPerProjectIni);
-	GConfig->GetInt(TEXT("LandscapeEdit"), *GET_MEMBER_NAME_CHECKED(ULandscapeEditorObject, AlphaTextureChannel).ToString(), InAlphaTextureChannel, GEditorPerProjectIni);
-	AlphaTextureChannel = (ELandscapeTextureColorChannel)InAlphaTextureChannel;
 	UTexture2D* LoadedTexture = LoadObject<UTexture2D>(nullptr, *AlphaTextureName, nullptr, LOAD_NoWarn);
 	if ((LoadedTexture == nullptr) && !AlphaTextureName.IsEmpty())
 	{
@@ -249,27 +271,12 @@ void ULandscapeEditorObject::Load()
 	}
 	SetAlphaTexture(LoadedTexture, AlphaTextureChannel);
 
-	int32 InFlattenMode = (int32)ELandscapeToolFlattenMode::Both;
-	GConfig->GetInt(TEXT("LandscapeEdit"), *GET_MEMBER_NAME_CHECKED(ULandscapeEditorObject, FlattenMode).ToString(), InFlattenMode, GEditorPerProjectIni);
-	FlattenMode = (ELandscapeToolFlattenMode)InFlattenMode;
-
-	bool InbUseSlopeFlatten = bUseSlopeFlatten;
-	GConfig->GetBool(TEXT("LandscapeEdit"), *GET_MEMBER_NAME_CHECKED(ULandscapeEditorObject, bUseSlopeFlatten).ToString(), InbUseSlopeFlatten, GEditorPerProjectIni);
-	bUseSlopeFlatten = InbUseSlopeFlatten;
-
-	bool InbPickValuePerApply = bPickValuePerApply;
-	GConfig->GetBool(TEXT("LandscapeEdit"), *GET_MEMBER_NAME_CHECKED(ULandscapeEditorObject, bPickValuePerApply).ToString(), InbPickValuePerApply, GEditorPerProjectIni);
-	bPickValuePerApply = InbPickValuePerApply;
-
-	bool InbUseFlattenTarget = bUseFlattenTarget;
-	GConfig->GetBool(TEXT("LandscapeEdit"), *GET_MEMBER_NAME_CHECKED(ULandscapeEditorObject, bUseFlattenTarget).ToString(), InbUseFlattenTarget, GEditorPerProjectIni);
-	bUseFlattenTarget = InbUseFlattenTarget;
-
+	FlattenMode = GetConfigEnum(TEXT("LandscapeEdit"), *GET_MEMBER_NAME_CHECKED(ULandscapeEditorObject, FlattenMode).ToString(), /*InDefaultValue = */ELandscapeToolFlattenMode::Both, /*InInvalidValue = */{ ELandscapeToolFlattenMode::Invalid });
+	GConfig->GetBool(TEXT("LandscapeEdit"), *GET_MEMBER_NAME_CHECKED(ULandscapeEditorObject, bUseSlopeFlatten).ToString(), bUseSlopeFlatten, GEditorPerProjectIni);
+	GConfig->GetBool(TEXT("LandscapeEdit"), *GET_MEMBER_NAME_CHECKED(ULandscapeEditorObject, bPickValuePerApply).ToString(), bPickValuePerApply, GEditorPerProjectIni);
+	GConfig->GetBool(TEXT("LandscapeEdit"), *GET_MEMBER_NAME_CHECKED(ULandscapeEditorObject, bUseFlattenTarget).ToString(), bUseFlattenTarget, GEditorPerProjectIni);
 	GConfig->GetFloat(TEXT("LandscapeEdit"), *GET_MEMBER_NAME_CHECKED(ULandscapeEditorObject, FlattenTarget).ToString(), FlattenTarget, GEditorPerProjectIni);
-
-	bool InbShowFlattenTargetPreview = bShowFlattenTargetPreview;
-	GConfig->GetBool(TEXT("LandscapeEdit"), *GET_MEMBER_NAME_CHECKED(ULandscapeEditorObject, bShowFlattenTargetPreview).ToString(), InbShowFlattenTargetPreview, GEditorPerProjectIni);
-	bShowFlattenTargetPreview = InbShowFlattenTargetPreview;
+	GConfig->GetBool(TEXT("LandscapeEdit"), *GET_MEMBER_NAME_CHECKED(ULandscapeEditorObject, bShowFlattenTargetPreview).ToString(), bShowFlattenTargetPreview, GEditorPerProjectIni);
 
 	GConfig->GetFloat(TEXT("LandscapeEdit"), *GET_MEMBER_NAME_CHECKED(ULandscapeEditorObject, TerraceSmooth).ToString(), TerraceSmooth, GEditorPerProjectIni);
 	GConfig->GetFloat(TEXT("LandscapeEdit"), *GET_MEMBER_NAME_CHECKED(ULandscapeEditorObject, TerraceInterval).ToString(), TerraceInterval, GEditorPerProjectIni);
@@ -284,57 +291,37 @@ void ULandscapeEditorObject::Load()
 	GConfig->GetInt(TEXT("LandscapeEdit"), *GET_MEMBER_NAME_CHECKED(ULandscapeEditorObject, ErodeThresh).ToString(), ErodeThresh, GEditorPerProjectIni);
 	GConfig->GetInt(TEXT("LandscapeEdit"), *GET_MEMBER_NAME_CHECKED(ULandscapeEditorObject, ErodeIterationNum).ToString(), ErodeIterationNum, GEditorPerProjectIni);
 	GConfig->GetInt(TEXT("LandscapeEdit"), *GET_MEMBER_NAME_CHECKED(ULandscapeEditorObject, ErodeSurfaceThickness).ToString(), ErodeSurfaceThickness, GEditorPerProjectIni);
-	int32 InErosionNoiseMode = (int32)ErosionNoiseMode;
-	GConfig->GetInt(TEXT("LandscapeEdit"), *GET_MEMBER_NAME_CHECKED(ULandscapeEditorObject, ErosionNoiseMode).ToString(), InErosionNoiseMode, GEditorPerProjectIni);
-	ErosionNoiseMode = (ELandscapeToolErosionMode)InErosionNoiseMode;
+	ErosionNoiseMode = GetConfigEnum(TEXT("LandscapeEdit"), *GET_MEMBER_NAME_CHECKED(ULandscapeEditorObject, ErosionNoiseMode).ToString(), /*InDefaultValue = */ELandscapeToolErosionMode::Both, /*InInvalidValue = */{ ELandscapeToolErosionMode::Invalid });
 	GConfig->GetFloat(TEXT("LandscapeEdit"), *GET_MEMBER_NAME_CHECKED(ULandscapeEditorObject, ErosionNoiseScale).ToString(), ErosionNoiseScale, GEditorPerProjectIni);
 	GConfig->GetBool(TEXT("LandscapeEdit"), *GET_MEMBER_NAME_CHECKED(ULandscapeEditorObject, bErosionUseLayerHardness).ToString(), bErosionUseLayerHardness, GEditorPerProjectIni);
 
 	GConfig->GetInt(TEXT("LandscapeEdit"), *GET_MEMBER_NAME_CHECKED(ULandscapeEditorObject, RainAmount).ToString(), RainAmount, GEditorPerProjectIni);
 	GConfig->GetFloat(TEXT("LandscapeEdit"), *GET_MEMBER_NAME_CHECKED(ULandscapeEditorObject, SedimentCapacity).ToString(), SedimentCapacity, GEditorPerProjectIni);
 	GConfig->GetInt(TEXT("LandscapeEdit"), *GET_MEMBER_NAME_CHECKED(ULandscapeEditorObject, HErodeIterationNum).ToString(), HErodeIterationNum, GEditorPerProjectIni);
-	int32 InRainDistMode = (int32)RainDistMode;
-	GConfig->GetInt(TEXT("LandscapeEdit"), *GET_MEMBER_NAME_CHECKED(ULandscapeEditorObject, RainDistMode).ToString(), InRainDistMode, GEditorPerProjectIni);
-	RainDistMode = (ELandscapeToolHydroErosionMode)InRainDistMode;
+	RainDistMode = GetConfigEnum(TEXT("LandscapeEdit"), *GET_MEMBER_NAME_CHECKED(ULandscapeEditorObject, RainDistMode).ToString(), /*InDefaultValue = */ELandscapeToolHydroErosionMode::Both, /*InInvalidValue = */{ ELandscapeToolHydroErosionMode::Invalid });
 	GConfig->GetFloat(TEXT("LandscapeEdit"), *GET_MEMBER_NAME_CHECKED(ULandscapeEditorObject, RainDistScale).ToString(), RainDistScale, GEditorPerProjectIni);
 	GConfig->GetFloat(TEXT("LandscapeEdit"), *GET_MEMBER_NAME_CHECKED(ULandscapeEditorObject, HErosionDetailScale).ToString(), HErosionDetailScale, GEditorPerProjectIni);
-	bool InbHErosionDetailSmooth = bHErosionDetailSmooth;
-	GConfig->GetBool(TEXT("LandscapeEdit"), *GET_MEMBER_NAME_CHECKED(ULandscapeEditorObject, bHErosionDetailSmooth).ToString(), InbHErosionDetailSmooth, GEditorPerProjectIni);
-	bHErosionDetailSmooth = InbHErosionDetailSmooth;
+	GConfig->GetBool(TEXT("LandscapeEdit"), *GET_MEMBER_NAME_CHECKED(ULandscapeEditorObject, bHErosionDetailSmooth).ToString(), bHErosionDetailSmooth, GEditorPerProjectIni);
 
-	int32 InNoiseMode = (int32)NoiseMode;
-	GConfig->GetInt(TEXT("LandscapeEdit"), *GET_MEMBER_NAME_CHECKED(ULandscapeEditorObject, NoiseMode).ToString(), InNoiseMode, GEditorPerProjectIni);
-	NoiseMode = (ELandscapeToolNoiseMode)InNoiseMode;
+	NoiseMode = GetConfigEnum(TEXT("LandscapeEdit"), *GET_MEMBER_NAME_CHECKED(ULandscapeEditorObject, NoiseMode).ToString(), /*InDefaultValue = */ELandscapeToolNoiseMode::Both, /*InInvalidValue = */{ ELandscapeToolNoiseMode::Invalid });
 	GConfig->GetFloat(TEXT("LandscapeEdit"), *GET_MEMBER_NAME_CHECKED(ULandscapeEditorObject, NoiseScale).ToString(), NoiseScale, GEditorPerProjectIni);
 
 	GConfig->GetInt(TEXT("LandscapeEdit"), *GET_MEMBER_NAME_CHECKED(ULandscapeEditorObject, SmoothFilterKernelSize).ToString(), SmoothFilterKernelSize, GEditorPerProjectIni);
 	GConfig->GetFloat(TEXT("LandscapeEdit"), *GET_MEMBER_NAME_CHECKED(ULandscapeEditorObject, DetailScale).ToString(), DetailScale, GEditorPerProjectIni);
-	bool InbDetailSmooth = bDetailSmooth;
-	GConfig->GetBool(TEXT("LandscapeEdit"), *GET_MEMBER_NAME_CHECKED(ULandscapeEditorObject, bDetailSmooth).ToString(), InbDetailSmooth, GEditorPerProjectIni);
-	bDetailSmooth = InbDetailSmooth;
+	GConfig->GetBool(TEXT("LandscapeEdit"), *GET_MEMBER_NAME_CHECKED(ULandscapeEditorObject, bDetailSmooth).ToString(), bDetailSmooth, GEditorPerProjectIni);
 
 	GConfig->GetFloat(TEXT("LandscapeEdit"), *GET_MEMBER_NAME_CHECKED(ULandscapeEditorObject, MaximumValueRadius).ToString(), MaximumValueRadius, GEditorPerProjectIni);
 
-	bool InbSmoothGizmoBrush = bSmoothGizmoBrush;
-	GConfig->GetBool(TEXT("LandscapeEdit"), *GET_MEMBER_NAME_CHECKED(ULandscapeEditorObject, bSmoothGizmoBrush).ToString(), InbSmoothGizmoBrush, GEditorPerProjectIni);
-	bSmoothGizmoBrush = InbSmoothGizmoBrush;
+	GConfig->GetBool(TEXT("LandscapeEdit"), *GET_MEMBER_NAME_CHECKED(ULandscapeEditorObject, bSmoothGizmoBrush).ToString(), bSmoothGizmoBrush, GEditorPerProjectIni);
 
-	int32 InPasteMode = (int32)ELandscapeToolPasteMode::Both;
-	GConfig->GetInt(TEXT("LandscapeEdit"), *GET_MEMBER_NAME_CHECKED(ULandscapeEditorObject, PasteMode).ToString(), InPasteMode, GEditorPerProjectIni);
-	SetPasteMode((ELandscapeToolPasteMode)InPasteMode);
+	SetPasteMode(GetConfigEnum(TEXT("LandscapeEdit"), *GET_MEMBER_NAME_CHECKED(ULandscapeEditorObject, PasteMode).ToString(), /*InDefaultValue = */ELandscapeToolPasteMode::Both, /*InInvalidValue = */{ ELandscapeToolPasteMode::Invalid }));
 
-	int32 InMirrorOp = (int32)ELandscapeMirrorOperation::MinusXToPlusX;
-	GConfig->GetInt(TEXT("LandscapeEdit"), *GET_MEMBER_NAME_CHECKED(ULandscapeEditorObject, MirrorOp).ToString(), InMirrorOp, GEditorPerProjectIni);
-	MirrorOp = (ELandscapeMirrorOperation)InMirrorOp;
+	MirrorOp = GetConfigEnum(TEXT("LandscapeEdit"), *GET_MEMBER_NAME_CHECKED(ULandscapeEditorObject, MirrorOp).ToString(), /*InDefaultValue = */ELandscapeMirrorOperation::MinusXToPlusX);
 
-	int32 InResizeLandscape_ConvertMode = (int32)ResizeLandscape_ConvertMode;
-	GConfig->GetInt(TEXT("LandscapeEdit"), *GET_MEMBER_NAME_CHECKED(ULandscapeEditorObject, ResizeLandscape_ConvertMode).ToString(), InResizeLandscape_ConvertMode, GEditorPerProjectIni);
-	ResizeLandscape_ConvertMode = (ELandscapeConvertMode)InResizeLandscape_ConvertMode;
+	ResizeLandscape_ConvertMode = GetConfigEnum(TEXT("LandscapeEdit"), *GET_MEMBER_NAME_CHECKED(ULandscapeEditorObject, ResizeLandscape_ConvertMode).ToString(), /*InDefaultValue = */ELandscapeConvertMode::Expand, /*InInvalidValue = */{ ELandscapeConvertMode::Invalid });
 
 	// Region
-	bool InbApplyToAllTargets = bApplyToAllTargets;
-	GConfig->GetBool(TEXT("LandscapeEdit"), *GET_MEMBER_NAME_CHECKED(ULandscapeEditorObject, bApplyToAllTargets).ToString(), InbApplyToAllTargets, GEditorPerProjectIni);
-	bApplyToAllTargets = InbApplyToAllTargets;
+	GConfig->GetBool(TEXT("LandscapeEdit"), *GET_MEMBER_NAME_CHECKED(ULandscapeEditorObject, bApplyToAllTargets).ToString(), bApplyToAllTargets, GEditorPerProjectIni);
 
 	GConfig->GetBool(TEXT("LandscapeEdit"), *GET_MEMBER_NAME_CHECKED(ULandscapeEditorObject, ShowUnusedLayers).ToString(), ShowUnusedLayers, GEditorPerProjectIni);
 
@@ -373,13 +360,12 @@ void ULandscapeEditorObject::Load()
 		NewLandscape_Material = LoadObject<UMaterialInterface>(nullptr, *NewLandscapeMaterialName, nullptr, LOAD_NoWarn);
 	}
 	
-	int32 InImportLandscape_AlphamapType = (uint8)ImportLandscape_AlphamapType;
-	GConfig->GetInt(TEXT("LandscapeEdit"), *GET_MEMBER_NAME_CHECKED(ULandscapeEditorObject, ImportLandscape_AlphamapType).ToString(), InImportLandscape_AlphamapType, GEditorPerProjectIni);
-	ImportLandscape_AlphamapType = (ELandscapeImportAlphamapType)InImportLandscape_AlphamapType;
+	ImportType = GetConfigEnum(TEXT("LandscapeEdit"), *GET_MEMBER_NAME_CHECKED(ULandscapeEditorObject, ImportType).ToString(), /*InDefaultValue = */ELandscapeImportTransformType::None);
+	ImportExportMode = GetConfigEnum(TEXT("LandscapeEdit"), *GET_MEMBER_NAME_CHECKED(ULandscapeEditorObject, ImportExportMode).ToString(), /*InDefaultValue = */ELandscapeImportExportMode::LoadedOnly);
+	SetGizmoSnapMode(GetConfigEnum(TEXT("LandscapeEdit"), *GET_MEMBER_NAME_CHECKED(ULandscapeEditorObject, SnapMode).ToString(), /*InDefaultValue = */ELandscapeGizmoSnapType::None));
+	ImportLandscape_AlphamapType = GetConfigEnum(TEXT("LandscapeEdit"), *GET_MEMBER_NAME_CHECKED(ULandscapeEditorObject, ImportLandscape_AlphamapType).ToString(), /*InDefaultValue = */ELandscapeImportAlphamapType::Additive);
 
-	int32 InWeightBlendedTargetLayerPaintMode = (int32)WeightBlendedTargetLayerPaintMode;
-	GConfig->GetInt(TEXT("LandscapeEdit"), *GET_MEMBER_NAME_CHECKED(ULandscapeEditorObject, WeightBlendedTargetLayerPaintMode).ToString(), InWeightBlendedTargetLayerPaintMode, GEditorPerProjectIni);
-	WeightBlendedTargetLayerPaintMode = (ELandscapeWeightBlendPaintMode)InWeightBlendedTargetLayerPaintMode;
+	WeightBlendedTargetLayerPaintMode = GetConfigEnum(TEXT("LandscapeEdit"), *GET_MEMBER_NAME_CHECKED(ULandscapeEditorObject, WeightBlendedTargetLayerPaintMode).ToString(), /*InDefaultValue = */ELandscapeWeightBlendPaintMode::ExclusiveRequiresNoCtrl);
 
 	RefreshImportLayersList();
 }
@@ -463,6 +449,9 @@ void ULandscapeEditorObject::Save()
 	const FString NewLandscapeMaterialName = (NewLandscape_Material != nullptr) ? NewLandscape_Material->GetPathName() : FString();
 	GConfig->SetString(TEXT("LandscapeEdit"), TEXT("NewLandscapeMaterialName"), *NewLandscapeMaterialName, GEditorPerProjectIni);
 
+	GConfig->SetInt(TEXT("LandscapeEdit"), *GET_MEMBER_NAME_CHECKED(ULandscapeEditorObject, ImportType).ToString(), (int8)ImportType, GEditorPerProjectIni);
+	GConfig->SetInt(TEXT("LandscapeEdit"), *GET_MEMBER_NAME_CHECKED(ULandscapeEditorObject, ImportExportMode).ToString(), (int8)ImportExportMode, GEditorPerProjectIni);
+	GConfig->SetInt(TEXT("LandscapeEdit"), *GET_MEMBER_NAME_CHECKED(ULandscapeEditorObject, SnapMode).ToString(), (int8)SnapMode, GEditorPerProjectIni);
 	GConfig->SetInt(TEXT("LandscapeEdit"), *GET_MEMBER_NAME_CHECKED(ULandscapeEditorObject, ImportLandscape_AlphamapType).ToString(), (uint8)ImportLandscape_AlphamapType, GEditorPerProjectIni);
 
 	GConfig->SetBool(TEXT("LandscapeEdit"), *GET_MEMBER_NAME_CHECKED(ULandscapeEditorObject, ShowUnusedLayers).ToString(), ShowUnusedLayers, GEditorPerProjectIni);
