Landscape Editor - Removing transactions from SetCurrentEditLayer and ToggleSplinesTool
* Only data model changes are typically tracked with transaction system

#Synced-CL 49281961

diff --git a/Engine/Source/Editor/LandscapeEditor/Private/LandscapeEditorDetailCustomization_EditLayers.cpp b/Engine/Source/Editor/LandscapeEditor/Private/LandscapeEditorDetailCustomization_EditLayers.cpp
--- a/Engine/Source/Editor/LandscapeEditor/Private/LandscapeEditorDetailCustomization_EditLayers.cpp
+++ b/Engine/Source/Editor/LandscapeEditor/Private/LandscapeEditorDetailCustomization_EditLayers.cpp
@@ -309,7 +309,6 @@ void FLandscapeEditLayerCustomizationCommon::OnLayerSelectionChanged(int32 InLay
 {
 	if (FEdModeLandscape* LandscapeEdMode = GetEditorMode(); LandscapeEdMode && LandscapeEdMode->GetSelectedEditLayerIndex() != InLayerIndex)
 	{
-		FScopedTransaction Transaction(LOCTEXT("Landscape_Layers_SetCurrentLayer", "Set Current Layer"));
 		LandscapeEdMode->SetSelectedEditLayer(InLayerIndex);
 		LandscapeEdMode->UpdateTargetList();
 	}
@@ -688,7 +687,6 @@ void FLandscapeEditLayerContextMenuCustomization_Splines::CustomizeContextMenu(U
 		{
 			if (LandscapeEdMode && InEditLayer)
 			{
-				FScopedTransaction Transaction(LOCTEXT("Landscape_Layers_ToggleSplinesTool", "ToggleSplinesTool"));
 				LandscapeEdMode->ToggleSplinesTool(InEditLayer);
 			}
 		}));
diff --git a/Engine/Source/Editor/LandscapeEditor/Private/LandscapeEditorDetailCustomization_EditLayersStack.cpp b/Engine/Source/Editor/LandscapeEditor/Private/LandscapeEditorDetailCustomization_EditLayersStack.cpp
--- a/Engine/Source/Editor/LandscapeEditor/Private/LandscapeEditorDetailCustomization_EditLayersStack.cpp
+++ b/Engine/Source/Editor/LandscapeEditor/Private/LandscapeEditorDetailCustomization_EditLayersStack.cpp
@@ -629,7 +629,6 @@ void FLandscapeEditorCustomNodeBuilder_Layers::OnLayerDoubleClick(int32 InLayerI
 {
 	if (FEdModeLandscape* LandscapeEdMode = GetEditorMode(); LandscapeEdMode && LandscapeEdMode->CurrentToolMode)
 	{
-		FScopedTransaction Transaction(LOCTEXT("Landscape_Layers_SetCurrentLayer", "Set Current Layer"));
 		const ULandscapeEditLayerBase* EditLayer = LandscapeEdMode->GetEditLayerConst(InLayerIndex);
 		LandscapeEdMode->ToggleSplinesTool(EditLayer);
 	}
@@ -738,7 +737,7 @@ TOptional<SDragAndDropVerticalBox::EItemDropZone> FLandscapeEditorCustomNodeBuil
 	TSharedPtr<FLandscapeListElementDragDropOp> DragDropOperation = DragDropEvent.GetOperationAs<FLandscapeListElementDragDropOp>();
 	const FEdModeLandscape* LandscapeEdMode = GetEditorMode();
 
-	if (DragDropOperation->IsOfType<FLandscapeBrushDragDropOp>() && DragDropOperation.IsValid() && LandscapeEdMode)
+	if (LandscapeEdMode && DragDropOperation.IsValid() && DragDropOperation->IsOfType<FLandscapeBrushDragDropOp>())
 	{
 		const int32 DestinationLayerIndex = SlotIndexToLayerIndex(SlotIndex);
 		
