Landscape Editor - Fix spline selection/navigation crash when unselecting spline points/segments
* The spline selection linear path is not invalidated after control point/segment deselection leaving the path/selection logic in an unsynced state
* Fixes an issue where Selection/Navigation controls are not automatically re-enabled when deselecting points from a non-linear path (unsupported) to a valid linear path. Previously, the selection controls only rebuilt when user added to the current selection

#Synced-CL 49627611

diff --git a/Engine/Source/Editor/LandscapeEditor/Private/LandscapeSplineSelection.cpp b/Engine/Source/Editor/LandscapeEditor/Private/LandscapeSplineSelection.cpp
--- a/Engine/Source/Editor/LandscapeEditor/Private/LandscapeSplineSelection.cpp
+++ b/Engine/Source/Editor/LandscapeEditor/Private/LandscapeSplineSelection.cpp
@@ -141,7 +141,14 @@ void ULandscapeSplineSelection::DeselectControlPoint(ULandscapeSplineControlPoin
 	ControlPoint->SetSplineSelected(false);
 
 	Modify(false);
+
+	// Update the Spline Selection state. Selection was previously disabled if multiple splines/linear paths were selected 
+	// After deselecting, update the linear path to see if the selection is now valid
 	SelectedSplineControlPoints.Remove(ControlPoint);
+	if (!SelectedSplineControlPoints.IsEmpty())
+	{
+		BuildLinearPathFromLastSelectedPointInternal();
+	}
 
 	if (EnumHasAllFlags(Flags, ESplineNavigationFlags::UpdatePropertiesWindows))
 	{
@@ -158,7 +165,14 @@ void ULandscapeSplineSelection::DeSelectSegment(ULandscapeSplineSegment* Segment
 	Segment->SetSplineSelected(false);
 
 	Modify(false);
+
+	// Update the Spline Selection state. Selection was previously disabled if multiple splines/linear paths were selected
+	// After deselecting, update the linear path to see if the selection is now valid
 	SelectedSplineSegments.Remove(Segment);
+	if (!SelectedSplineSegments.IsEmpty())
+	{
+		BuildLinearPathFromLastSelectedPointInternal();
+	}
 
 	if (EnumHasAllFlags(Flags, ESplineNavigationFlags::UpdatePropertiesWindows))
 	{
@@ -484,9 +498,14 @@ bool ULandscapeSplineSelection::HasAdjacentSegmentInLinearPath(ESplineNavigation
 	check(EnumHasAllFlags(Flags, ESplineNavigationFlags::SegmentSelectModeEnabled));
 	check(FMath::CountBits(static_cast<uint8>(Flags & ESplineNavigationFlags::SelectModeMask)) == 1);
 
-	const ULandscapeSplineSegment* AdjacentSegment = GetAdjacentSegmentInLinearPath(Flags);
+	// Early out when the current selection disables navigation/selection (GetAdjacentSegmentInLinearPath reqiures a valid linear path)
+	if (!IsSelectionValidForNavigation())
+	{
+		return false;
+	}
 
-	if (!AdjacentSegment || !IsSelectionValidForNavigation())
+	const ULandscapeSplineSegment* AdjacentSegment = GetAdjacentSegmentInLinearPath(Flags);
+	if (!AdjacentSegment)
 	{
 		return false;
 	}
@@ -503,10 +522,15 @@ bool ULandscapeSplineSelection::HasAdjacentControlPointInLinearPath(ESplineNavig
 {
 	check(EnumHasAllFlags(Flags, ESplineNavigationFlags::ControlPointSelectModeEnabled));
 	check(FMath::CountBits(static_cast<uint8>(Flags & ESplineNavigationFlags::SelectModeMask)) == 1);
+	
+	// Early out when the current selection disables navigation/selection (GetAdjacentControlPointInPath reqiures a valid linear path)
+	if (!IsSelectionValidForNavigation())
+	{
+		return false;
+	}
 
 	const ULandscapeSplineControlPoint* AdjacentPoint = GetAdjacentControlPointInPath(Flags);
-
-	if (!AdjacentPoint || !IsSelectionValidForNavigation())
+	if (!AdjacentPoint)
 	{
 		return false;
 	}
