Allow Editor.Landscape tests to run successfully in a regular local editor session.  Create test landscapes with the true default landscape material instead of the cached material from the editor ini.

#Synced-CL 50537731

#Synced-CL 50537731

diff --git a/Engine/Source/Editor/LandscapeEditor/Private/Tests/LandscapeCoreWorkflowTests.cpp b/Engine/Source/Editor/LandscapeEditor/Private/Tests/LandscapeCoreWorkflowTests.cpp
--- a/Engine/Source/Editor/LandscapeEditor/Private/Tests/LandscapeCoreWorkflowTests.cpp
+++ b/Engine/Source/Editor/LandscapeEditor/Private/Tests/LandscapeCoreWorkflowTests.cpp
@@ -20,6 +20,7 @@
 #include "LandscapeEditorDetailCustomization_NewLandscape.h"
 #include "LandscapeEditorObject.h"
 #include "LandscapeEdModeTools.h"
+#include "LandscapeSettings.h"
 #include "LevelEditor.h"
 #include "Layers/LayersSubsystem.h"
 #include "Tests/AutomationEditorCommon.h"
@@ -258,6 +259,12 @@ TEST_CLASS_WITH_FLAGS(LandscapeCoreWorkflow, "Editor.Landscape", EAutomationTest
 		LandscapeEditMode->UISettings->NewLandscape_ComponentCount.Y = 8;
 		LandscapeEditMode->UISettings->NewLandscape_ClampSize();
 
+		// The tests expect the default landscape material, not the cached material from the editor ini.
+		if (const ULandscapeSettings* LandscapeSettings = GetDefault<ULandscapeSettings>())
+		{
+			LandscapeEditMode->UISettings->NewLandscape_Material = LandscapeSettings->GetDefaultLandscapeMaterial().LoadSynchronous();
+		}
+
 		// Call the function that triggers the landscape creation
 		const TSharedPtr<FLandscapeEditorDetailCustomization_NewLandscape> Customization = MakeShareable(new FLandscapeEditorDetailCustomization_NewLandscape);
 		const FReply Reply = Customization->OnCreateButtonClicked();
