Landscape - Enforced all proxies to generate a visibility layer on load
* Fixed issue where Create Layers From Materials considered visibility layer as a "missing target layer"
* Fixed landscape patch crash when re-initializing weight patches with Affects Visibility==true when no visibility layer exists on the proxy

#Synced-CL 50650520

diff --git a/Engine/Source/Editor/LandscapeEditor/Private/LandscapeEdMode.cpp b/Engine/Source/Editor/LandscapeEditor/Private/LandscapeEdMode.cpp
--- a/Engine/Source/Editor/LandscapeEditor/Private/LandscapeEdMode.cpp
+++ b/Engine/Source/Editor/LandscapeEditor/Private/LandscapeEdMode.cpp
@@ -442,6 +442,8 @@ void FEdModeLandscape::Enter()
 				}
 			}
 		}
+
+		OnLevelsChangedDelegateHandle = GetWorld()->OnLevelsChanged().AddRaw(this, &FEdModeLandscape::HandleLevelsChanged);
 	}
 
 	OnLevelActorDeletedDelegateHandle = GEngine->OnLevelActorDeleted().AddSP(this, &FEdModeLandscape::OnLevelActorRemoved);
@@ -492,8 +494,6 @@ void FEdModeLandscape::Enter()
 	UpdateBrushList();
 
 	OnWorldChangeDelegateHandle                 = FEditorSupportDelegates::WorldChange.AddRaw(this, &FEdModeLandscape::HandleLevelsChanged);
-	OnLevelsChangedDelegateHandle				= GetWorld()->OnLevelsChanged().AddRaw(this, &FEdModeLandscape::HandleLevelsChanged);
-	OnMaterialCompilationFinishedDelegateHandle = UMaterial::OnMaterialCompilationFinished().AddRaw(this, &FEdModeLandscape::OnMaterialCompilationFinished);
 
 	if (CurrentToolTarget.LandscapeInfo.IsValid() && World != nullptr)
 	{
@@ -687,6 +687,8 @@ void FEdModeLandscape::Exit()
 				}
 			}
 		}
+
+		World->OnLevelsChanged().Remove(OnLevelsChangedDelegateHandle);
 	}
 
 	FLevelEditorModule& LevelEditorModule = FModuleManager::GetModuleChecked<FLevelEditorModule>("LevelEditor");
@@ -699,8 +701,6 @@ void FEdModeLandscape::Exit()
 	GEngine->OnLevelActorAdded().Remove(OnLevelActorAddedDelegateHandle);
 	
 	FEditorSupportDelegates::WorldChange.Remove(OnWorldChangeDelegateHandle);
-	GetWorld()->OnLevelsChanged().Remove(OnLevelsChangedDelegateHandle);
-	UMaterial::OnMaterialCompilationFinished().Remove(OnMaterialCompilationFinishedDelegateHandle);
 
 	if (CurrentToolTarget.LandscapeInfo.IsValid() && World != nullptr)
 	{
@@ -2657,7 +2657,7 @@ void FEdModeLandscape::UpdateTargetList(bool bRegenerateThumbnails)
 			TSet<FName> LayerNames;
 			CurrentToolTarget.LandscapeInfo->ForEachLandscapeProxy([&](ALandscapeProxy* Proxy)
 				{
-					LayerNames.Append(Proxy->RetrieveTargetLayerNamesFromMaterials());
+					LayerNames.Append(Proxy->RetrieveTargetLayerNamesFromMaterials(/*bIncludeVisibilityLayer =*/false));
 					return true;
 				});
 
@@ -3075,19 +3075,6 @@ void FEdModeLandscape::HandleLevelsChanged()
 	}
 }
 
-void FEdModeLandscape::OnMaterialCompilationFinished(UMaterialInterface* MaterialInterface)
-{
-	// TODO [jonathan.bard] : we should remove this now that the target layer list is not automatically filled anymore
-	if (CurrentToolTarget.LandscapeInfo.IsValid() &&
-		CurrentToolTarget.LandscapeInfo->GetLandscapeProxy() != nullptr &&
-		CurrentToolTarget.LandscapeInfo->GetLandscapeProxy()->GetLandscapeMaterial() != nullptr &&
-		CurrentToolTarget.LandscapeInfo->GetLandscapeProxy()->GetLandscapeMaterial()->IsDependent(MaterialInterface))
-	{
-		CurrentToolTarget.LandscapeInfo->UpdateLayerInfoMap();
-		UpdateTargetList();
-	}
-}
-
 /** FEdMode: Render the mesh paint tool */
 void FEdModeLandscape::Render(const FSceneView* View, FViewport* Viewport, FPrimitiveDrawInterface* PDI)
 {
diff --git a/Engine/Source/Editor/LandscapeEditor/Private/LandscapeEdMode.h b/Engine/Source/Editor/LandscapeEditor/Private/LandscapeEdMode.h
--- a/Engine/Source/Editor/LandscapeEditor/Private/LandscapeEdMode.h
+++ b/Engine/Source/Editor/LandscapeEditor/Private/LandscapeEdMode.h
@@ -591,8 +591,6 @@ public:
 	/** Handle notification that visible levels may have changed and we should update the editable landscapes list */
 	void HandleLevelsChanged();
 
-	void OnMaterialCompilationFinished(UMaterialInterface* MaterialInterface);
-
 	void ReimportData(const FLandscapeTargetListInfo& TargetInfo);
 	void ImportData(const FLandscapeTargetListInfo& TargetInfo, const FString& Filename);
 	void ImportHeightData(ULandscapeInfo* LandscapeInfo, const FGuid& LayerGuid, const FString& Filename, const FIntRect& ImportRegionVerts, ELandscapeImportTransformType TransformType = ELandscapeImportTransformType::ExpandCentered, FIntPoint Offset = FIntPoint(0,0), const ELandscapeLayerPaintingRestriction& PaintRestriction = ELandscapeLayerPaintingRestriction::None, bool bFlipYAxis = false);
@@ -663,7 +661,6 @@ private:
 
 	FDelegateHandle OnWorldChangeDelegateHandle;
 	FDelegateHandle OnLevelsChangedDelegateHandle;
-	FDelegateHandle OnMaterialCompilationFinishedDelegateHandle;
 
 	FDelegateHandle OnLevelActorDeletedDelegateHandle;
 	FDelegateHandle OnLevelActorAddedDelegateHandle;
diff --git a/Engine/Source/Editor/LandscapeEditor/Private/LandscapeEditorDetailCustomization_TargetLayers.cpp b/Engine/Source/Editor/LandscapeEditor/Private/LandscapeEditorDetailCustomization_TargetLayers.cpp
--- a/Engine/Source/Editor/LandscapeEditor/Private/LandscapeEditorDetailCustomization_TargetLayers.cpp
+++ b/Engine/Source/Editor/LandscapeEditor/Private/LandscapeEditorDetailCustomization_TargetLayers.cpp
@@ -560,7 +560,7 @@ const TArray<const TSharedRef<FLandscapeTargetListInfo>> FLandscapeEditorCustomN
 		TSet<FName> LayerNames;
 		LandscapeActor->GetLandscapeInfo()->ForEachLandscapeProxy([&LayerNames](ALandscapeProxy* Proxy)
 			{
-				LayerNames.Append(Proxy->RetrieveTargetLayerNamesFromMaterials());
+				LayerNames.Append(Proxy->RetrieveTargetLayerNamesFromMaterials(/*bIncludeVisibilityLayer =*/false));
 				return true;
 			});
 
diff --git a/Engine/Source/Runtime/Landscape/Classes/LandscapeProxy.h b/Engine/Source/Runtime/Landscape/Classes/LandscapeProxy.h
--- a/Engine/Source/Runtime/Landscape/Classes/LandscapeProxy.h
+++ b/Engine/Source/Runtime/Landscape/Classes/LandscapeProxy.h
@@ -1331,7 +1331,7 @@ PRAGMA_ENABLE_DEPRECATION_WARNINGS
 	LANDSCAPE_API const TArray<FName> GetLayersFromMaterial() const;
 
 	// Get all Layer names used by the materials assigned the proxy and components (Proxy Material, Component, Component LOD Material and Hole Material)
-	LANDSCAPE_API TArray<FName> RetrieveTargetLayerNamesFromMaterials() const;
+	LANDSCAPE_API TArray<FName> RetrieveTargetLayerNamesFromMaterials(bool bIncludeVisibilityLayer = true) const;
 	
 	// Query all the LandscapeLayerInfo Objects from the weight layers allocated on this proxy.
 	TMap<FName, ULandscapeLayerInfoObject*> RetrieveTargetLayerInfosFromAllocations() const;
diff --git a/Engine/Source/Runtime/Landscape/Private/Landscape.cpp b/Engine/Source/Runtime/Landscape/Private/Landscape.cpp
--- a/Engine/Source/Runtime/Landscape/Private/Landscape.cpp
+++ b/Engine/Source/Runtime/Landscape/Private/Landscape.cpp
@@ -5634,6 +5634,14 @@ void ALandscapeProxy::FixupSharedData(ALandscape* Landscape, const bool bMapChec
 		SynchronizedProperties.Add(GET_MEMBER_NAME_CHECKED(ALandscapeProxy, TargetLayers));
 	}
 
+	// Silently add a visibility layer to each proxy. The cost is minimal, since components do not store weightmap info until data is written to the layer
+	// Target layers are no longer auto generated during material compilation so this guarantees the visibility layer is always found
+	const FName VisibilityLayerName = (ALandscapeProxy::VisibilityLayer != nullptr) ? ALandscapeProxy::VisibilityLayer->GetLayerName() : FName();
+	if (!VisibilityLayerName.IsNone() && !HasTargetLayer(VisibilityLayerName))
+	{
+		AddTargetLayer(VisibilityLayerName, FLandscapeTargetLayerSettings(), /*bPostEditChange =*/false);
+	}
+
 	if (bUpdated)
 	{
 		// In cases where LandscapeInfo is not fully ready yet, we forward the provided ALandscape. If ALandscape is available in LandscapeInfo, we let the object function naturally.
diff --git a/Engine/Source/Runtime/Landscape/Private/LandscapeEdit.cpp b/Engine/Source/Runtime/Landscape/Private/LandscapeEdit.cpp
--- a/Engine/Source/Runtime/Landscape/Private/LandscapeEdit.cpp
+++ b/Engine/Source/Runtime/Landscape/Private/LandscapeEdit.cpp
@@ -3084,7 +3084,7 @@ const TArray<FName> ALandscapeProxy::GetLayersFromMaterial() const
 	return RetrieveTargetLayerNamesFromMaterials();
 }
 
-TArray<FName> ALandscapeProxy::RetrieveTargetLayerNamesFromMaterials() const
+TArray<FName> ALandscapeProxy::RetrieveTargetLayerNamesFromMaterials(bool bIncludeVisibilityLayer) const
 {
 	TSet<UMaterialInterface*> LandscapeMaterials;
 	RetrieveAllLandscapeMaterials(LandscapeMaterials);
@@ -3092,7 +3092,7 @@ TArray<FName> ALandscapeProxy::RetrieveTargetLayerNamesFromMaterials() const
 	TSet<FName> LayerNames;
 	for (UMaterialInterface* Material : LandscapeMaterials)
 	{
-		LayerNames.Append(UE::Landscape::RetrieveTargetLayerNamesFromMaterial(Material));
+		LayerNames.Append(UE::Landscape::RetrieveTargetLayerNamesFromMaterial(Material, bIncludeVisibilityLayer));
 	}
 	return LayerNames.Array();
 }
diff --git a/Engine/Source/Runtime/Landscape/Private/LandscapeUtils.cpp b/Engine/Source/Runtime/Landscape/Private/LandscapeUtils.cpp
--- a/Engine/Source/Runtime/Landscape/Private/LandscapeUtils.cpp
+++ b/Engine/Source/Runtime/Landscape/Private/LandscapeUtils.cpp
@@ -571,7 +571,7 @@ bool DeleteActors(const TArray<AActor*>& InActorsToDelete, UWorld* InWorld, bool
 	return bSuccess;
 }
 
-TArray<FName> RetrieveTargetLayerNamesFromMaterial(UMaterialInterface* InMaterial)
+TArray<FName> RetrieveTargetLayerNamesFromMaterial(UMaterialInterface* InMaterial, bool bIncludeVisibilityLayer)
 {
 	TArray<FName> LayerNames;
 	if (InMaterial)
@@ -580,9 +580,11 @@ TArray<FName> RetrieveTargetLayerNamesFromMaterial(UMaterialInterface* InMateria
 		if (CachedExpressionData.EditorOnlyData)
 		{
 			LayerNames = CachedExpressionData.EditorOnlyData->LandscapeLayerNames;
-			LayerNames.RemoveAll([](const FName& Name)
+
+			const FName VisibilityLayerName = (ALandscapeProxy::VisibilityLayer != nullptr) ? ALandscapeProxy::VisibilityLayer->GetLayerName() : FName();
+			LayerNames.RemoveAll([bIncludeVisibilityLayer, &VisibilityLayerName](const FName& Name)
 				{
-					return Name.IsNone();
+					return Name.IsNone() || (!bIncludeVisibilityLayer && Name == VisibilityLayerName);
 				});
 		}
 	}
diff --git a/Engine/Source/Runtime/Landscape/Public/LandscapeUtils.h b/Engine/Source/Runtime/Landscape/Public/LandscapeUtils.h
--- a/Engine/Source/Runtime/Landscape/Public/LandscapeUtils.h
+++ b/Engine/Source/Runtime/Landscape/Public/LandscapeUtils.h
@@ -360,10 +360,11 @@ LANDSCAPE_API bool DeleteActors(const TArray<AActor*>& InActorsToDelete, UWorld*
 /**
  * Analyzes the material's material expressions and return the landscape target layer names declared in them
  * @param InMaterial the material to analyze
+ * @param bIncludeVisibilityLayer when false, removes the landscape visibility layer from the list of names returned by the material
  * 
  * @return the list of unique target layer names declared in that material
  */
-LANDSCAPE_API TArray<FName> RetrieveTargetLayerNamesFromMaterial(UMaterialInterface* InMaterial);
+LANDSCAPE_API TArray<FName> RetrieveTargetLayerNamesFromMaterial(UMaterialInterface* InMaterial, bool bIncludeVisibilityLayer = true);
 
 #endif // WITH_EDITOR
 
