Added customization for landscape grass types overrides TMap property in landscape streaming proxies : this allows to display an editable combo box, showing the user the existing grass names available on the landscape or entering his/her own grass name ahead of time
With this, the grass map remapping feature is now complete

#Synced-CL 50597675

diff --git a/Engine/Source/Editor/LandscapeEditor/Private/LandscapeEditorModule.cpp b/Engine/Source/Editor/LandscapeEditor/Private/LandscapeEditorModule.cpp
--- a/Engine/Source/Editor/LandscapeEditor/Private/LandscapeEditorModule.cpp
+++ b/Engine/Source/Editor/LandscapeEditor/Private/LandscapeEditorModule.cpp
@@ -102,6 +102,7 @@ public:
 		PropertyModule.RegisterCustomPropertyTypeLayout("GizmoImportLayer", FOnGetPropertyTypeCustomizationInstance::CreateStatic(&FLandscapeEditorStructCustomization_FGizmoImportLayer::MakeInstance));
 		PropertyModule.RegisterCustomPropertyTypeLayout("LandscapeImportLayer", FOnGetPropertyTypeCustomizationInstance::CreateStatic(&FLandscapeEditorStructCustomization_FLandscapeImportLayer::MakeInstance));
 		PropertyModule.RegisterCustomPropertyTypeLayout("LandscapeTargetLayerAssetFilePath", FOnGetPropertyTypeCustomizationInstance::CreateStatic(&FLandscapeEditorStructCustomization_FTargetLayerAssetPath::MakeInstance));
+		PropertyModule.RegisterCustomPropertyTypeLayout("LandscapeGrassName", FOnGetPropertyTypeCustomizationInstance::CreateStatic(&FLandscapeGrassNameCustomization::MakeInstance));
 
 		PropertyModule.RegisterCustomClassLayout("LandscapeSplineControlPoint", FOnGetDetailCustomizationInstance::CreateStatic(&FLandscapeSplineDetails::MakeInstance));
 		PropertyModule.RegisterCustomClassLayout("LandscapeSplineSegment", FOnGetDetailCustomizationInstance::CreateStatic(&FLandscapeSplineDetails::MakeInstance));
@@ -201,6 +202,7 @@ public:
 		PropertyModule.UnregisterCustomPropertyTypeLayout("GizmoImportLayer");
 		PropertyModule.UnregisterCustomPropertyTypeLayout("LandscapeImportLayer");
 		PropertyModule.UnregisterCustomPropertyTypeLayout("LandscapeTargetLayerAssetFilePath");
+		PropertyModule.UnregisterCustomPropertyTypeLayout("LandscapeGrassName");
 
 		PropertyModule.UnregisterCustomClassLayout("LandscapeSplineControlPoint");
 		PropertyModule.UnregisterCustomClassLayout("LandscapeSplineSegment");
diff --git a/Engine/Source/Editor/LandscapeEditor/Private/LandscapeGrassTypeDetails.cpp b/Engine/Source/Editor/LandscapeEditor/Private/LandscapeGrassTypeDetails.cpp
--- a/Engine/Source/Editor/LandscapeEditor/Private/LandscapeGrassTypeDetails.cpp
+++ b/Engine/Source/Editor/LandscapeEditor/Private/LandscapeGrassTypeDetails.cpp
@@ -18,6 +18,13 @@
 #include "UObject/UObjectGlobals.h"
 #include "UObject/UnrealType.h"
 #include "LandscapeGrassType.h"
+#include "IDetailChildrenBuilder.h"
+#include "PropertyCustomizationHelpers.h"
+#include "LandscapeProxy.h"
+#include "MaterialCachedData.h"
+#include "Widgets/Input/SComboBox.h"
+#include "Widgets/Input/SEditableTextBox.h"
+#include "Widgets/Text/STextBlock.h"
 
 class UObject;
 
@@ -99,4 +106,119 @@ void FLandscapeGrassTypeDetails::CustomizeDetails( IDetailLayoutBuilder& DetailL
 		}
 	}
 }
+
+
+// ----------------------------------------------------------------------------------
+
+TSharedRef<IPropertyTypeCustomization> FLandscapeGrassNameCustomization::MakeInstance()
+{
+	return MakeShareable(new FLandscapeGrassNameCustomization);
+}
+
+// This customization allows to show an "editable" combo box, that can be used as a key in a TMap property
+void FLandscapeGrassNameCustomization::CustomizeHeader(TSharedRef<IPropertyHandle> StructPropertyHandle, class FDetailWidgetRow& HeaderRow, IPropertyTypeCustomizationUtils& StructCustomizationUtils)
+{
+	// Grab the property handle for the Name in order to retrieve its value
+	TSharedPtr<IPropertyHandle> PropertyHandle = StructPropertyHandle->GetChildHandle(GET_MEMBER_NAME_CHECKED(FLandscapeGrassName, Name));
+	check(PropertyHandle.IsValid());
+	TSharedPtr<IPropertyHandle> ParentPropertyHandle = PropertyHandle->GetParentHandle();
+
+	// Gather the grass names from all materials that might be used by the landscape proxies owning this property :
+	TArray<UObject*> CustomizedObjects;
+	StructPropertyHandle->GetOuterObjects(CustomizedObjects);
+
+	TSet<UMaterialInterface*> LandscapeMaterials;
+	for (UObject* CustomizedObject : CustomizedObjects)
+	{
+		if (ALandscapeProxy* Proxy = Cast<ALandscapeProxy>(CustomizedObject))
+		{
+			Proxy->RetrieveAllLandscapeMaterials(LandscapeMaterials);
+		}
+	}
+
+	TSet<FName> AvailableGrassNames;
+	for (UMaterialInterface* LandscapeMaterial : LandscapeMaterials)
+	{
+		for (TPair<FName, TObjectPtr<ULandscapeGrassType>> NamedGrassType : LandscapeMaterial->GetCachedExpressionData().NamedGrassTypes)
+		{
+			if (NamedGrassType.Key != NAME_None)
+			{
+				AvailableGrassNames.Add(NamedGrassType.Key);
+			}
+		}
+	}
+
+	ComboBoxOptions = AvailableGrassNames.Array();
+	ComboBoxOptions.Sort(FNameLexicalLess());
+
+	auto TextToName = [](const FText& InText) -> FName
+		{
+			FString String = InText.ToString();
+			String.TrimStartAndEndInline();
+			return FName(String);
+		};
+
+	auto SetPropertyValue = [ParentPropertyHandle, PropertyHandle](const FName& InName)
+		{
+			// Prefer setting the Name's value through the parent property's handle (FLandscapeGrassName) if possible, so that the TMap's key uniqueness enforcement works as expected
+			if (ParentPropertyHandle.IsValid())
+			{
+				// Since we'll use a formatted string (Name="MyName") to address the property, we need to escape the quote (") character, so that it doesn't interfere with the string quotes 
+				FString EscapedName = InName.ToString();
+				EscapedName.ReplaceInline(TEXT("\""), TEXT("\\\""));
+				// This is a bit of a hack, but to be able to call SetValueFromFormattedString (in order to set the property's inner Name), we need to address it using this property path : 
+				ParentPropertyHandle->SetValueFromFormattedString(FString::Format(TEXT("({0}=\"{1}\")"), { GET_MEMBER_NAME_CHECKED(FLandscapeGrassName, Name).ToString(), EscapedName }));
+
+			}
+			else if (PropertyHandle.IsValid())
+			{
+				PropertyHandle->SetValueFromFormattedString(InName.ToString());
+			}
+		};
+
+	TSharedPtr<SEditableTextBox> EditableTextBox = 
+		SNew(SEditableTextBox)
+		.Text_Lambda([PropertyHandle]()
+			{
+				FString CurrentValue;
+				FPropertyAccess::Result CurrentValueResult = PropertyHandle->GetValueAsFormattedString(CurrentValue);
+				if (CurrentValueResult == FPropertyAccess::MultipleValues)
+				{
+					return LOCTEXT("MultipleValues", "Multiple Values");
+				}
+				return FText::FromString(CurrentValue);
+			})
+		.OnTextCommitted_Lambda([SetPropertyValue, TextToName/*, IsGrassNameUnique*/](const FText& InNewText, ETextCommit::Type)
+			{
+				// This might seem weird to go from FText to FString to FName and then back to FString but this allows
+				//  to transform an empty string into "None" (NAME_None) since, ultimately, this serves to set a FName
+				const FName NewName = TextToName(InNewText);
+				SetPropertyValue(NewName);
+			});
+
+	HeaderRow
+		.ValueContent()
+		[
+			SNew(SComboBox<FName>)
+			.OptionsSource(&ComboBoxOptions)
+			.OnGenerateWidget_Lambda([](FName InName)
+				{
+					return SNew(STextBlock)
+						.Text(FText::FromName(InName));
+				})
+			.OnSelectionChanged_Lambda([SetPropertyValue](FName InSelected, ESelectInfo::Type)
+				{
+					SetPropertyValue(InSelected);
+				})
+			.Content()
+			[
+				EditableTextBox.ToSharedRef()
+			]
+		];
+}
+
+// No children to display here, all the customization lies in the header 
+void FLandscapeGrassNameCustomization::CustomizeChildren(TSharedRef<IPropertyHandle> StructPropertyHandle, class IDetailChildrenBuilder& ChildBuilder, IPropertyTypeCustomizationUtils& StructCustomizationUtils)
+{
+}
 #undef LOCTEXT_NAMESPACE
diff --git a/Engine/Source/Editor/LandscapeEditor/Private/LandscapeGrassTypeDetails.h b/Engine/Source/Editor/LandscapeEditor/Private/LandscapeGrassTypeDetails.h
--- a/Engine/Source/Editor/LandscapeEditor/Private/LandscapeGrassTypeDetails.h
+++ b/Engine/Source/Editor/LandscapeEditor/Private/LandscapeGrassTypeDetails.h
@@ -4,6 +4,7 @@
 
 #include "Delegates/IDelegateInstance.h"
 #include "IDetailCustomization.h"
+#include "IPropertyTypeCustomization.h"
 #include "Templates/SharedPointer.h"
 #include "UObject/WeakObjectPtr.h"
 #include "UObject/WeakObjectPtrTemplates.h"
@@ -11,8 +12,12 @@
 class FPropertyRestriction;
 class IDetailLayoutBuilder;
 class UObject;
+class SToolTip;
 struct FPropertyChangedEvent;
 
+
+// ----------------------------------------------------------------------------------
+
 class FLandscapeGrassTypeDetails : public IDetailCustomization
 {
 public:
@@ -25,3 +30,22 @@ public:
 	virtual void CustomizeDetails( IDetailLayoutBuilder& DetailLayout ) override;
 };
 
+
+// ----------------------------------------------------------------------------------
+
+class FLandscapeGrassNameCustomization : public IPropertyTypeCustomization
+{
+public:
+	/** Makes a new instance of this detail layout class for a specific detail view requesting it */
+	static TSharedRef<IPropertyTypeCustomization> MakeInstance();
+
+	//~ Begin IPropertyTypeCustomization overrides.
+	virtual void CustomizeHeader(TSharedRef<IPropertyHandle> StructPropertyHandle, class FDetailWidgetRow& HeaderRow, IPropertyTypeCustomizationUtils& StructCustomizationUtils) override;
+
+	/** Overridden to do nothing */
+	virtual void CustomizeChildren(TSharedRef<IPropertyHandle> StructPropertyHandle, class IDetailChildrenBuilder& ChildBuilder, IPropertyTypeCustomizationUtils& StructCustomizationUtils) override;
+	// ~End IPropertyTypeCustomization overrides.
+
+private:
+	TArray<FName> ComboBoxOptions;
+};
\ No newline at end of file
diff --git a/Engine/Source/Runtime/Landscape/Classes/LandscapeProxy.h b/Engine/Source/Runtime/Landscape/Classes/LandscapeProxy.h
--- a/Engine/Source/Runtime/Landscape/Classes/LandscapeProxy.h
+++ b/Engine/Source/Runtime/Landscape/Classes/LandscapeProxy.h
@@ -84,6 +84,35 @@ DECLARE_MULTICAST_DELEGATE_TwoParams(FOnLandscapeProxyFixupSharedDataDelegate, A
 #endif // WITH_EDITOR
 
 
+// ----------------------------------------------------------------------------------
+
+// This struct is a FName wrapper.It's just a little trick to be able to customize the display of the grass name in such a way that it can work as a TMap's key.
+USTRUCT()
+struct FLandscapeGrassName
+{
+	GENERATED_USTRUCT_BODY()
+
+	FLandscapeGrassName() = default;
+	FLandscapeGrassName(const FName& InName)
+		: Name(InName)
+	{}
+
+	UPROPERTY(EditAnywhere, Category=Grass)
+	FName Name;
+
+	/** Used so we can have a TMap of this struct */
+	inline bool operator==(FLandscapeGrassName const& Other) const
+	{
+		return Name == Other.Name;
+	}
+
+	/** Used so we can have a TMap of this struct */
+	inline friend uint32 GetTypeHash(const FLandscapeGrassName& InGrassName)
+	{
+		return GetTypeHash(InGrassName.Name);
+	}
+};
+
 // ----------------------------------------------------------------------------------
 
 USTRUCT()
@@ -808,9 +837,8 @@ PRAGMA_ENABLE_DEPRECATION_WARNINGS
 	UPROPERTY(EditAnywhere, Category=Lightmass, meta = (LandscapeInherited))
 	FLightmassPrimitiveSettings LightmassSettings;
 
-	// TODO [jonathan.bard] Hidden for now because we want a UX solution to help the user choose appropriate grass names from the assigned materials. Add this when ready : EditAnywhere, Category=Grass
-	UPROPERTY(meta = (LandscapeOverridable))
-	TMap<FName, TObjectPtr<ULandscapeGrassType>> GrassTypesOverrides;
+	UPROPERTY(EditAnywhere, Category=Grass, meta = (LandscapeOverridable, ForceInlineRow))
+	TMap<FLandscapeGrassName, TObjectPtr<ULandscapeGrassType>> GrassTypesOverrides;
 
 	// Landscape LOD to use for collision tests. Higher numbers use less memory and process faster, but are much less accurate
 	UPROPERTY(EditAnywhere, Category=Collision, meta = (LandscapeOverridable))
@@ -1287,7 +1315,7 @@ PRAGMA_ENABLE_DEPRECATION_WARNINGS
 	 */
 	LANDSCAPE_API void RetrieveAllLandscapeMaterials(TSet<UMaterialInterface*>& OutMaterials) const;
 
-	LANDSCAPE_API TMap<FName, TObjectPtr<ULandscapeGrassType>> GetGrassTypeOverrides() const
+	inline TMap<FLandscapeGrassName, TObjectPtr<ULandscapeGrassType>> GetGrassTypesOverrides() const
 	{
 		return GrassTypesOverrides;
 	}
diff --git a/Engine/Source/Runtime/Landscape/Private/Landscape.cpp b/Engine/Source/Runtime/Landscape/Private/Landscape.cpp
--- a/Engine/Source/Runtime/Landscape/Private/Landscape.cpp
+++ b/Engine/Source/Runtime/Landscape/Private/Landscape.cpp
@@ -718,7 +718,7 @@ void ALandscapeProxy::ValidateGrassTypeOverrides()
 	TArray<FString> MissingGrassOverrides;
 	for (const auto& ItPair : GrassTypesOverrides)
 	{
-		FName GrassName = ItPair.Key;
+		FName GrassName = ItPair.Key.Name;
 		if (!GrassName.IsNone())
 		{
 			// If that grass name is unknown to all of our landscape materials, report : 
diff --git a/Engine/Source/Runtime/Landscape/Private/LandscapeGrass.cpp b/Engine/Source/Runtime/Landscape/Private/LandscapeGrass.cpp
--- a/Engine/Source/Runtime/Landscape/Private/LandscapeGrass.cpp
+++ b/Engine/Source/Runtime/Landscape/Private/LandscapeGrass.cpp
@@ -1051,12 +1051,12 @@ bool ULandscapeComponent::UpdateGrassTypes(bool bForceUpdate)
 		{
 			// Build the final list of grass types for this component : start by using the material's list (we won't ever be able to support anything that is not declared in that list anyway) : 
 			TMap<FName, TObjectPtr<ULandscapeGrassType>> LocalNamedGrassTypes = InMaterial->GetMaterial()->GetCachedExpressionData().NamedGrassTypes;
-			const TMap<FName, TObjectPtr<ULandscapeGrassType>>& GrassTypeOverrides = InProxy->GetGrassTypeOverrides();
+			const TMap<FLandscapeGrassName, TObjectPtr<ULandscapeGrassType>>& GrassTypesOverrides = InProxy->GetGrassTypesOverrides();
 
 			// And apply eventual overrides onto it : 
 			for (TPair<FName, TObjectPtr<ULandscapeGrassType>>& ItPair : LocalNamedGrassTypes)
 			{
-				if (const TObjectPtr<ULandscapeGrassType>* GrassTypeOverride = GrassTypeOverrides.Find(ItPair.Key))
+				if (const TObjectPtr<ULandscapeGrassType>* GrassTypeOverride = GrassTypesOverrides.Find(FLandscapeGrassName(ItPair.Key)))
 				{
 					ItPair.Value = *GrassTypeOverride;
 				}
