* UX fixes for landscape overridable properties, so that they appear properly disabled when not overridden
* Added tooltips for some landscape proxy properties

#Synced-CL 50661191

diff --git a/Engine/Source/Editor/LandscapeEditor/Private/LandscapeGrassTypeDetails.cpp b/Engine/Source/Editor/LandscapeEditor/Private/LandscapeGrassTypeDetails.cpp
--- a/Engine/Source/Editor/LandscapeEditor/Private/LandscapeGrassTypeDetails.cpp
+++ b/Engine/Source/Editor/LandscapeEditor/Private/LandscapeGrassTypeDetails.cpp
@@ -122,6 +122,8 @@ void FLandscapeGrassNameCustomization::CustomizeHeader(TSharedRef<IPropertyHandl
 	TSharedPtr<IPropertyHandle> PropertyHandle = StructPropertyHandle->GetChildHandle(GET_MEMBER_NAME_CHECKED(FLandscapeGrassName, Name));
 	check(PropertyHandle.IsValid());
 	TSharedPtr<IPropertyHandle> ParentPropertyHandle = PropertyHandle->GetParentHandle();
+	
+	const bool bIsEditable = PropertyHandle->IsEditable();
 
 	// Gather the grass names from all materials that might be used by the landscape proxies owning this property :
 	TArray<UObject*> CustomizedObjects;
@@ -214,6 +216,8 @@ void FLandscapeGrassNameCustomization::CustomizeHeader(TSharedRef<IPropertyHandl
 			[
 				EditableTextBox.ToSharedRef()
 			]
+			// Disable this custom widget by hand if the property is disabled (since this is custom, it doesn't automatically support the CanEditChange method)
+			.IsEnabled(bIsEditable)
 		];
 }
 
diff --git a/Engine/Source/Editor/LandscapeEditor/Private/LandscapeProxyUIDetails.cpp b/Engine/Source/Editor/LandscapeEditor/Private/LandscapeProxyUIDetails.cpp
--- a/Engine/Source/Editor/LandscapeEditor/Private/LandscapeProxyUIDetails.cpp
+++ b/Engine/Source/Editor/LandscapeEditor/Private/LandscapeProxyUIDetails.cpp
@@ -86,6 +86,8 @@ TSharedRef<IDetailCustomization> FLandscapeProxyUIDetails::MakeInstance()
 
 void FLandscapeProxyUIDetails::CustomizeDetails( IDetailLayoutBuilder& DetailBuilder )
 {
+	DetailLayoutBuilder = &DetailBuilder;
+
 	TArray<TWeakObjectPtr<UObject>> EditingObjects;
 	DetailBuilder.GetObjectsBeingCustomized(EditingObjects);
 
@@ -445,7 +447,7 @@ void FLandscapeProxyUIDetails::CustomizeDetails( IDetailLayoutBuilder& DetailBui
 			else if (LandscapeStreamingProxy->IsPropertyOverridable(Property))
 			{
 				TSharedPtr<IPropertyHandle> PropertyHandle = DetailBuilder.GetProperty(Property->GetFName());
-
+				
 				if (PropertyHandle->IsValidHandle())
 				{
 					const FText TooltipText = LOCTEXT("OverriddenProperty", "Check this box to override the parent landscape's property.");
@@ -453,7 +455,7 @@ void FLandscapeProxyUIDetails::CustomizeDetails( IDetailLayoutBuilder& DetailBui
 					IDetailPropertyRow* DetailRow = DetailBuilder.EditDefaultProperty(PropertyHandle);
 					TSharedPtr<SWidget> NameWidget = nullptr;
 					TSharedPtr<SWidget> ValueWidget = nullptr;
-
+					
 					DetailRow->GetDefaultWidgets(NameWidget, ValueWidget);
 
 					TAttribute<bool> EnabledAttribute = TAttribute<bool>::Create([LandscapeStreamingProxy, PropertyName]() -> bool
@@ -470,72 +472,24 @@ void FLandscapeProxyUIDetails::CustomizeDetails( IDetailLayoutBuilder& DetailBui
 						[
 							SNew(SCheckBox)
 							.ToolTipText(TooltipText)
-							.IsChecked_Lambda([EditingStreamingProxies, PropertyName]() -> ECheckBoxState
-							{
-								bool bContainsOverriddenProperty = false;
-								bool bContainsDefaultProperty = false;
-
-								for (const TWeakObjectPtr<ALandscapeStreamingProxy> StreamingProxy : EditingStreamingProxies)
-								{
-									if (!StreamingProxy.IsValid())
-									{
-										continue;
-									}
-
-									const bool bIsPropertyOverridden = StreamingProxy->IsSharedPropertyOverridden(PropertyName);
-
-									bContainsOverriddenProperty |= bIsPropertyOverridden;
-									bContainsDefaultProperty |= !bIsPropertyOverridden;
-								}
-
-								if (bContainsOverriddenProperty && bContainsDefaultProperty)
-								{
-									return ECheckBoxState::Undetermined;
-								}
-
-								return bContainsOverriddenProperty ? ECheckBoxState::Checked : ECheckBoxState::Unchecked;
-							})
-							.OnCheckStateChanged_Lambda([EditingStreamingProxies, PropertyName](ECheckBoxState NewState)
-							{
-								if (NewState == ECheckBoxState::Undetermined)
-								{
-									return;
-								}
-
-								const bool bChecked = NewState == ECheckBoxState::Checked;
-
-								const FScopedTransaction Transaction(LOCTEXT("SetSharedPropertyOverride", "Change Property Override"));
-
-								for (const TWeakObjectPtr<ALandscapeStreamingProxy> StreamingProxy : EditingStreamingProxies)
-								{
-									if (!StreamingProxy.IsValid())
-									{
-										continue;
-									}
-
-									StreamingProxy->SetSharedPropertyOverride(PropertyName, bChecked);
-								}
-							})
+							.IsChecked(this, &FLandscapeProxyUIDetails::IsOverriddenPropertyChecked, EditingStreamingProxies, PropertyName)
+							.OnCheckStateChanged(this, &FLandscapeProxyUIDetails::OnCheckStateChanged, EditingStreamingProxies, PropertyName)
 						]
 						+ SHorizontalBox::Slot()
 						.AutoWidth()
 						[
-							SNew(SBox)
-							.IsEnabled(EnabledAttribute)
-							[
-								NameWidget->AsShared()
-							]
+							NameWidget->AsShared()
 						]
 					]
 					.ValueContent()
 					[
-						SNew(SBox)
-						.IsEnabled(EnabledAttribute)
-						[
-							ValueWidget->AsShared()
-						]
-					]
-					.IsValueEnabled(EnabledAttribute);
+						ValueWidget->AsShared()
+					];
+					
+					// Hook the EnabledAttribute on the 2 sub-widgets by hand, so that we don't disable the entire custom widget
+					//  when the property is not overridden (since that would also disable the override checkbox...)
+					NameWidget->SetEnabled(EnabledAttribute);
+					ValueWidget->SetEnabled(EnabledAttribute);
 
 					// The Reset To Default button corresponds to the overridden state of the overridable property : 
 					DetailRow->OverrideResetToDefault(FResetToDefaultOverride::Create(
@@ -562,6 +516,58 @@ void FLandscapeProxyUIDetails::CustomizeDetails( IDetailLayoutBuilder& DetailBui
 	}
 }
 
+ECheckBoxState FLandscapeProxyUIDetails::IsOverriddenPropertyChecked(TArray<TWeakObjectPtr<ALandscapeStreamingProxy>> InEditingStreamingProxies, FName InPropertyName) const
+{
+	bool bContainsOverriddenProperty = false;
+	bool bContainsDefaultProperty = false;
+
+	for (const TWeakObjectPtr<ALandscapeStreamingProxy> StreamingProxy : InEditingStreamingProxies)
+	{
+		if (!StreamingProxy.IsValid())
+		{
+			continue;
+		}
+
+		const bool bIsPropertyOverridden = StreamingProxy->IsSharedPropertyOverridden(InPropertyName);
+
+		bContainsOverriddenProperty |= bIsPropertyOverridden;
+		bContainsDefaultProperty |= !bIsPropertyOverridden;
+	}
+
+	if (bContainsOverriddenProperty && bContainsDefaultProperty)
+	{
+		return ECheckBoxState::Undetermined;
+	}
+
+	return bContainsOverriddenProperty ? ECheckBoxState::Checked : ECheckBoxState::Unchecked;
+}
+
+void FLandscapeProxyUIDetails::OnCheckStateChanged(ECheckBoxState InNewState, TArray<TWeakObjectPtr<ALandscapeStreamingProxy>> InEditingStreamingProxies, FName InPropertyName)
+{
+	if (InNewState == ECheckBoxState::Undetermined)
+	{
+		return;
+	}
+
+	const bool bChecked = (InNewState == ECheckBoxState::Checked);
+
+	const FScopedTransaction Transaction(LOCTEXT("SetSharedPropertyOverride", "Change Property Override"));
+
+	for (const TWeakObjectPtr<ALandscapeStreamingProxy> StreamingProxy : InEditingStreamingProxies)
+	{
+		if (!StreamingProxy.IsValid())
+		{
+			continue;
+		}
+
+		StreamingProxy->SetSharedPropertyOverride(InPropertyName, bChecked);
+	}
+
+	check(DetailLayoutBuilder != nullptr);
+	// Refresh the details to disable/enable the widgets that should be : we need this to force CanEditChange to be re-evaluated after the property was overridden/de-overridden
+	DetailLayoutBuilder->ForceRefreshDetails();
+}
+
 static void GetMissingRuntimeVirtualTextureVolumes(ALandscape* InLandscapeActor, TArray<URuntimeVirtualTexture*>& OutVirtualTextures)
 {
 	UWorld* World = InLandscapeActor != nullptr ? InLandscapeActor->GetWorld() : nullptr;
diff --git a/Engine/Source/Editor/LandscapeEditor/Private/LandscapeProxyUIDetails.h b/Engine/Source/Editor/LandscapeEditor/Private/LandscapeProxyUIDetails.h
--- a/Engine/Source/Editor/LandscapeEditor/Private/LandscapeProxyUIDetails.h
+++ b/Engine/Source/Editor/LandscapeEditor/Private/LandscapeProxyUIDetails.h
@@ -7,7 +7,9 @@
 #include "Templates/SharedPointer.h"
 
 class ALandscape;
+class ALandscapeStreamingProxy;
 class IDetailLayoutBuilder;
+enum class ECheckBoxState : uint8;
 
 class FLandscapeProxyUIDetails : public IDetailCustomization
 {
@@ -25,6 +27,10 @@ private:
 
 	/** Returns true if SetBounds button is enabled */
 	bool IsCreateRuntimeVirtualTextureVolumeEnabled(ALandscape* InLandscapeActor) const;
+
+	/** Callbacks for the override property checkbox */
+	ECheckBoxState IsOverriddenPropertyChecked(TArray<TWeakObjectPtr<ALandscapeStreamingProxy>> InEditingStreamingProxies, FName InPropertyName) const;
+	void OnCheckStateChanged(ECheckBoxState InNewState, TArray<TWeakObjectPtr<ALandscapeStreamingProxy>> InEditingStreamingProxies, FName InPropertyName);
 	
 	/** Callback for Set Bounds button */
 	FReply CreateRuntimeVirtualTextureVolume(ALandscape* InLandscapeActor);
diff --git a/Engine/Source/Runtime/Landscape/Classes/LandscapeProxy.h b/Engine/Source/Runtime/Landscape/Classes/LandscapeProxy.h
--- a/Engine/Source/Runtime/Landscape/Classes/LandscapeProxy.h
+++ b/Engine/Source/Runtime/Landscape/Classes/LandscapeProxy.h
@@ -490,6 +490,7 @@ protected:
 	UPROPERTY(EditAnywhere, BlueprintReadOnly, Category = Nanite, meta = (LandscapeInherited))
 	bool bEnableNanite = false;
 
+	/** Specify an alternate material to use for a given LOD level (for example, to use a simpler material at a distance) This is not effective on Nanite landscape. */
 	UPROPERTY(EditAnywhere, Category = Landscape, meta = (LandscapeOverridable))
 	TArray<FLandscapePerLODMaterialOverride> PerLODOverrideMaterials;
 
@@ -837,6 +838,7 @@ PRAGMA_ENABLE_DEPRECATION_WARNINGS
 	UPROPERTY(EditAnywhere, Category=Lightmass, meta = (LandscapeInherited))
 	FLightmassPrimitiveSettings LightmassSettings;
 
+	/** Specify alternate landscape grass types for some specific grass map names. This allows to reuse a given material's grass outputs, but change the grass meshes it spawns in the context of a specific landscape actor. */
 	UPROPERTY(EditAnywhere, Category=Grass, meta = (LandscapeOverridable, ForceInlineRow))
 	TMap<FLandscapeGrassName, TObjectPtr<ULandscapeGrassType>> GrassTypesOverrides;
 
diff --git a/Engine/Source/Runtime/Landscape/Private/LandscapeEdit.cpp b/Engine/Source/Runtime/Landscape/Private/LandscapeEdit.cpp
--- a/Engine/Source/Runtime/Landscape/Private/LandscapeEdit.cpp
+++ b/Engine/Source/Runtime/Landscape/Private/LandscapeEdit.cpp
@@ -6285,6 +6285,10 @@ bool ALandscapeStreamingProxy::CanEditChange(const FProperty* InProperty) const
 			return !World->GetSubsystem<ULandscapeSubsystem>()->IsGridBased();
 		}
 	}
+	else if (InProperty && IsPropertyOverridable(InProperty))
+	{
+		return IsSharedPropertyOverridden(InProperty->GetFName());
+	}
 
 	return true;
 }
