[Editor TRS] Landscape splines update with new gizmos, via Begin/EndTransform

#Synced-CL 51038722

diff --git a/Engine/Source/Editor/Experimental/EditorInteractiveToolsFramework/Private/EditorGizmos/EditorGizmoStateTarget.cpp b/Engine/Source/Editor/Experimental/EditorInteractiveToolsFramework/Private/EditorGizmos/EditorGizmoStateTarget.cpp
--- a/Engine/Source/Editor/Experimental/EditorInteractiveToolsFramework/Private/EditorGizmos/EditorGizmoStateTarget.cpp
+++ b/Engine/Source/Editor/Experimental/EditorInteractiveToolsFramework/Private/EditorGizmos/EditorGizmoStateTarget.cpp
@@ -18,8 +18,8 @@ void UEditorGizmoStateTarget::BeginUpdate()
 			TransactionManager->BeginUndoTransaction(TransactionDescription);
 		}
 
-		// empty state for now
-		constexpr FGizmoState State;
+		FGizmoState State;
+		State.TransformMode = GetTransformMode();
 		(void)ModeTools->BeginTransform(State);
 	}
 }
@@ -35,8 +35,8 @@ void UEditorGizmoStateTarget::EndUpdate()
 	const TSharedPtr<FEditorModeTools> ModeTools = WeakModeTools.IsValid() ? WeakModeTools.Pin() : nullptr;
 	if (ensure(ModeTools))
 	{
-		// empty state for now
-		constexpr FGizmoState State;
+		FGizmoState State;
+		State.TransformMode = GetTransformMode();
 		(void)ModeTools->EndTransform(State);
 
 		if (TransactionManager)
@@ -57,8 +57,8 @@ void UEditorGizmoStateTarget::CancelUpdate()
 	const TSharedPtr<FEditorModeTools> ModeTools = WeakModeTools.IsValid() ? WeakModeTools.Pin() : nullptr;
 	if (ensure(ModeTools))
 	{
-		// empty state for now
-		constexpr FGizmoState State;
+		FGizmoState State;
+		State.TransformMode = GetTransformMode();
 		(void)ModeTools->EndTransform(State);
 
 		if (TransactionManager)
@@ -91,3 +91,14 @@ void UEditorGizmoStateTarget::SetTransformGizmo(UTransformGizmo* InGizmo)
 {
 	TransformGizmo = InGizmo;
 }
+
+EGizmoTransformMode UEditorGizmoStateTarget::GetTransformMode() const
+{
+	EGizmoTransformMode TransformMode = EGizmoTransformMode::None;
+	if (TransformGizmo.IsValid() && TransformGizmo->TransformGizmoSource)
+	{
+		TransformMode = TransformGizmo->TransformGizmoSource->GetGizmoMode();
+	}
+
+	return TransformMode;
+}
diff --git a/Engine/Source/Editor/Experimental/EditorInteractiveToolsFramework/Public/EditorGizmos/EditorGizmoStateTarget.h b/Engine/Source/Editor/Experimental/EditorInteractiveToolsFramework/Public/EditorGizmos/EditorGizmoStateTarget.h
--- a/Engine/Source/Editor/Experimental/EditorInteractiveToolsFramework/Public/EditorGizmos/EditorGizmoStateTarget.h
+++ b/Engine/Source/Editor/Experimental/EditorInteractiveToolsFramework/Public/EditorGizmos/EditorGizmoStateTarget.h
@@ -6,7 +6,7 @@
 
 #include "Templates/Function.h"
 #include "Internationalization/Text.h"
-
+#include "TransformGizmoInterfaces.h"
 #include "UObject/Object.h"
 #include "UObject/WeakObjectPtr.h"
 #include "UObject/WeakObjectPtrTemplates.h"
@@ -76,6 +76,10 @@ public:
 	 */
 	void SetTransformGizmo(UTransformGizmo* InGizmo);
 	
+protected:
+	/** Gets the current TransformMode, from the associated TransformGizmo (if available). */
+	EGizmoTransformMode GetTransformMode() const;
+	
 protected:
 	/**
 	 * The transform gizmo from which we can get data if required.
diff --git a/Engine/Source/Editor/Experimental/EditorInteractiveToolsFramework/Public/GizmoEdModeInterface.h b/Engine/Source/Editor/Experimental/EditorInteractiveToolsFramework/Public/GizmoEdModeInterface.h
--- a/Engine/Source/Editor/Experimental/EditorInteractiveToolsFramework/Public/GizmoEdModeInterface.h
+++ b/Engine/Source/Editor/Experimental/EditorInteractiveToolsFramework/Public/GizmoEdModeInterface.h
@@ -2,6 +2,7 @@
 
 #pragma once
 
+#include "EditorGizmos/TransformGizmoInterfaces.h"
 #include "UObject/Interface.h"
 
 #include "GizmoEdModeInterface.generated.h"
@@ -12,6 +13,7 @@
 
 struct FGizmoState
 {
+	EGizmoTransformMode TransformMode = EGizmoTransformMode::None;
 };
 
 /**
diff --git a/Engine/Source/Editor/LandscapeEditor/LandscapeEditor.Build.cs b/Engine/Source/Editor/LandscapeEditor/LandscapeEditor.Build.cs
--- a/Engine/Source/Editor/LandscapeEditor/LandscapeEditor.Build.cs
+++ b/Engine/Source/Editor/LandscapeEditor/LandscapeEditor.Build.cs
@@ -38,6 +38,7 @@ public class LandscapeEditor : ModuleRules
 				"CQTest",
 				"NavigationSystem",
 				"VirtualTexturingEditor",
+				"EditorInteractiveToolsFramework",
 			}
 			);
 
diff --git a/Engine/Source/Editor/LandscapeEditor/Private/LandscapeEdMode.cpp b/Engine/Source/Editor/LandscapeEditor/Private/LandscapeEdMode.cpp
--- a/Engine/Source/Editor/LandscapeEditor/Private/LandscapeEdMode.cpp
+++ b/Engine/Source/Editor/LandscapeEditor/Private/LandscapeEdMode.cpp
@@ -1178,6 +1178,36 @@ bool FEdModeLandscape::EndTracking(FEditorViewportClient* InViewportClient, FVie
 	return false;
 }
 
+bool FEdModeLandscape::BeginTransform(const FGizmoState& InState)
+{
+	if (!IsEditingEnabled())
+	{
+		return false;
+	}
+
+	if (CurrentTool && CurrentTool->BeginTransform(InState))
+	{
+		return true;
+	}
+
+	return false;
+}
+
+bool FEdModeLandscape::EndTransform(const FGizmoState& InState)
+{
+	if (!IsEditingEnabled())
+	{
+		return false;
+	}
+
+	if (CurrentTool && CurrentTool->EndTransform(InState))
+	{
+		return true;
+	}
+
+	return false;
+}
+
 /** Trace under the mouse cursor and return the landscape hit and the hit location (in landscape quad space) */
 bool FEdModeLandscape::LandscapeMouseTrace(FEditorViewportClient* ViewportClient, float& OutHitX, float& OutHitY)
 {
diff --git a/Engine/Source/Editor/LandscapeEditor/Private/LandscapeEdMode.h b/Engine/Source/Editor/LandscapeEditor/Private/LandscapeEdMode.h
--- a/Engine/Source/Editor/LandscapeEditor/Private/LandscapeEdMode.h
+++ b/Engine/Source/Editor/LandscapeEditor/Private/LandscapeEdMode.h
@@ -408,6 +408,12 @@ public:
 	/** FEdMode: Called when a mouse button is released */
 	virtual bool EndTracking(FEditorViewportClient* InViewportClient, FViewport* InViewport) override;
 
+	/** FEdMode: Called when a transform operation begins */
+	virtual bool BeginTransform(const FGizmoState& InState) override;
+
+	/** FEdMode: Called when a transform operation ends */
+	virtual bool EndTransform(const FGizmoState& InState) override;
+
 	/** FEdMode: Allow us to disable mouse delta tracking during painting */
 	virtual bool DisallowMouseDeltaTracking() const override;
 
diff --git a/Engine/Source/Editor/LandscapeEditor/Private/LandscapeEdModeSplineTools.cpp b/Engine/Source/Editor/LandscapeEditor/Private/LandscapeEdModeSplineTools.cpp
--- a/Engine/Source/Editor/LandscapeEditor/Private/LandscapeEdModeSplineTools.cpp
+++ b/Engine/Source/Editor/LandscapeEditor/Private/LandscapeEdModeSplineTools.cpp
@@ -18,6 +18,7 @@
 #include "EditorModeManager.h"
 #include "UnrealEdGlobals.h"
 #include "EditorViewportClient.h"
+#include "Framework/Application/SlateApplication.h"
 #include "LandscapeToolInterface.h"
 #include "LandscapeProxy.h"
 #include "LandscapeSplineActor.h"
@@ -35,6 +36,7 @@
 #include "ControlPointMeshComponent.h"
 #include "Algo/AnyOf.h"
 #include "Algo/Copy.h"
+#include "GizmoEdModeInterface.h"
 #include "HAL/PlatformApplicationMisc.h"
 #include "UnrealExporter.h"
 #include "UnrealWidget.h"
@@ -1341,23 +1343,9 @@ public:
 				{
 					if (HitProxy->IsA(HWidgetAxis::StaticGetType()))
 					{
-						checkSlow(SelectedSplineControlPoints.Num() > 0);
-						bMovingControlPoint = true;
-
-						if (SelectedSplineControlPoints.Num() == 1 && InViewportClient->IsAltPressed() && InViewportClient->GetWidgetMode() == UE::Widget::WM_Translate && InViewportClient->GetCurrentWidgetAxis() != EAxisList::None)
-						{
-							GEditor->BeginTransaction(LOCTEXT("LandscapeSpline_DuplicateControlPoint", "Duplicate Landscape Spline Control Point"));
-						}
-						else
-						{
-							GEditor->BeginTransaction(LOCTEXT("LandscapeSpline_ModifyControlPoint", "Modify Landscape Spline Control Point"));
-						}
-
-						for (ULandscapeSplineControlPoint* ControlPoint : SelectedSplineControlPoints)
-						{
-							ControlPoint->Modify();
-							ControlPoint->GetOuterULandscapeSplinesComponent()->Modify();
-						}
+						// BeginTransform isn't called for Legacy gizmos, so we call it here for a common codepath with new gizmos
+						const bool bIsDuplicatingControlPoints = InViewportClient->IsAltPressed();
+						BeginTransformInternal(FGizmoState(), bIsDuplicatingControlPoints);
 
 						return false; // We're not actually handling this case ourselves, just wrapping it in a transaction
 					}
@@ -1386,24 +1374,7 @@ public:
 			{
 				if (bMovingControlPoint)
 				{
-					bMovingControlPoint = false;
-
-					for (ULandscapeSplineControlPoint* ControlPoint : SelectedSplineControlPoints)
-					{
-						if (bDuplicatingControlPoint && bAutoRotateOnJoin)
-						{
-							ControlPoint->AutoCalcRotation(bAlwaysRotateForward);
-						}
-
-						ControlPoint->UpdateSplinePoints(true);
-					}
-
-					ResetAllowDuplication();
-
-					EdMode->AutoUpdateDirtyLandscapeSplines();
-					GEditor->EndTransaction();
-
-					return false; // We're not actually handling this case ourselves, just wrapping it in a transaction
+					EndTransform(FGizmoState());
 				}
 				else if (DraggingTangent_Segment)
 				{
@@ -1686,6 +1657,66 @@ public:
 		return false;
 	}
 
+	virtual bool BeginTransform(const FGizmoState& InState) override
+	{
+		const TArray<ULandscapeSplineControlPoint*> SelectedSplineControlPoints = SplineSelection->GetSelectedSplineControlPoints();
+
+		const bool bShouldDuplicateSelectedControlPoints = 
+			SelectedSplineControlPoints.Num() == 1 
+			&& FSlateApplication::Get().GetModifierKeys().IsAltDown()
+			&& InState.TransformMode == EGizmoTransformMode::Translate;
+		return BeginTransformInternal(InState, bShouldDuplicateSelectedControlPoints);
+	}
+
+	bool BeginTransformInternal(const FGizmoState& InState, const bool bInIsDuplicating)
+	{
+		TArray<ULandscapeSplineControlPoint*> SelectedSplineControlPoints = SplineSelection->GetSelectedSplineControlPoints();
+
+		checkSlow(SelectedSplineControlPoints.Num() > 0);
+		bMovingControlPoint = true;
+
+		if (bInIsDuplicating)
+		{
+			GEditor->BeginTransaction(LOCTEXT("LandscapeSpline_DuplicateControlPoint", "Duplicate Landscape Spline Control Point"));
+		}
+		else
+		{
+			GEditor->BeginTransaction(LOCTEXT("LandscapeSpline_ModifyControlPoint", "Modify Landscape Spline Control Point"));
+		}
+
+		for (ULandscapeSplineControlPoint* ControlPoint : SelectedSplineControlPoints)
+		{
+			ControlPoint->Modify();
+			ControlPoint->GetOuterULandscapeSplinesComponent()->Modify();
+		}
+
+		return false;
+	}
+
+	virtual bool EndTransform(const FGizmoState& InState) override
+	{
+		TArray<ULandscapeSplineControlPoint*> SelectedSplineControlPoints = SplineSelection->GetSelectedSplineControlPoints();
+
+		bMovingControlPoint = false;
+
+		for (ULandscapeSplineControlPoint* ControlPoint : SelectedSplineControlPoints)
+		{
+			if (bDuplicatingControlPoint && bAutoRotateOnJoin)
+			{
+				ControlPoint->AutoCalcRotation(bAlwaysRotateForward);
+			}
+
+			ControlPoint->UpdateSplinePoints(true);
+		}
+
+		ResetAllowDuplication();
+
+		EdMode->AutoUpdateDirtyLandscapeSplines();
+		GEditor->EndTransaction();
+
+		return false; // We're not actually handling this case ourselves, just wrapping it in a transaction
+	}
+
 	void OnUndo()
 	{
 		if (EdMode && EdMode->CurrentToolMode && EdMode->CurrentToolMode->CurrentToolName == "Splines" && SplineSelection)
diff --git a/Engine/Source/Editor/LandscapeEditor/Public/LandscapeToolInterface.h b/Engine/Source/Editor/LandscapeEditor/Public/LandscapeToolInterface.h
--- a/Engine/Source/Editor/LandscapeEditor/Public/LandscapeToolInterface.h
+++ b/Engine/Source/Editor/LandscapeEditor/Public/LandscapeToolInterface.h
@@ -247,6 +247,12 @@ public:
 	virtual bool InputDelta(FEditorViewportClient* InViewportClient, FViewport* InViewport, FVector& InDrag, FRotator& InRot, FVector& InScale) { return false; }
 	virtual bool GetCursor(EMouseCursor::Type& OutCursor) const { return false;  }
 
+	/** FEdMode: Called when a transform operation begins */
+	virtual bool BeginTransform(const FGizmoState& InState) { return false; }
+
+	/** FEdMode: Called when a transform operation ends */
+	virtual bool EndTransform(const FGizmoState& InState) { return false; }
+
 	FLandscapeTool() {}
 	virtual ~FLandscapeTool() = default;
 	virtual const TCHAR* GetToolName() const = 0;
