Fixed landscape file cache's file monitoring ability that got broken at some point because of relative/absolute paths

#Synced-CL 51883779

diff --git a/Engine/Source/Editor/LandscapeEditor/Private/LandscapeImageFileCache.cpp b/Engine/Source/Editor/LandscapeEditor/Private/LandscapeImageFileCache.cpp
--- a/Engine/Source/Editor/LandscapeEditor/Private/LandscapeImageFileCache.cpp
+++ b/Engine/Source/Editor/LandscapeEditor/Private/LandscapeImageFileCache.cpp
@@ -3,6 +3,7 @@
 #include "LandscapeImageFileCache.h"
 #include "LandscapeSettings.h"
 
+#include "Misc/Paths.h"
 #include "Modules/ModuleManager.h"
 #include "Internationalization/Internationalization.h"
 #include "IDirectoryWatcher.h"
@@ -39,7 +40,9 @@ FLandscapeImageFileCache::~FLandscapeImageFileCache()
 template<typename T>
 FLandscapeFileInfo FLandscapeImageFileCache::FindImage(const TCHAR* InImageFilename, FLandscapeImageDataRef& OutImageData)
 {
-	FCacheEntry* CacheEntry = ChooseCache<T>().Find(InImageFilename);
+	FString FullyQualifiedFileName = FPaths::ConvertRelativePathToFull(FString(InImageFilename));
+
+	FCacheEntry* CacheEntry = ChooseCache<T>().Find(FullyQualifiedFileName);
 	FLandscapeFileInfo Result;
 
 	if (CacheEntry)
@@ -103,7 +106,7 @@ FLandscapeFileInfo FLandscapeImageFileCache::FindImage(const TCHAR* InImageFilen
 
 	Trim();
 	OutImageData = NewImageData;
-	Add<T>(FString(InImageFilename), OutImageData);
+	Add<T>(FullyQualifiedFileName, OutImageData);
 
 	Result.PossibleResolutions.Add(FLandscapeFileResolution(OutImageData.Resolution.X, OutImageData.Resolution.Y));
 	Result.ResultCode = NewImageData.Result;
@@ -124,14 +127,18 @@ void FLandscapeImageFileCache::MonitorCallback(const TArray<struct FFileChangeDa
 	{
 		if (Change.Action == FFileChangeData::FCA_Modified || Change.Action == FFileChangeData::FCA_Removed)
 		{
-			Remove(Change.Filename);
+			// The file cache expects path to be fully qualified
+			FString FullyQualifiedFileName = FPaths::ConvertRelativePathToFull(Change.Filename);
+			Remove(FullyQualifiedFileName);
 		}
 	}
 }
 
-bool FLandscapeImageFileCache::MonitorFile(const FString& Filename)
+bool FLandscapeImageFileCache::MonitorFile(const FString& InFullyQualifiedFileName)
 {
-	FString Directory = FPaths::GetPath(Filename);
+	check(!FPaths::IsRelative(InFullyQualifiedFileName));
+
+	FString Directory = FPaths::GetPath(InFullyQualifiedFileName);
 
 	FDirectoryWatcherModule& DirectoryWatcherModule = FModuleManager::LoadModuleChecked<FDirectoryWatcherModule>("DirectoryWatcher");
 
@@ -152,9 +159,11 @@ bool FLandscapeImageFileCache::MonitorFile(const FString& Filename)
 	}
 }
 
-void FLandscapeImageFileCache::UnmonitorFile(const FString& Filename)
+void FLandscapeImageFileCache::UnmonitorFile(const FString& InFullyQualifiedFileName)
 {
-	FString Directory = FPaths::GetPath(Filename);
+	check(!FPaths::IsRelative(InFullyQualifiedFileName));
+
+	FString Directory = FPaths::GetPath(InFullyQualifiedFileName);
 
 	FDirectoryWatcherModule& DirectoryWatcherModule = FModuleManager::LoadModuleChecked<FDirectoryWatcherModule>("DirectoryWatcher");
 
@@ -166,32 +175,35 @@ void FLandscapeImageFileCache::UnmonitorFile(const FString& Filename)
 		if (Monitor->NumFiles == 0)
 		{
 			DirectoryWatcherModule.Get()->UnregisterDirectoryChangedCallback_Handle(Directory, Monitor->MonitorHandle);
+			MonitoredDirs.Remove(Directory);
 		}
-
-		MonitoredDirs.Remove(Directory);
 	}
 }
 
 template <typename T>
-void FLandscapeImageFileCache::Add(const FString& Filename, FLandscapeImageDataRef NewImageData)
+void FLandscapeImageFileCache::Add(const FString& InFullyQualifiedFileName, FLandscapeImageDataRef NewImageData)
 {
+	check(!FPaths::IsRelative(InFullyQualifiedFileName));
+
 	CacheType& Cache = ChooseCache<T>();
-	check(Cache.Find(Filename) == nullptr);
+	check(Cache.Find(InFullyQualifiedFileName) == nullptr);
 
-	Cache.Add(FString(Filename), FCacheEntry(NewImageData));
-	MonitorFile(Filename);
+	Cache.Add(InFullyQualifiedFileName, FCacheEntry(NewImageData));
+	MonitorFile(InFullyQualifiedFileName);
 	CacheSize += NewImageData.Data->Num();
 }
 
-void FLandscapeImageFileCache::Remove(const FString& Filename)
+void FLandscapeImageFileCache::Remove(const FString& InFullyQualifiedFileName)
 {
+	check(!FPaths::IsRelative(InFullyQualifiedFileName));
+
 	for (CacheType& Cache : CachedImages)
 	{
-		if (FCacheEntry* CacheEntry = Cache.Find(Filename))
+		if (FCacheEntry* CacheEntry = Cache.Find(InFullyQualifiedFileName))
 		{
 			CacheSize -= CacheEntry->ImageData.Data->Num();
-			Cache.Remove(Filename);
-			UnmonitorFile(Filename);
+			Cache.Remove(InFullyQualifiedFileName);
+			UnmonitorFile(InFullyQualifiedFileName);
 		}
 	}
 }
diff --git a/Engine/Source/Editor/LandscapeEditor/Private/LandscapeImageFileCache.h b/Engine/Source/Editor/LandscapeEditor/Private/LandscapeImageFileCache.h
--- a/Engine/Source/Editor/LandscapeEditor/Private/LandscapeImageFileCache.h
+++ b/Engine/Source/Editor/LandscapeEditor/Private/LandscapeImageFileCache.h
@@ -59,8 +59,8 @@ private:
 	
 	void OnLandscapeSettingsChanged(UObject* InObject, struct FPropertyChangedEvent& InPropertyChangedEvent);
 	 
-	bool MonitorFile(const FString& Filename);
-	void UnmonitorFile(const FString& Filename);
+	bool MonitorFile(const FString& InFullyQualifiedFileName);
+	void UnmonitorFile(const FString& InFullyQualifiedFileName);
 
 	void MonitorCallback(const TArray<struct FFileChangeData>& Changes);
 
@@ -81,8 +81,8 @@ private:
 	TMap<FString, FDirectoryMonitor> MonitoredDirs;
 
 	template <typename T>
-	void Add(const FString& Filename, FLandscapeImageDataRef NewImageData);
-	void Remove(const FString& Filename);
+	void Add(const FString& InFullyQualifiedFileName, FLandscapeImageDataRef NewImageData);
+	void Remove(const FString& InFullyQualifiedFileName);
 
 	void Trim();
 
